Added missing PreEditChange to force syncrhonization with the async texture build framework when modifying water texture.

diff --git a/Engine/Plugins/Experimental/Water/Source/Editor/Private/WaterEditorSubsystem.cpp b/Engine/Plugins/Experimental/Water/Source/Editor/Private/WaterEditorSubsystem.cpp
--- a/Engine/Plugins/Experimental/Water/Source/Editor/Private/WaterEditorSubsystem.cpp
+++ b/Engine/Plugins/Experimental/Water/Source/Editor/Private/WaterEditorSubsystem.cpp
@@ -40,6 +40,7 @@ void UpdateSingleTexture(UTexture2D*& DestTexure, UTextureRenderTarget2D* SrcRen
 		break;
 	}
 
+	DestTexure->PreEditChange(nullptr); // Ensures synchronization with TextureCompilingManager.
 	DestTexure->LODGroup = GetDefault<UWaterEditorSettings>()->TextureGroupForGeneratedTextures;
 	DestTexure->MipGenSettings = TMGS_NoMipmaps;
 	DestTexure->MaxTextureSize = GetDefault<UWaterEditorSettings>()->MaxWaterVelocityAndHeightTextureSize;
