WaterSystem: asserting on tick after closing sub editor

diff --git a/Engine/Plugins/Experimental/Water/Source/Runtime/Private/WaterSubsystem.cpp b/Engine/Plugins/Experimental/Water/Source/Runtime/Private/WaterSubsystem.cpp
--- a/Engine/Plugins/Experimental/Water/Source/Runtime/Private/WaterSubsystem.cpp
+++ b/Engine/Plugins/Experimental/Water/Source/Runtime/Private/WaterSubsystem.cpp
@@ -621,10 +621,13 @@ void UWaterSubsystem::SetMPCTime(float Time, float PrevTime)
 		if (MaterialParameterCollection)
 		{
 			UMaterialParameterCollectionInstance* MaterialParameterCollectionInstance = GetWorld()->GetParameterCollectionInstance(MaterialParameterCollection);
-			const static FName TimeParam(TEXT("Time"));
-			const static FName PrevTimeParam(TEXT("PrevTime"));
-			MaterialParameterCollectionInstance->SetScalarParameterValue(TimeParam, Time);
-			MaterialParameterCollectionInstance->SetScalarParameterValue(PrevTimeParam, PrevTime);
+			if (World->IsValidLowLevel())
+			{
+				const static FName TimeParam(TEXT("Time"));
+				const static FName PrevTimeParam(TEXT("PrevTime"));
+				MaterialParameterCollectionInstance->SetScalarParameterValue(TimeParam, Time);
+				MaterialParameterCollectionInstance->SetScalarParameterValue(PrevTimeParam, PrevTime);
+			}
 		}
 	}
 }
diff --git a/Engine/Source/Runtime/Engine/Classes/Materials/MaterialParameterCollectionInstance.h b/Engine/Source/Runtime/Engine/Classes/Materials/MaterialParameterCollectionInstance.h
--- a/Engine/Source/Runtime/Engine/Classes/Materials/MaterialParameterCollectionInstance.h
+++ b/Engine/Source/Runtime/Engine/Classes/Materials/MaterialParameterCollectionInstance.h
@@ -30,6 +30,9 @@ class ENGINE_API UMaterialParameterCollectionInstance : public UObject
 	virtual void FinishDestroy() override;
 	//~ End UObject Interface.
 
+	/** Returns true if this parameter collection instance is valid */
+	bool IsValid() const { return World.IsValid() && Collection; }
+
 	/** Initializes the instance with the collection it is based off of and the world it is owned by. */
 	void SetCollection(UMaterialParameterCollection* InCollection, UWorld* InWorld);
 
