WaterSystem: asserting on tick after closing sub editor

diff --git a/Engine/Plugins/Experimental/Water/Source/Runtime/Private/WaterSubsystem.cpp b/Engine/Plugins/Experimental/Water/Source/Runtime/Private/WaterSubsystem.cpp
--- a/Engine/Plugins/Experimental/Water/Source/Runtime/Private/WaterSubsystem.cpp
+++ b/Engine/Plugins/Experimental/Water/Source/Runtime/Private/WaterSubsystem.cpp
@@ -621,7 +621,7 @@ void UWaterSubsystem::SetMPCTime(float Time, float PrevTime)
 		if (MaterialParameterCollection)
 		{
 			UMaterialParameterCollectionInstance* MaterialParameterCollectionInstance = GetWorld()->GetParameterCollectionInstance(MaterialParameterCollection);
-			if (World->IsValidLowLevel())
+			if (MaterialParameterCollectionInstance->IsValid())
 			{
 				const static FName TimeParam(TEXT("Time"));
 				const static FName PrevTimeParam(TEXT("PrevTime"));
