Cherry-pick 14923831 to fix CIS

diff --git a/Engine/Plugins/Experimental/Water/Source/Editor/Private/WaterBrushCacheContainer.h b/Engine/Plugins/Experimental/Water/Source/Editor/Private/WaterBrushCacheContainer.h
--- a/Engine/Plugins/Experimental/Water/Source/Editor/Private/WaterBrushCacheContainer.h
+++ b/Engine/Plugins/Experimental/Water/Source/Editor/Private/WaterBrushCacheContainer.h
@@ -16,10 +16,10 @@ struct FWaterBodyBrushCache
 	GENERATED_USTRUCT_BODY()
 
 	UPROPERTY(EditAnywhere, BlueprintReadWrite, Category = "Cache")
-	UTextureRenderTarget2D* CacheRenderTarget;
+	UTextureRenderTarget2D* CacheRenderTarget = nullptr;
 
 	UPROPERTY(EditAnywhere, BlueprintReadWrite, Category = "Cache")
-	bool CacheIsValid;
+	bool CacheIsValid = false;
 };
 
 UCLASS(config = Engine, Blueprintable, BlueprintType)
