Cherry-pick 15065810 and revert 14931698 (hack for fixing water subsystem being ticked after being deinitialized):

* Fixed water subsystem keeping on ticking after being deinitialized
* Fixed water subsystem being ticked outside of the world tick
* Ensure water subsystem is ticked in SWaterWavesEditorViewport's viewport client's Tick (preview worlds are not ticked automatically)

diff --git a/Engine/Plugins/Experimental/Water/Source/Editor/Private/WaterWavesEditorViewport.cpp b/Engine/Plugins/Experimental/Water/Source/Editor/Private/WaterWavesEditorViewport.cpp
--- a/Engine/Plugins/Experimental/Water/Source/Editor/Private/WaterWavesEditorViewport.cpp
+++ b/Engine/Plugins/Experimental/Water/Source/Editor/Private/WaterWavesEditorViewport.cpp
@@ -8,13 +8,16 @@
 #include "WaterBodyCustomActor.h"
 #include "SCommonEditorViewportToolbarBase.h"
 #include "WaterEditorSettings.h"
+#include "WaterEditorModule.h"
 #include "WaterSplineComponent.h"
 
 #include "WaterSubsystem.h"
 
 SWaterWavesEditorViewport::SWaterWavesEditorViewport()
-	: PreviewScene(MakeShareable(new FAdvancedPreviewScene(FPreviewScene::ConstructionValues())))
 {
+	// Temporarily allow water subsystem to be created on preview worlds because we need one here : 
+	UWaterSubsystem::FScopedAllowWaterSubsystemOnPreviewWorld AllowWaterSubsystemOnPreviewWorldScope(true);
+	PreviewScene = MakeShareable(new FAdvancedPreviewScene(FPreviewScene::ConstructionValues()));
 }
 
 void SWaterWavesEditorViewport::Construct(const FArguments& InArgs)
@@ -63,10 +66,7 @@ void SWaterWavesEditorViewport::AddReferencedObjects(FReferenceCollector& Collec
 
 TSharedRef<FEditorViewportClient> SWaterWavesEditorViewport::MakeEditorViewportClient()
 {
-	EditorViewportClient = MakeShareable(new FEditorViewportClient(nullptr, PreviewScene.Get(), SharedThis(this)));
-	EditorViewportClient->bSetListenerPosition = false;
-	EditorViewportClient->SetRealtime(true);
-	EditorViewportClient->EngineShowFlags.Grid = false;
+	EditorViewportClient = MakeShareable(new FWaterWavesEditorViewportClient(PreviewScene.Get(), SharedThis(this)));
 	return EditorViewportClient.ToSharedRef();
 }
 
@@ -80,4 +80,27 @@ void SWaterWavesEditorViewport::SetShouldPauseWaveTime(bool bShouldFreeze)
 	UWaterSubsystem* WaterSubsystem = EditorViewportClient->GetWorld()->GetSubsystem<UWaterSubsystem>();
 	check(WaterSubsystem != nullptr);
 	WaterSubsystem->SetShouldPauseWaveTime(bShouldFreeze);
-}
\ No newline at end of file
+}
+
+
+// ----------------------------------------------------------------------------------
+
+FWaterWavesEditorViewportClient::FWaterWavesEditorViewportClient(FPreviewScene* InPreviewScene, const TWeakPtr<SEditorViewport>& InEditorViewportWidget)
+	: FEditorViewportClient(nullptr, InPreviewScene, InEditorViewportWidget)
+{
+	bSetListenerPosition = false;
+	SetRealtime(true);
+	EngineShowFlags.Grid = false;
+}
+
+void FWaterWavesEditorViewportClient::Tick(float DeltaSeconds)
+{
+	Super::Tick(DeltaSeconds);
+
+	// Tick the preview scene world.
+	if (!GIntraFrameDebuggingGameThread)
+	{
+		PreviewScene->GetWorld()->Tick(IsRealtime() ? LEVELTICK_All : LEVELTICK_TimeOnly, DeltaSeconds);
+	}
+}
+
diff --git a/Engine/Plugins/Experimental/Water/Source/Editor/Private/WaterWavesEditorViewport.h b/Engine/Plugins/Experimental/Water/Source/Editor/Private/WaterWavesEditorViewport.h
--- a/Engine/Plugins/Experimental/Water/Source/Editor/Private/WaterWavesEditorViewport.h
+++ b/Engine/Plugins/Experimental/Water/Source/Editor/Private/WaterWavesEditorViewport.h
@@ -3,6 +3,7 @@
 #pragma once
 
 #include "CoreMinimal.h"
+#include "EditorViewportClient.h"
 #include "SAssetEditorViewport.h"
 #include "SCommonEditorViewportToolbarBase.h"
 
@@ -13,6 +14,9 @@ class AWaterBodyCustom;
 class UWaterWavesAssetReference;
 class SEditorViewport;
 
+
+// ----------------------------------------------------------------------------------
+
 class SWaterWavesEditorViewport : public SAssetEditorViewport, public ICommonEditorViewportToolbarInfoProvider, public FGCObject
 {
 public:
@@ -48,4 +52,19 @@ private:
 	TWeakPtr<FWaterWavesEditorToolkit> WaterWavesEditorToolkitPtr;
 
 	AWaterBodyCustom* CustomWaterBody = nullptr;
+};
+
+
+// ----------------------------------------------------------------------------------
+
+class FWaterWavesEditorViewportClient : public FEditorViewportClient
+{
+public:
+	using Super = FEditorViewportClient;
+
+	FWaterWavesEditorViewportClient(FPreviewScene* InPreviewScene, const TWeakPtr<SEditorViewport>& InEditorViewportWidget = nullptr);
+
+	// FEditorViewportClient interface
+	virtual void Tick(float DeltaSeconds) override;
+	// End of FEditorViewportClient
 };
\ No newline at end of file
diff --git a/Engine/Plugins/Experimental/Water/Source/Runtime/Private/WaterSubsystem.cpp b/Engine/Plugins/Experimental/Water/Source/Runtime/Private/WaterSubsystem.cpp
--- a/Engine/Plugins/Experimental/Water/Source/Runtime/Private/WaterSubsystem.cpp
+++ b/Engine/Plugins/Experimental/Water/Source/Runtime/Private/WaterSubsystem.cpp
@@ -129,6 +129,14 @@ struct FUnderwaterPostProcessDebugInfo
 
 // ----------------------------------------------------------------------------------
 
+#if WITH_EDITOR
+
+bool UWaterSubsystem::bAllowWaterSubsystemOnPreviewWorld = false;
+
+#endif // WITH_EDITOR
+
+// ----------------------------------------------------------------------------------
+
 UWaterSubsystem::UWaterSubsystem()
 {
 	SmoothedWorldTimeSeconds = 0.f;
@@ -136,6 +144,7 @@ UWaterSubsystem::UWaterSubsystem()
 	PrevWorldTimeSeconds = 0.f;
 	bUnderWaterForAudio = false;
 	bPauseWaveTime = false;
+	bInitialized = false;
 
 	struct FConstructorStatics
 	{
@@ -176,7 +185,8 @@ FWaterBodyManager* UWaterSubsystem::GetWaterBodyManager(UWorld* InWorld)
 
 void UWaterSubsystem::Tick(float DeltaTime)
 {
-	UWorld* World = GetWorld();
+	check(bInitialized); // Ticking should have been disabled for an uninitialized subsystem
+	check(GetWorld() != nullptr);
 	if (FreezeWaves == 0 && bPauseWaveTime == false)
 	{
 		NonSmoothedWorldTimeSeconds += DeltaTime;
@@ -212,12 +222,23 @@ TStatId UWaterSubsystem::GetStatId() const
 
 bool UWaterSubsystem::DoesSupportWorldType(EWorldType::Type WorldType) const
 {
-	return WorldType == EWorldType::Game || WorldType == EWorldType::Editor || WorldType == EWorldType::PIE || WorldType == EWorldType::EditorPreview;
+#if WITH_EDITOR
+	// In editor, don't let preview worlds instantiate a water subsystem (except if explicitly allowed by a tool that requested it by setting bAllowWaterSubsystemOnPreviewWorld)
+	if (WorldType == EWorldType::EditorPreview)
+	{
+		return bAllowWaterSubsystemOnPreviewWorld;
+	}
+#endif // WITH_EDITOR
+
+	return WorldType == EWorldType::Game || WorldType == EWorldType::Editor || WorldType == EWorldType::PIE;
 }
 
 void UWaterSubsystem::Initialize(FSubsystemCollectionBase& Collection)
 {
-	WaterBodyManager.Initialize(GetWorld());
+	UWorld* World = GetWorld();
+	check(World != nullptr);
+
+	WaterBodyManager.Initialize(World);
 
 	bUsingSmoothedTime = false;
 	FConsoleVariableDelegate NotifyWaterScalabilityChanged = FConsoleVariableDelegate::CreateUObject(this, &UWaterSubsystem::NotifyWaterScalabilityChangedInternal);
@@ -232,15 +253,20 @@ void UWaterSubsystem::Initialize(FSubsystemCollectionBase& Collection)
 #endif //WITH_EDITOR
 	ApplyRuntimeSettings(GetDefault<UWaterRuntimeSettings>(), EPropertyChangeType::ValueSet);
 
-	GetWorld()->OnBeginPostProcessSettings.AddUObject(this, &UWaterSubsystem::ComputeUnderwaterPostProcess);
-	GetWorld()->InsertPostProcessVolume(&UnderwaterPostProcessVolume);
+	World->OnBeginPostProcessSettings.AddUObject(this, &UWaterSubsystem::ComputeUnderwaterPostProcess);
+	World->InsertPostProcessVolume(&UnderwaterPostProcessVolume);
 
 	UCollisionProfile::Get()->OnLoadProfileConfig.AddUObject(this, &UWaterSubsystem::OnLoadProfileConfig);
 	AddWaterCollisionProfile();
+
+	bInitialized = true;
 }
 
 void UWaterSubsystem::Deinitialize()
 {
+	UWorld* World = GetWorld();
+	check(World != nullptr);
+
 	UCollisionProfile::Get()->OnLoadProfileConfig.RemoveAll(this);
 
 	FConsoleVariableDelegate NullCallback;
@@ -248,33 +274,37 @@ void UWaterSubsystem::Deinitialize()
 	CVarShallowWaterSim->SetOnChangedCallback(NullCallback);
 	CVarWaterEnabled->SetOnChangedCallback(NullCallback);
 
-	GetWorld()->OnBeginPostProcessSettings.RemoveAll(this);
-	GetWorld()->RemovePostProcessVolume(&UnderwaterPostProcessVolume);
+	World->OnBeginPostProcessSettings.RemoveAll(this);
+	World->RemovePostProcessVolume(&UnderwaterPostProcessVolume);
 
 	WaterBodyManager.Deinitialize();
 
 #if WITH_EDITOR
 	GetDefault<UWaterRuntimeSettings>()->OnSettingsChange.RemoveAll(this);
 #endif //WITH_EDITOR
+
+	bInitialized = false;
 }
 
 void UWaterSubsystem::ApplyRuntimeSettings(const UWaterRuntimeSettings* Settings, EPropertyChangeType::Type ChangeType)
 {
+	UWorld* World = GetWorld();
+	check(World != nullptr);
 	UnderwaterTraceChannel = Settings->CollisionChannelForWaterTraces;
 	MaterialParameterCollection = Settings->MaterialParameterCollection.LoadSynchronous();
 
 #if WITH_EDITOR
-	for (TActorIterator<AWaterBody> ActorItr(GetWorld()); ActorItr; ++ActorItr)
+	for (TActorIterator<AWaterBody> ActorItr(World); ActorItr; ++ActorItr)
 	{
 		(*ActorItr)->UpdateActorIcon();
 	}
 
-	for (TActorIterator<AWaterBodyIsland> ActorItr(GetWorld()); ActorItr; ++ActorItr)
+	for (TActorIterator<AWaterBodyIsland> ActorItr(World); ActorItr; ++ActorItr)
 	{
 		(*ActorItr)->UpdateActorIcon();
 	}
 
-	for (TActorIterator<AWaterBodyExclusionVolume> ActorItr(GetWorld()); ActorItr; ++ActorItr)
+	for (TActorIterator<AWaterBodyExclusionVolume> ActorItr(World); ActorItr; ++ActorItr)
 	{
 		(*ActorItr)->UpdateActorIcon();
 	}
@@ -405,35 +435,43 @@ void UWaterSubsystem::SetShouldPauseWaveTime(bool bInPauseWaveTime)
 
 void UWaterSubsystem::SetOceanFloodHeight(float InFloodHeight)
 {
-	const float ClampedFloodHeight = FMath::Max(0.0f, InFloodHeight);
-
-	if (FloodHeight != ClampedFloodHeight)
+	if (UWorld* World = GetWorld())
 	{
-		FloodHeight = ClampedFloodHeight;
-		MarkAllWaterMeshesForRebuild();
+		const float ClampedFloodHeight = FMath::Max(0.0f, InFloodHeight);
 
-		// the ocean body is dynamic and needs to be readjusted when the flood height changes : 
-		if (OceanActor.IsValid())
+		if (FloodHeight != ClampedFloodHeight)
 		{
-			OceanActor->SetHeightOffset(InFloodHeight);
-		}
+			FloodHeight = ClampedFloodHeight;
+			MarkAllWaterMeshesForRebuild();
 
-		// All water body actors need to update their underwater post process MID as it depends on the ocean global height : 
-		for (TActorIterator<AWaterBody> ActorItr(GetWorld()); ActorItr; ++ActorItr)
-		{
-			AWaterBody* WaterBody = *ActorItr;
-			WaterBody->UpdateMaterialInstances();
+			// the ocean body is dynamic and needs to be readjusted when the flood height changes : 
+			if (OceanActor.IsValid())
+			{
+				OceanActor->SetHeightOffset(InFloodHeight);
+			}
+
+			// All water body actors need to update their underwater post process MID as it depends on the ocean global height : 
+			for (TActorIterator<AWaterBody> ActorItr(World); ActorItr; ++ActorItr)
+			{
+				AWaterBody* WaterBody = *ActorItr;
+				WaterBody->UpdateMaterialInstances();
+			}
 		}
 	}
 }
 
 AWaterMeshActor* UWaterSubsystem::GetWaterMeshActor() const
 {
-	// @todo water: this assumes only one water mesh actor right now.  In the future we may need to associate a water mesh actor with a water body more directly
-	TActorIterator<AWaterMeshActor> It(GetWorld());
-	WaterMeshActor = It ? *It : nullptr;
+	if (UWorld* World = GetWorld())
+	{
+		// @todo water: this assumes only one water mesh actor right now.  In the future we may need to associate a water mesh actor with a water body more directly
+		TActorIterator<AWaterMeshActor> It(World);
+		WaterMeshActor = It ? *It : nullptr;
 
-	return WaterMeshActor;
+		return WaterMeshActor;
+	}
+
+	return nullptr;
 }
 
 float UWaterSubsystem::GetOceanBaseHeight() const
@@ -448,9 +486,12 @@ float UWaterSubsystem::GetOceanBaseHeight() const
 
 void UWaterSubsystem::MarkAllWaterMeshesForRebuild()
 {
-	for (AWaterMeshActor* WaterMesh : TActorRange<AWaterMeshActor>(GetWorld()))
+	if (UWorld* World = GetWorld())
 	{
-		WaterMesh->MarkWaterMeshComponentForRebuild();
+		for (AWaterMeshActor* WaterMesh : TActorRange<AWaterMeshActor>(World))
+		{
+			WaterMesh->MarkWaterMeshComponentForRebuild();
+		}
 	}
 }
 
@@ -461,10 +502,13 @@ void UWaterSubsystem::NotifyWaterScalabilityChangedInternal(IConsoleVariable* CV
 
 void UWaterSubsystem::NotifyWaterEnabledChangedInternal(IConsoleVariable* CVar)
 {
-	// Water body visibility depends on CVarWaterEnabled
-	for (AWaterBody* WaterBody : TActorRange<AWaterBody>(GetWorld()))
+	if (UWorld* World = GetWorld())
 	{
-		WaterBody->UpdateWaterComponentVisibility();
+		// Water body visibility depends on CVarWaterEnabled
+		for (AWaterBody* WaterBody : TActorRange<AWaterBody>(World))
+		{
+			WaterBody->UpdateWaterComponentVisibility();
+		}
 	}
 }
 
@@ -504,7 +548,8 @@ void UWaterSubsystem::ComputeUnderwaterPostProcess(FVector ViewLocation, FSceneV
 {
 	SCOPE_CYCLE_COUNTER(STAT_WaterIsUnderwater);
 
-	if (SceneView->Family->EngineShowFlags.PostProcessing == 0)
+	UWorld* World = GetWorld();
+	if ((World == nullptr) || (SceneView->Family->EngineShowFlags.PostProcessing == 0))
 	{
 		return;
 	}
@@ -527,7 +572,7 @@ void UWaterSubsystem::ComputeUnderwaterPostProcess(FVector ViewLocation, FSceneV
 	TArray<FHitResult> Hits;
 	TArray<FWaterBodyPostProcessQuery, TInlineAllocator<4>> WaterBodyQueriesToProcess;
 	const AWaterMeshActor* LocalWaterMeshActor = GetWaterMeshActor();
-	if ((LocalWaterMeshActor != nullptr) && GetWorld()->SweepMultiByChannel(Hits, ViewLocation, ViewLocation + FVector(0, 0, TraceDistance), FQuat::Identity, UnderwaterTraceChannel, FCollisionShape::MakeSphere(TraceDistance), TraceSimple))
+	if ((LocalWaterMeshActor != nullptr) && World->SweepMultiByChannel(Hits, ViewLocation, ViewLocation + FVector(0, 0, TraceDistance), FQuat::Identity, UnderwaterTraceChannel, FCollisionShape::MakeSphere(TraceDistance), TraceSimple))
 	{
 		if (Hits.Num() > 1)
 		{
@@ -620,14 +665,11 @@ void UWaterSubsystem::SetMPCTime(float Time, float PrevTime)
 	{
 		if (MaterialParameterCollection)
 		{
-			UMaterialParameterCollectionInstance* MaterialParameterCollectionInstance = GetWorld()->GetParameterCollectionInstance(MaterialParameterCollection);
-			if (MaterialParameterCollectionInstance->IsValid())
-			{
-				const static FName TimeParam(TEXT("Time"));
-				const static FName PrevTimeParam(TEXT("PrevTime"));
-				MaterialParameterCollectionInstance->SetScalarParameterValue(TimeParam, Time);
-				MaterialParameterCollectionInstance->SetScalarParameterValue(PrevTimeParam, PrevTime);
-			}
+			UMaterialParameterCollectionInstance* MaterialParameterCollectionInstance = World->GetParameterCollectionInstance(MaterialParameterCollection);
+			const static FName TimeParam(TEXT("Time"));
+			const static FName PrevTimeParam(TEXT("PrevTime"));
+			MaterialParameterCollectionInstance->SetScalarParameterValue(TimeParam, Time);
+			MaterialParameterCollectionInstance->SetScalarParameterValue(PrevTimeParam, PrevTime);
 		}
 	}
 }
@@ -722,4 +764,21 @@ void UWaterSubsystem::ShowOnScreenDebugInfo(const FUnderwaterPostProcessDebugInf
 
 // ----------------------------------------------------------------------------------
 
-#undef LOCTEXT_NAMESPACE
\ No newline at end of file
+#if WITH_EDITOR
+
+UWaterSubsystem::FScopedAllowWaterSubsystemOnPreviewWorld::FScopedAllowWaterSubsystemOnPreviewWorld(bool bNewValue)
+{
+	bPreviousValue = UWaterSubsystem::GetAllowWaterSubsystemOnPreviewWorld();
+	UWaterSubsystem::SetAllowWaterSubsystemOnPreviewWorld(bNewValue);
+}
+
+UWaterSubsystem::FScopedAllowWaterSubsystemOnPreviewWorld::~FScopedAllowWaterSubsystemOnPreviewWorld()
+{
+	UWaterSubsystem::SetAllowWaterSubsystemOnPreviewWorld(bPreviousValue);
+}
+
+#endif // WITH_EDITOR
+
+// ----------------------------------------------------------------------------------
+
+#undef LOCTEXT_NAMESPACE
diff --git a/Engine/Plugins/Experimental/Water/Source/Runtime/Public/WaterSubsystem.h b/Engine/Plugins/Experimental/Water/Source/Runtime/Public/WaterSubsystem.h
--- a/Engine/Plugins/Experimental/Water/Source/Runtime/Public/WaterSubsystem.h
+++ b/Engine/Plugins/Experimental/Water/Source/Runtime/Public/WaterSubsystem.h
@@ -61,16 +61,26 @@ class WATER_API UWaterSubsystem : public UWorldSubsystem, public FTickableGameOb
 {
 	GENERATED_BODY()
 
+public:
 	UWaterSubsystem();
 
+	// FTickableGameObject implementation Begin
+	virtual UWorld* GetTickableGameObjectWorld() const override { return GetWorld(); }
+	virtual bool IsTickable() const override { return bInitialized; }
+	virtual bool IsTickableInEditor() const override { return bInitialized; }
 	virtual void Tick(float DeltaTime) override;
-
-	/** return the stat id to use for this tickable **/
 	virtual TStatId GetStatId() const override;
+	// FTickableGameObject implementation End
 
+	// UWorldSubsystem implementation Begin
 	/** Override to support water subsystems in editor preview worlds */
 	virtual bool DoesSupportWorldType(EWorldType::Type WorldType) const override;
-public:
+	// UWorldSubsystem implementation End
+
+	// USubsystem implementation Begin
+	virtual void Initialize(FSubsystemCollectionBase& Collection) override;
+	virtual void Deinitialize() override;
+	// USubsystem implementation End
 
 	/** Static helper function to get a water subsystem from a world, returns nullptr if world or subsystem don't exist */
 	static UWaterSubsystem* GetWaterSubsystem(const UWorld* InWorld);
@@ -78,13 +88,6 @@ public:
 	/** Static helper function to get a waterbody manager from a world, returns nullptr if world or manager don't exist */
 	static FWaterBodyManager* GetWaterBodyManager(UWorld* InWorld);
 
-	virtual bool IsTickableInEditor() const override { return true; }
-
-	// Begin USubsystem
-	virtual void Initialize(FSubsystemCollectionBase& Collection) override;
-	virtual void Deinitialize() override;
-	// End USubsystem
-
 	FWaterBodyManager WaterBodyManager;
 
 	AWaterMeshActor* GetWaterMeshActor() const;
@@ -150,21 +153,26 @@ public:
 	UMaterialParameterCollection* GetMaterialParameterCollection() const {	return MaterialParameterCollection; }
 	
 	void MarkAllWaterMeshesForRebuild();
-public:
-	DECLARE_EVENT_OneParam(UWaterSubsystem, FOnWaterSubsystemInitialized, UWaterSubsystem*)
-	static FOnWaterSubsystemInitialized OnWaterSubsystemInitialized;
 
-	UPROPERTY(BlueprintAssignable, Category = Water)
-	FOnCameraUnderwaterStateChanged OnCameraUnderwaterStateChanged;
+#if WITH_EDITOR
+	/** Little scope object to temporarily change the value of bAllowWaterSubsystemOnPreviewWorld */
+	struct WATER_API FScopedAllowWaterSubsystemOnPreviewWorld
+	{
+		FScopedAllowWaterSubsystemOnPreviewWorld(bool bNewValue);
+		~FScopedAllowWaterSubsystemOnPreviewWorld();
 
-	UPROPERTY(BlueprintAssignable, Category = Water)
-	FOnWaterScalabilityChanged OnWaterScalabilityChanged;
+		// Non-copyable
+	private:
+		FScopedAllowWaterSubsystemOnPreviewWorld() = delete;
+		FScopedAllowWaterSubsystemOnPreviewWorld& operator=(const FScopedAllowWaterSubsystemOnPreviewWorld&) = delete;
+		FScopedAllowWaterSubsystemOnPreviewWorld(const FScopedAllowWaterSubsystemOnPreviewWorld&) = delete;
+
+		bool bPreviousValue = false;
+	};
+	static void SetAllowWaterSubsystemOnPreviewWorld(bool bInValue) { bAllowWaterSubsystemOnPreviewWorld = bInValue; }
+	static bool GetAllowWaterSubsystemOnPreviewWorld() { return bAllowWaterSubsystemOnPreviewWorld; }
+#endif // WITH_EDITOR
 
-	UPROPERTY()
-	UStaticMesh* DefaultRiverMesh;
-	
-	UPROPERTY()
-	UStaticMesh* DefaultLakeMesh;
 private:
 	void NotifyWaterScalabilityChangedInternal(IConsoleVariable* CVar);
 	void NotifyWaterEnabledChangedInternal(IConsoleVariable* CVar);
@@ -179,8 +187,23 @@ private:
 	void ShowOnScreenDebugInfo(const FUnderwaterPostProcessDebugInfo& InDebugInfo);
 #endif // !(UE_BUILD_SHIPPING || UE_BUILD_TEST)
 
-private:
+public:
+	DECLARE_EVENT_OneParam(UWaterSubsystem, FOnWaterSubsystemInitialized, UWaterSubsystem*)
+	static FOnWaterSubsystemInitialized OnWaterSubsystemInitialized;
+
+	UPROPERTY(BlueprintAssignable, Category = Water)
+	FOnCameraUnderwaterStateChanged OnCameraUnderwaterStateChanged;
+
+	UPROPERTY(BlueprintAssignable, Category = Water)
+	FOnWaterScalabilityChanged OnWaterScalabilityChanged;
+
+	UPROPERTY()
+	UStaticMesh* DefaultRiverMesh;
 
+	UPROPERTY()
+	UStaticMesh* DefaultLakeMesh;
+
+private:
 	UPROPERTY()
 	mutable AWaterMeshActor* WaterMeshActor;
 
@@ -198,10 +221,16 @@ private:
 	bool bUsingOverrideWorldTimeSeconds;
 	bool bUnderWaterForAudio;
 	bool bPauseWaveTime;
+	bool bInitialized;
 
 	/** The parameter collection asset that holds the global parameters that are updated by this actor */
 	UPROPERTY()
 	UMaterialParameterCollection* MaterialParameterCollection;
 
 	FUnderwaterPostProcessVolume UnderwaterPostProcessVolume;
+
+#if WITH_EDITOR
+	/** By default, there is no water subsystem allowed on preview worlds except when explicitly requested : */
+	static bool bAllowWaterSubsystemOnPreviewWorld;
+#endif // WITH_EDITOR
 };
diff --git a/Engine/Source/Runtime/Engine/Classes/Materials/MaterialParameterCollectionInstance.h b/Engine/Source/Runtime/Engine/Classes/Materials/MaterialParameterCollectionInstance.h
--- a/Engine/Source/Runtime/Engine/Classes/Materials/MaterialParameterCollectionInstance.h
+++ b/Engine/Source/Runtime/Engine/Classes/Materials/MaterialParameterCollectionInstance.h
@@ -30,9 +30,6 @@ class ENGINE_API UMaterialParameterCollectionInstance : public UObject
 	virtual void FinishDestroy() override;
 	//~ End UObject Interface.
 
-	/** Returns true if this parameter collection instance is valid */
-	bool IsValid() const { return World.IsValid() && Collection; }
-
 	/** Initializes the instance with the collection it is based off of and the world it is owned by. */
 	void SetCollection(UMaterialParameterCollection* InCollection, UWorld* InWorld);
 
