* Fixed missing water sprite category in the viewport options : billboard components need to be added to the CDO in order to be detected by the editor as proper sprite categories
* Also fixes the oversized icon the first time you place a water actor

diff --git a/Engine/Plugins/Experimental/Water/Source/Editor/Private/WaterLandscapeBrush.cpp b/Engine/Plugins/Experimental/Water/Source/Editor/Private/WaterLandscapeBrush.cpp
--- a/Engine/Plugins/Experimental/Water/Source/Editor/Private/WaterLandscapeBrush.cpp
+++ b/Engine/Plugins/Experimental/Water/Source/Editor/Private/WaterLandscapeBrush.cpp
@@ -34,7 +34,7 @@ AWaterLandscapeBrush::AWaterLandscapeBrush(const FObjectInitializer& ObjectIniti
 {
 	SetAffectsHeightmap(true);
 
-	ActorIcon = FWaterIconHelper::EnsureSpriteComponentCreated(this, TEXT("/Water/Icons/WaterLandscapeBrushSprite"), LOCTEXT("WaterLandscapeBrushSpriteName", "Water Landscape Brush"));
+	ActorIcon = FWaterIconHelper::EnsureSpriteComponentCreated(this, TEXT("/Water/Icons/WaterLandscapeBrushSprite"));
 }
 
 void AWaterLandscapeBrush::AddActorInternal(AActor* Actor, const UWorld* ThisWorld, UObject* InCache, bool bTriggerEvent, bool bModify)
diff --git a/Engine/Plugins/Experimental/Water/Source/Runtime/Private/WaterBodyActor.cpp b/Engine/Plugins/Experimental/Water/Source/Runtime/Private/WaterBodyActor.cpp
--- a/Engine/Plugins/Experimental/Water/Source/Runtime/Private/WaterBodyActor.cpp
+++ b/Engine/Plugins/Experimental/Water/Source/Runtime/Private/WaterBodyActor.cpp
@@ -106,7 +106,7 @@ AWaterBody::AWaterBody(const FObjectInitializer& ObjectInitializer)
 		SplineComp->OnSplineDataChanged().AddUObject(this, &AWaterBody::OnSplineDataChanged);
 	}
 
-	ActorIcon = FWaterIconHelper::EnsureSpriteComponentCreated(this, TEXT("/Water/Icons/WaterSprite"), LOCTEXT("WaterBodySpriteName", "Water Body"));
+	ActorIcon = FWaterIconHelper::EnsureSpriteComponentCreated(this, TEXT("/Water/Icons/WaterSprite"));
 #endif
 
 	RootComponent = SplineComp;
diff --git a/Engine/Plugins/Experimental/Water/Source/Runtime/Private/WaterBodyCustomActor.cpp b/Engine/Plugins/Experimental/Water/Source/Runtime/Private/WaterBodyCustomActor.cpp
--- a/Engine/Plugins/Experimental/Water/Source/Runtime/Private/WaterBodyCustomActor.cpp
+++ b/Engine/Plugins/Experimental/Water/Source/Runtime/Private/WaterBodyCustomActor.cpp
@@ -18,7 +18,7 @@ AWaterBodyCustom::AWaterBodyCustom(const FObjectInitializer& ObjectInitializer)
 	bAffectsLandscape = false;
 
 #if WITH_EDITOR
-	ActorIcon = FWaterIconHelper::EnsureSpriteComponentCreated(this, TEXT("/Water/Icons/WaterBodyCustomSprite"), NSLOCTEXT("Water", "WaterBodyCustomSpriteName", "Water Body Custom"));
+	ActorIcon = FWaterIconHelper::EnsureSpriteComponentCreated(this, TEXT("/Water/Icons/WaterBodyCustomSprite"));
 #endif
 
 	// @todo_water : Remove these checks (Once AWaterBody is no more Blueprintable, these methods should become PURE_VIRTUAL and this class should overload them)
diff --git a/Engine/Plugins/Experimental/Water/Source/Runtime/Private/WaterBodyExclusionVolume.cpp b/Engine/Plugins/Experimental/Water/Source/Runtime/Private/WaterBodyExclusionVolume.cpp
--- a/Engine/Plugins/Experimental/Water/Source/Runtime/Private/WaterBodyExclusionVolume.cpp
+++ b/Engine/Plugins/Experimental/Water/Source/Runtime/Private/WaterBodyExclusionVolume.cpp
@@ -17,7 +17,7 @@ AWaterBodyExclusionVolume::AWaterBodyExclusionVolume(const FObjectInitializer& O
 	: Super(ObjectInitializer)
 {
 #if WITH_EDITOR
-	ActorIcon = FWaterIconHelper::EnsureSpriteComponentCreated(this, TEXT("/Water/Icons/WaterBodyExclusionVolumeSprite"), NSLOCTEXT("Water", "WaterBodyExclusionVolumeSpriteName", "Water Body Exclusion Volume"));
+	ActorIcon = FWaterIconHelper::EnsureSpriteComponentCreated(this, TEXT("/Water/Icons/WaterBodyExclusionVolumeSprite"));
 #endif
 }
 
diff --git a/Engine/Plugins/Experimental/Water/Source/Runtime/Private/WaterBodyIslandActor.cpp b/Engine/Plugins/Experimental/Water/Source/Runtime/Private/WaterBodyIslandActor.cpp
--- a/Engine/Plugins/Experimental/Water/Source/Runtime/Private/WaterBodyIslandActor.cpp
+++ b/Engine/Plugins/Experimental/Water/Source/Runtime/Private/WaterBodyIslandActor.cpp
@@ -36,7 +36,7 @@ AWaterBodyIsland::AWaterBodyIsland(const FObjectInitializer& ObjectInitializer)
 		SplineComp->OnSplineDataChanged().AddUObject(this, &AWaterBodyIsland::OnSplineDataChanged);
 	}
 
-	ActorIcon = FWaterIconHelper::EnsureSpriteComponentCreated(this, TEXT("/Water/Icons/WaterBodyIslandSprite"), NSLOCTEXT("Water", "WaterBodyIslandSpriteName", "Water Body Island"));
+	ActorIcon = FWaterIconHelper::EnsureSpriteComponentCreated(this, TEXT("/Water/Icons/WaterBodyIslandSprite"));
 #endif
 
 	RootComponent = SplineComp;
diff --git a/Engine/Plugins/Experimental/Water/Source/Runtime/Private/WaterBodyLakeActor.cpp b/Engine/Plugins/Experimental/Water/Source/Runtime/Private/WaterBodyLakeActor.cpp
--- a/Engine/Plugins/Experimental/Water/Source/Runtime/Private/WaterBodyLakeActor.cpp
+++ b/Engine/Plugins/Experimental/Water/Source/Runtime/Private/WaterBodyLakeActor.cpp
@@ -20,7 +20,7 @@ AWaterBodyLake::AWaterBodyLake(const FObjectInitializer& ObjectInitializer)
 	WaterBodyType = GetWaterBodyType();
 
 #if WITH_EDITOR
-	ActorIcon = FWaterIconHelper::EnsureSpriteComponentCreated(this, TEXT("/Water/Icons/WaterBodyLakeSprite"), NSLOCTEXT("Water", "WaterBodyLakeSpriteName", "Water Body Lake"));
+	ActorIcon = FWaterIconHelper::EnsureSpriteComponentCreated(this, TEXT("/Water/Icons/WaterBodyLakeSprite"));
 #endif
 
 	// @todo_water : Remove these checks (Once AWaterBody is no more Blueprintable, these methods should become PURE_VIRTUAL and this class should overload them)
diff --git a/Engine/Plugins/Experimental/Water/Source/Runtime/Private/WaterBodyOceanActor.cpp b/Engine/Plugins/Experimental/Water/Source/Runtime/Private/WaterBodyOceanActor.cpp
--- a/Engine/Plugins/Experimental/Water/Source/Runtime/Private/WaterBodyOceanActor.cpp
+++ b/Engine/Plugins/Experimental/Water/Source/Runtime/Private/WaterBodyOceanActor.cpp
@@ -21,7 +21,7 @@ AWaterBodyOcean::AWaterBodyOcean(const FObjectInitializer& ObjectInitializer)
 	CollisionExtents = FVector(50000.f, 50000.f, 10000.f);
 
 #if WITH_EDITOR
-	ActorIcon = FWaterIconHelper::EnsureSpriteComponentCreated(this, TEXT("/Water/Icons/WaterBodyOceanSprite"), NSLOCTEXT("Water", "WaterBodyOceanSpriteName", "Water Body Ocean"));
+	ActorIcon = FWaterIconHelper::EnsureSpriteComponentCreated(this, TEXT("/Water/Icons/WaterBodyOceanSprite"));
 #endif
 
 	// @todo_water : Remove these checks (Once AWaterBody is no more Blueprintable, these methods should become PURE_VIRTUAL and this class should overload them)
diff --git a/Engine/Plugins/Experimental/Water/Source/Runtime/Private/WaterBodyRiverActor.cpp b/Engine/Plugins/Experimental/Water/Source/Runtime/Private/WaterBodyRiverActor.cpp
--- a/Engine/Plugins/Experimental/Water/Source/Runtime/Private/WaterBodyRiverActor.cpp
+++ b/Engine/Plugins/Experimental/Water/Source/Runtime/Private/WaterBodyRiverActor.cpp
@@ -29,7 +29,7 @@ AWaterBodyRiver::AWaterBodyRiver(const FObjectInitializer& ObjectInitializer)
 	WaterBodyType = GetWaterBodyType();
 
 #if WITH_EDITOR
-	ActorIcon = FWaterIconHelper::EnsureSpriteComponentCreated(this, TEXT("/Water/Icons/WaterBodyRiverSprite"), NSLOCTEXT("Water", "WaterBodyRiverSpriteName", "Water Body River Sprite"));
+	ActorIcon = FWaterIconHelper::EnsureSpriteComponentCreated(this, TEXT("/Water/Icons/WaterBodyRiverSprite"));
 #endif
 
 	// @todo_water : Remove these checks (Once AWaterBody is no more Blueprintable, these methods should become PURE_VIRTUAL and this class should overload them)
diff --git a/Engine/Plugins/Experimental/Water/Source/Runtime/Private/WaterIconHelper.cpp b/Engine/Plugins/Experimental/Water/Source/Runtime/Private/WaterIconHelper.cpp
--- a/Engine/Plugins/Experimental/Water/Source/Runtime/Private/WaterIconHelper.cpp
+++ b/Engine/Plugins/Experimental/Water/Source/Runtime/Private/WaterIconHelper.cpp
@@ -11,34 +11,29 @@
 #include "Modules/ModuleManager.h"
 #include "WaterModule.h"
 
-UBillboardComponent* FWaterIconHelper::EnsureSpriteComponentCreated_Internal(AActor* Actor, UClass* InClass, const TCHAR* InIconTextureName, const FText& InDisplayName)
+UBillboardComponent* FWaterIconHelper::EnsureSpriteComponentCreated_Internal(AActor* Actor, UClass* InClass, const TCHAR* InIconTextureName)
 {
 	UBillboardComponent* ActorIcon = nullptr;
-	if (!Actor->IsTemplate())
+	
+	ActorIcon = Actor->FindComponentByClass<UBillboardComponent>();
+	if (ActorIcon == nullptr)
 	{
-		ActorIcon = Actor->FindComponentByClass<UBillboardComponent>();
-		if (ActorIcon == nullptr)
-		{
-			ActorIcon = Actor->CreateEditorOnlyDefaultSubobject<UBillboardComponent>(TEXT("Sprite"), true);
-		}
-
-		if (ActorIcon != nullptr)
+		ActorIcon = Actor->CreateEditorOnlyDefaultSubobject<UBillboardComponent>(TEXT("Sprite"), true);
+	}
+	if (ActorIcon != nullptr)
+	{
+		ConstructorHelpers::FObjectFinderOptional<UTexture2D> Texture(InIconTextureName);
+		IWaterModuleInterface& WaterModule = FModuleManager::GetModuleChecked<IWaterModuleInterface>(TEXT("Water"));
+		if (IWaterEditorServices* WaterEditorServices = WaterModule.GetWaterEditorServices())
 		{
-			ConstructorHelpers::FObjectFinderOptional<UTexture2D> Texture(InIconTextureName);
-
-			IWaterModuleInterface& WaterModule = FModuleManager::GetModuleChecked<IWaterModuleInterface>(TEXT("Water"));
-			if (IWaterEditorServices* WaterEditorServices = WaterModule.GetWaterEditorServices())
-			{
-				WaterEditorServices->RegisterWaterActorSprite(InClass, Texture.Get());
-			}
-
-			ActorIcon->Sprite = Texture.Get();
-			ActorIcon->bHiddenInGame = true;
-			ActorIcon->SpriteInfo.Category = TEXT("Water");
-			ActorIcon->SpriteInfo.DisplayName = InDisplayName;
-			ActorIcon->SetupAttachment(Actor->GetRootComponent());
-			UpdateSpriteComponent(Actor, ActorIcon->Sprite);
+			WaterEditorServices->RegisterWaterActorSprite(InClass, Texture.Get());
 		}
+		ActorIcon->Sprite = Texture.Get();
+		ActorIcon->bHiddenInGame = true;
+		ActorIcon->SpriteInfo.Category = TEXT("Water");
+		ActorIcon->SpriteInfo.DisplayName = NSLOCTEXT("SpriteCategory", "Water", "Water");
+		ActorIcon->SetupAttachment(Actor->GetRootComponent());
+		UpdateSpriteComponent(Actor, ActorIcon->Sprite);
 	}
 	return ActorIcon;
 }
diff --git a/Engine/Plugins/Experimental/Water/Source/Runtime/Private/WaterMeshActor.cpp b/Engine/Plugins/Experimental/Water/Source/Runtime/Private/WaterMeshActor.cpp
--- a/Engine/Plugins/Experimental/Water/Source/Runtime/Private/WaterMeshActor.cpp
+++ b/Engine/Plugins/Experimental/Water/Source/Runtime/Private/WaterMeshActor.cpp
@@ -24,7 +24,7 @@ AWaterMeshActor::AWaterMeshActor(const FObjectInitializer& Initializer)
 		LevelEditorModule.OnActorSelectionChanged().AddUObject(this, &AWaterMeshActor::OnActorSelectionChanged);
 	}
 
-	ActorIcon = FWaterIconHelper::EnsureSpriteComponentCreated(this, TEXT("/Water/Icons/WaterMeshActorSprite"), NSLOCTEXT("Water", "WaterMeshActorSpriteName", "Water Mesh Actor"));
+	ActorIcon = FWaterIconHelper::EnsureSpriteComponentCreated(this, TEXT("/Water/Icons/WaterMeshActorSprite"));
 #endif
 }
 
diff --git a/Engine/Plugins/Experimental/Water/Source/Runtime/Public/WaterIconHelper.h b/Engine/Plugins/Experimental/Water/Source/Runtime/Public/WaterIconHelper.h
--- a/Engine/Plugins/Experimental/Water/Source/Runtime/Public/WaterIconHelper.h
+++ b/Engine/Plugins/Experimental/Water/Source/Runtime/Public/WaterIconHelper.h
@@ -15,16 +15,23 @@ struct WATER_API FWaterIconHelper
 {
 	/** Ensures a billboard component is created and added to the actor's components. */
 	template <typename ActorType>
+	static UBillboardComponent* EnsureSpriteComponentCreated(ActorType* Actor, const TCHAR* InIconTextureName)
+	{
+		return EnsureSpriteComponentCreated_Internal(Actor, ActorType::StaticClass(), InIconTextureName);
+	}
+	
+	template <typename ActorType>
+	UE_DEPRECATED(4.27, "void EnsureSpriteComponentCreated(Actor, InIconTextureName, InDisplayName) is deprecated. Use EnsureSpriteComponentCreated(Actor, InIconTextureName) instead.")
 	static UBillboardComponent* EnsureSpriteComponentCreated(ActorType* Actor, const TCHAR* InIconTextureName, const FText& InDisplayName)
 	{
-		return EnsureSpriteComponentCreated_Internal(Actor, ActorType::StaticClass(), InIconTextureName, InDisplayName);
+		return EnsureSpriteComponentCreated_Internal(Actor, ActorType::StaticClass(), InIconTextureName);
 	}
 
 	/** Updates the texture/scale/position of the actor's billboard component, if any */
 	static void UpdateSpriteComponent(AActor* Actor, UTexture2D* InTexture);
 
 private:
-	static UBillboardComponent* EnsureSpriteComponentCreated_Internal(AActor* Actor, UClass* InClass, const TCHAR* InIconTextureName, const FText& InDisplayName);
+	static UBillboardComponent* EnsureSpriteComponentCreated_Internal(AActor* Actor, UClass* InClass, const TCHAR* InIconTextureName);
 };
 
 #endif
\ No newline at end of file
