Fixed water selection (don't keep the WaterMeshComponent in the hit proxy, as it will take precedence against the water body actor)

#rnx

diff --git a/Engine/Plugins/Experimental/Water/Source/Runtime/Private/WaterMeshComponent.cpp b/Engine/Plugins/Experimental/Water/Source/Runtime/Private/WaterMeshComponent.cpp
--- a/Engine/Plugins/Experimental/Water/Source/Runtime/Private/WaterMeshComponent.cpp
+++ b/Engine/Plugins/Experimental/Water/Source/Runtime/Private/WaterMeshComponent.cpp
@@ -238,7 +238,7 @@ void UWaterMeshComponent::RebuildWaterMesh(float InTileSize, const FIntPoint& In
 		RenderData.SurfaceBaseHeight = WaterBody->GetActorLocation().Z;
 		RenderData.WaterBodyType = (int8)WaterBodyType;
 #if WITH_WATER_SELECTION_SUPPORT
-		RenderData.HitProxy = new HActor(WaterBody, this);
+		RenderData.HitProxy = new HActor(/*InActor = */WaterBody, /*InPrimComponent = */nullptr);
 		RenderData.bWaterBodySelected = WaterBody->IsSelected();
 #endif // WITH_WATER_SELECTION_SUPPORT
 
