Fixed crash due to water quad tree not allocating enough nodes upon construction.

diff --git a/Engine/Plugins/Experimental/Water/Source/Runtime/Private/WaterQuadTree.cpp b/Engine/Plugins/Experimental/Water/Source/Runtime/Private/WaterQuadTree.cpp
--- a/Engine/Plugins/Experimental/Water/Source/Runtime/Private/WaterQuadTree.cpp
+++ b/Engine/Plugins/Experimental/Water/Source/Runtime/Private/WaterQuadTree.cpp
@@ -312,6 +312,8 @@ void FWaterQuadTree::FNode::AddNodes(FNodeData& InNodeData, const FBox& InMeshBo
 
 			if (ChildBounds.IntersectXY(InWaterBodyBounds) && ChildBounds.IntersectXY(InMeshBounds))
 			{
+				// All nodes have been allocated upfront, no reallocation should occur : 
+				check(InNodeData.Nodes.Num() < InNodeData.Nodes.Max());
 				Children[i] = InNodeData.Nodes.Emplace();
 				InNodeData.Nodes[Children[i]].Bounds = ChildBounds; 
 				InNodeData.Nodes[Children[i]].ParentIndex = InParentIndex;
@@ -421,11 +423,16 @@ void FWaterQuadTree::InitTree(const FBox2D& InBounds, float InTileSize, FIntPoin
 	LeafSize = InTileSize;
 	ExtentInTiles = InExtentInTiles;
 
+	// Calculate the depth of the tree. This also corresponds to the LOD count. 0 means root is leaf node
+	// Find a pow2 tile resolution that contains the user defined extent in tiles
+	const int32 MaxDim = (int32)FMath::Max(InExtentInTiles.X * 2, InExtentInTiles.Y * 2);
+	const float RootDim = (float)FMath::RoundUpToPowerOfTwo(MaxDim);
+
 	TileRegion = InBounds;
 
 	// Allocate theoretical max, shrink later in Lock()
 	// This is so that the node array doesn't move in memory while inserting
-	NodeData.Nodes.Empty((float)(MaxLeafCount * 4) / 3.0f);
+	NodeData.Nodes.Empty((float)(FMath::Square(RootDim) * 4) / 3.0f);
 
 	// Add defaulted water body render data to slot 0. This is the "null" render data, pointed to by all newly created nodes. Has lowest priority so it will always be overwritten
 	NodeData.WaterBodyRenderData.Empty(1);
@@ -436,11 +443,6 @@ void FWaterQuadTree::InitTree(const FBox2D& InBounds, float InTileSize, FIntPoin
 	// Add the root node at slot 0
 	NodeData.Nodes.Emplace();
 
-
-	// Calculate the depth of the tree. This also corresponds to the LOD count. 0 means root is leaf node
-	// Find a pow2 tile resolution that contains the user defined extent in tiles
-	const int32 MaxDim = (int32)FMath::Max(InExtentInTiles.X*2, InExtentInTiles.Y*2);
-	const float RootDim = (float)FMath::RoundUpToPowerOfTwo(MaxDim);
 	const float RootWorldSize = RootDim * InTileSize;
 
 	TreeDepth = (int32)FMath::Log2(RootDim);
