Disabled Nanite during water brush's scene captures to workaround the fact that it doesn't support USceneCaptureComponent's ShowOnlyComponents ATM.

diff --git a/Engine/Plugins/Experimental/Water/Source/Editor/Private/WaterBrushManager.cpp b/Engine/Plugins/Experimental/Water/Source/Editor/Private/WaterBrushManager.cpp
--- a/Engine/Plugins/Experimental/Water/Source/Editor/Private/WaterBrushManager.cpp
+++ b/Engine/Plugins/Experimental/Water/Source/Editor/Private/WaterBrushManager.cpp
@@ -47,6 +47,8 @@ AWaterBrushManager::AWaterBrushManager(const FObjectInitializer& ObjectInitializ
 	SceneCaptureComponent2D->bCaptureOnMovement = false;
 	SceneCaptureComponent2D->SetRelativeRotation(FRotator(-90.0f, 0.0f, -90.0f));
 	SceneCaptureComponent2D->SetRelativeScale3D(FVector(0.01f, 0.01f, 0.01f));
+	// HACK [jonathan.bard] : Nanite doesn't support USceneCaptureComponent's ShowOnlyComponents ATM so just disable Nanite during captures : 
+	SceneCaptureComponent2D->ShowFlagSettings.Add(FEngineShowFlagsSetting { TEXT("NaniteMeshes"), false } );
 
 	PrimaryActorTick.TickGroup = ETickingGroup::TG_PrePhysics;
 	bIsEditorOnlyActor = false;
