Fix motion vectors for water vertex factory

#rnx

diff --git a/Engine/Plugins/Experimental/Water/Shaders/Private/WaterMeshVertexFactory.ush b/Engine/Plugins/Experimental/Water/Shaders/Private/WaterMeshVertexFactory.ush
--- a/Engine/Plugins/Experimental/Water/Shaders/Private/WaterMeshVertexFactory.ush
+++ b/Engine/Plugins/Experimental/Water/Shaders/Private/WaterMeshVertexFactory.ush
@@ -343,12 +343,7 @@ float4 VertexFactoryGetWorldPosition(FPositionOnlyVertexFactoryInput Input)
 // @return previous translated world position
 float4 VertexFactoryGetPreviousWorldPosition(FVertexFactoryInput Input, FVertexFactoryIntermediates Intermediates)
 {
-	float4x4 PreviousLocalToWorldTranslated = GetPrimitiveData(Intermediates.PrimitiveId).PreviousLocalToWorld;
-	PreviousLocalToWorldTranslated[3][0] += ResolvedView.PrevPreViewTranslation.x;
-	PreviousLocalToWorldTranslated[3][1] += ResolvedView.PrevPreViewTranslation.y;
-	PreviousLocalToWorldTranslated[3][2] += ResolvedView.PrevPreViewTranslation.z;
-
-	return mul(Input.Position, PreviousLocalToWorldTranslated);
+	return float4(Intermediates.MorphedWorldPosRaw + ResolvedView.PrevPreViewTranslation.xyz, 1);
 }
 
 float4 VertexFactoryGetTranslatedPrimitiveVolumeBounds(FVertexFactoryInterpolantsVSToPS Interpolants)
