Hide transient actor UBuoyancyManager from scene outliner (it will eventually become a subsystem)

diff --git a/Engine/Plugins/Experimental/Water/Source/Runtime/Private/WaterSubsystem.cpp b/Engine/Plugins/Experimental/Water/Source/Runtime/Private/WaterSubsystem.cpp
--- a/Engine/Plugins/Experimental/Water/Source/Runtime/Private/WaterSubsystem.cpp
+++ b/Engine/Plugins/Experimental/Water/Source/Runtime/Private/WaterSubsystem.cpp
@@ -262,6 +262,12 @@ void UWaterSubsystem::Initialize(FSubsystemCollectionBase& Collection)
 		FActorSpawnParameters SpawnInfo;
 		SpawnInfo.SpawnCollisionHandlingOverride = ESpawnActorCollisionHandlingMethod::AlwaysSpawn;
 		SpawnInfo.ObjectFlags = RF_Transient;
+
+#if WITH_EDITOR
+		// The buoyancy manager should be a subsytem really, but for now, just hide it from the outliner : 
+		SpawnInfo.bHideFromSceneOutliner = true;
+#endif //WITH_EDITOR
+
 		// Store the buoyancy manager we create for future use.
 		BuoyancyManager = World->SpawnActor<ABuoyancyManager>(SpawnInfo);
 	}
@@ -444,39 +450,39 @@ void UWaterSubsystem::SetOceanFloodHeight(float InFloodHeight)
 {
 	if (UWorld* World = GetWorld())
 	{
-		const float ClampedFloodHeight = FMath::Max(0.0f, InFloodHeight);
+	const float ClampedFloodHeight = FMath::Max(0.0f, InFloodHeight);
 
-		if (FloodHeight != ClampedFloodHeight)
-		{
-			FloodHeight = ClampedFloodHeight;
-			MarkAllWaterMeshesForRebuild();
+	if (FloodHeight != ClampedFloodHeight)
+	{
+		FloodHeight = ClampedFloodHeight;
+		MarkAllWaterMeshesForRebuild();
 
-			// the ocean body is dynamic and needs to be readjusted when the flood height changes : 
-			if (OceanActor.IsValid())
-			{
-				OceanActor->SetHeightOffset(InFloodHeight);
-			}
+		// the ocean body is dynamic and needs to be readjusted when the flood height changes : 
+		if (OceanActor.IsValid())
+		{
+			OceanActor->SetHeightOffset(InFloodHeight);
+		}
 
-			// All water body actors need to update their underwater post process MID as it depends on the ocean global height : 
+		// All water body actors need to update their underwater post process MID as it depends on the ocean global height : 
 			for (TActorIterator<AWaterBody> ActorItr(World); ActorItr; ++ActorItr)
-			{
-				AWaterBody* WaterBody = *ActorItr;
-				WaterBody->UpdateMaterialInstances();
-			}
+		{
+			AWaterBody* WaterBody = *ActorItr;
+			WaterBody->UpdateMaterialInstances();
 		}
 	}
 }
+}
 
 AWaterMeshActor* UWaterSubsystem::GetWaterMeshActor() const
 {
 	if (UWorld* World = GetWorld())
 	{
-		// @todo water: this assumes only one water mesh actor right now.  In the future we may need to associate a water mesh actor with a water body more directly
+	// @todo water: this assumes only one water mesh actor right now.  In the future we may need to associate a water mesh actor with a water body more directly
 		TActorIterator<AWaterMeshActor> It(World);
-		WaterMeshActor = It ? *It : nullptr;
+	WaterMeshActor = It ? *It : nullptr;
 
-		return WaterMeshActor;
-	}
+	return WaterMeshActor;
+}
 
 	return nullptr;
 }
@@ -496,11 +502,11 @@ void UWaterSubsystem::MarkAllWaterMeshesForRebuild()
 	if (UWorld* World = GetWorld())
 	{
 		for (AWaterMeshActor* WaterMesh : TActorRange<AWaterMeshActor>(World))
-		{
-			WaterMesh->MarkWaterMeshComponentForRebuild();
-		}
+	{
+		WaterMesh->MarkWaterMeshComponentForRebuild();
 	}
 }
+}
 
 void UWaterSubsystem::NotifyWaterScalabilityChangedInternal(IConsoleVariable* CVar)
 {
@@ -511,13 +517,13 @@ void UWaterSubsystem::NotifyWaterEnabledChangedInternal(IConsoleVariable* CVar)
 {
 	if (UWorld* World = GetWorld())
 	{
-		// Water body visibility depends on CVarWaterEnabled
+	// Water body visibility depends on CVarWaterEnabled
 		for (AWaterBody* WaterBody : TActorRange<AWaterBody>(World))
-		{
-			WaterBody->UpdateWaterComponentVisibility();
-		}
+	{
+		WaterBody->UpdateWaterComponentVisibility();
 	}
 }
+}
 
 struct FWaterBodyPostProcessQuery
 {
@@ -791,4 +797,4 @@ UWaterSubsystem::FScopedAllowWaterSubsystemOnPreviewWorld::~FScopedAllowWaterSub
 
 // ----------------------------------------------------------------------------------
 
-#undef LOCTEXT_NAMESPACE
+#undef LOCTEXT_NAMESPACE
\ No newline at end of file
