Fixed division by zero crash when using 0 scale on lakes

diff --git a/Engine/Plugins/Experimental/Water/Source/Runtime/Private/WaterBodyLakeActor.cpp b/Engine/Plugins/Experimental/Water/Source/Runtime/Private/WaterBodyLakeActor.cpp
--- a/Engine/Plugins/Experimental/Water/Source/Runtime/Private/WaterBodyLakeActor.cpp
+++ b/Engine/Plugins/Experimental/Water/Source/Runtime/Private/WaterBodyLakeActor.cpp
@@ -150,7 +150,10 @@ void ULakeGenerator::OnUpdateBody(bool bWithExclusionVolumes)
 		LakeCollision->SetGenerateOverlapEvents(true);
 
 		const float Depth = OwnerBody->GetChannelDepth() / 2;
-		FVector LakeCollisionExtent = FVector(SplineExtent.X, SplineExtent.Y, Depth) / OwnerBody->GetActorScale();
+		FVector Scale = OwnerBody->GetActorScale();
+		// Avoid dividing by 0 but keep the scale's sign :
+		Scale = Scale.GetSignVector() * FVector::Max(Scale.GetAbs(), FVector::OneVector * 0.001f);
+		FVector LakeCollisionExtent = FVector(SplineExtent.X, SplineExtent.Y, Depth) / Scale;
 		LakeCollision->SetWorldLocation(WorldLoc + FVector(0, 0, -Depth));
 		LakeCollision->UpdateCollision(LakeCollisionExtent, true);
 	}
