Fixed linker load error when creating water bodies through the factory : the inner WaterWaves object should be outered with the water body actor, not the actor's outer, as this changes upon save in World Partition

diff --git a/Engine/Plugins/Experimental/Water/Source/Editor/Private/WaterBodyActorFactory.cpp b/Engine/Plugins/Experimental/Water/Source/Editor/Private/WaterBodyActorFactory.cpp
--- a/Engine/Plugins/Experimental/Water/Source/Editor/Private/WaterBodyActorFactory.cpp
+++ b/Engine/Plugins/Experimental/Water/Source/Editor/Private/WaterBodyActorFactory.cpp
@@ -114,7 +114,7 @@ void UWaterBodyOceanActorFactory::PostSpawnActor(UObject* Asset, AActor* NewActo
 	AWaterBodyOcean* WaterBodyOcean = CastChecked<AWaterBodyOcean>(NewActor);
 	if (const UWaterWavesBase* DefaultWaterWaves = GetDefault<UWaterEditorSettings>()->WaterBodyOceanDefaults.WaterWaves)
 	{
-		UWaterWavesBase* WaterWaves = DuplicateObject(DefaultWaterWaves, NewActor->GetOuter(), MakeUniqueObjectName(NewActor->GetOuter(), DefaultWaterWaves->GetClass(), TEXT("OceanWaterWaves")));
+		UWaterWavesBase* WaterWaves = DuplicateObject(DefaultWaterWaves, NewActor, MakeUniqueObjectName(NewActor, DefaultWaterWaves->GetClass(), TEXT("OceanWaterWaves")));
 		WaterBodyOcean->SetWaterWaves(WaterWaves);
 	}
 
@@ -149,7 +149,7 @@ void UWaterBodyLakeActorFactory::PostSpawnActor(UObject* Asset, AActor* NewActor
 	AWaterBodyLake* WaterBodyLake = CastChecked<AWaterBodyLake>(NewActor);
 	if (const UWaterWavesBase* DefaultWaterWaves = GetDefault<UWaterEditorSettings>()->WaterBodyLakeDefaults.WaterWaves)
 	{
-		UWaterWavesBase* WaterWaves = DuplicateObject(DefaultWaterWaves, NewActor->GetOuter(), MakeUniqueObjectName(NewActor->GetOuter(), DefaultWaterWaves->GetClass(), TEXT("LakeWaterWaves")));
+		UWaterWavesBase* WaterWaves = DuplicateObject(DefaultWaterWaves, NewActor, MakeUniqueObjectName(NewActor, DefaultWaterWaves->GetClass(), TEXT("LakeWaterWaves")));
 		WaterBodyLake->SetWaterWaves(WaterWaves);
 	}
 
diff --git a/Engine/Plugins/Experimental/Water/Source/Runtime/Private/WaterBodyActor.cpp b/Engine/Plugins/Experimental/Water/Source/Runtime/Private/WaterBodyActor.cpp
--- a/Engine/Plugins/Experimental/Water/Source/Runtime/Private/WaterBodyActor.cpp
+++ b/Engine/Plugins/Experimental/Water/Source/Runtime/Private/WaterBodyActor.cpp
@@ -1373,7 +1373,19 @@ void AWaterBody::PostLoad()
 		static_assert(sizeof(WaterHeightmapSettings) == sizeof(TerrainCarvingSettings_DEPRECATED), "Both old and old water heightmap settings struct should be exactly similar");
 		FMemory::Memcpy((void*)&WaterHeightmapSettings, (void*)&TerrainCarvingSettings_DEPRECATED, sizeof(WaterHeightmapSettings));
 	}
-#endif
+#endif // WITH_EDITORONLY_DATA
+
+#if WITH_EDITOR
+	if (GIsEditor && !HasAnyFlags(RF_ClassDefaultObject))
+	{
+		if (WaterWaves && (WaterWaves->GetOuter() != this))
+		{
+			WaterWaves->ClearFlags(RF_Public);
+			// At one point, WaterWaves's outer was the level. We need them to be outered by the water body : 
+			WaterWaves->Rename(nullptr, this, REN_DoNotDirty | REN_DontCreateRedirectors | REN_ForceNoResetLoaders | REN_NonTransactional);
+		}
+	}
+#endif // WITH_EDITOR
 
 #if WITH_EDITOR
 	RegisterOnUpdateWavesData(WaterWaves, /* bRegister = */true);
