Fixed invalid river spline mesh collision component when going from 2 spline points to 1 and from 0 to 1 (undo)

diff --git a/Engine/Plugins/Experimental/Water/Source/Editor/Private/WaterBrushManager.cpp b/Engine/Plugins/Experimental/Water/Source/Editor/Private/WaterBrushManager.cpp
--- a/Engine/Plugins/Experimental/Water/Source/Editor/Private/WaterBrushManager.cpp
+++ b/Engine/Plugins/Experimental/Water/Source/Editor/Private/WaterBrushManager.cpp
@@ -516,13 +516,15 @@ bool AWaterBrushManager::SetupRiverSplineRenderMIDs(const FBrushActorRenderConte
 
 	for (int32 MIDIndex = 0; MIDIndex < NumSplineMids; ++MIDIndex)
 	{
+		USplineMeshComponent* SplineMeshComponent = SplineMeshComponents[MIDIndex];
+		check(SplineMeshComponent != nullptr);
 		if (bClearMIDs)
 		{
-			SplineMeshComponents[MIDIndex]->SetMaterial(0, nullptr);
+			SplineMeshComponent->SetMaterial(0, nullptr);
 		}
 		else
 		{
-			check(MIDIndex < SplineMeshComponents.Num());
+			RiverSplineMIDs.IsValidIndex(MIDIndex);
 			RiverSplineMIDs[MIDIndex] = FWaterUtils::GetOrCreateTransientMID(RiverSplineMIDs[MIDIndex], TEXT("RiverSplineMID"), RenderRiverSplineDepthMaterial);
 			UMaterialInstanceDynamic* TempMID = RiverSplineMIDs[MIDIndex];
 			if (TempMID != nullptr)
@@ -533,7 +535,7 @@ bool AWaterBrushManager::SetupRiverSplineRenderMIDs(const FBrushActorRenderConte
 				TempMID->SetScalarParameterValue(FName(TEXT("VelA")), SplineComponent->GetFloatPropertyAtSplinePoint(MIDIndex, FName(TEXT("WaterVelocityScalar"))));
 				TempMID->SetScalarParameterValue(FName(TEXT("VelB")), SplineComponent->GetFloatPropertyAtSplinePoint(MIDIndex + 1, FName(TEXT("WaterVelocityScalar"))));
 
-				SplineMeshComponents[MIDIndex]->SetMaterial(0, TempMID);
+				SplineMeshComponent->SetMaterial(0, TempMID);
 			}
 			else
 			{
diff --git a/Engine/Plugins/Experimental/Water/Source/Runtime/Private/WaterBodyRiverActor.cpp b/Engine/Plugins/Experimental/Water/Source/Runtime/Private/WaterBodyRiverActor.cpp
--- a/Engine/Plugins/Experimental/Water/Source/Runtime/Private/WaterBodyRiverActor.cpp
+++ b/Engine/Plugins/Experimental/Water/Source/Runtime/Private/WaterBodyRiverActor.cpp
@@ -9,7 +9,7 @@
 #include "PhysicsEngine/ConvexElem.h"
 #include "PhysicsEngine/BodySetup.h"
 #include "Engine/StaticMesh.h"
-#include "Algo/Transform.h"
+#include "Algo/ForEach.h"
 
 #if WITH_CHAOS
 // for working around Chaos issue
@@ -103,7 +103,13 @@ TArray<UPrimitiveComponent*> AWaterBodyRiver::GetBrushRenderableComponents() con
 	{
 		const TArray<USplineMeshComponent*>& SplineMeshComponents = RiverGenerator->GetSplineMeshComponents();
 		BrushRenderableComponents.Reserve(SplineMeshComponents.Num());
-		Algo::Transform(SplineMeshComponents, BrushRenderableComponents, [](USplineMeshComponent* PrimitiveComponent) { return StaticCast<UPrimitiveComponent*>(PrimitiveComponent); });
+		Algo::ForEach(SplineMeshComponents, [&BrushRenderableComponents](USplineMeshComponent* Comp)
+		{
+			if (Comp != nullptr)
+			{
+				BrushRenderableComponents.Add(StaticCast<UPrimitiveComponent*>(Comp));
+			}
+		});
 	}
 	return BrushRenderableComponents;
 }
@@ -169,7 +175,7 @@ void URiverGenerator::OnUpdateBody(bool bWithExclusionVolumes)
 	{
 		const int32 NumSplinePoints = WaterSplineComp->GetNumberOfSplinePoints();
 		const int32 NumberOfMeshComponentsNeeded = WaterSplineComp->IsClosedLoop() ? NumSplinePoints : NumSplinePoints - 1;
-		if (NumSplinePoints > 1 && SplineMeshComponents.Num() != NumberOfMeshComponentsNeeded)
+		if (SplineMeshComponents.Num() != NumberOfMeshComponentsNeeded)
 		{
 			GenerateMeshes();
 		}
@@ -201,7 +207,7 @@ TArray<UPrimitiveComponent*> URiverGenerator::GetCollisionComponents() const
 	Result.Reserve(SplineMeshComponents.Num());
 	for (USplineMeshComponent* Comp : SplineMeshComponents)
 	{
-		if (Comp->GetCollisionEnabled() != ECollisionEnabled::NoCollision)
+		if ((Comp != nullptr) && (Comp->GetCollisionEnabled() != ECollisionEnabled::NoCollision))
 		{
 			Result.Add(Comp);
 		}
