Fixed static analysis

#rnx

diff --git a/Engine/Plugins/Experimental/Water/Source/Runtime/Private/WaterBodyActor.cpp b/Engine/Plugins/Experimental/Water/Source/Runtime/Private/WaterBodyActor.cpp
--- a/Engine/Plugins/Experimental/Water/Source/Runtime/Private/WaterBodyActor.cpp
+++ b/Engine/Plugins/Experimental/Water/Source/Runtime/Private/WaterBodyActor.cpp
@@ -1359,6 +1359,7 @@ void AWaterBody::PostLoad()
 		// At one point, some attributes from UGerstnerWaterWaves were transient, recompute those here at load-time (nowadays, they are serialized properly so they should be properly recompute on property change)
 		if (HasWaves())
 		{
+			check(WaterWaves != nullptr);
 			if (UGerstnerWaterWaves* GerstnerWaterWaves = Cast<UGerstnerWaterWaves>(WaterWaves->GetWaterWaves()))
 			{
 				GerstnerWaterWaves->ConditionalPostLoad();
