GerstnerWater: Handle null GEngine in shutdown.
#rnx

diff --git a/Engine/Plugins/Experimental/Water/Source/Runtime/Private/GerstnerWaterWaveViewExtension.cpp b/Engine/Plugins/Experimental/Water/Source/Runtime/Private/GerstnerWaterWaveViewExtension.cpp
--- a/Engine/Plugins/Experimental/Water/Source/Runtime/Private/GerstnerWaterWaveViewExtension.cpp
+++ b/Engine/Plugins/Experimental/Water/Source/Runtime/Private/GerstnerWaterWaveViewExtension.cpp
@@ -15,6 +15,7 @@ FGerstnerWaterWaveViewExtension::~FGerstnerWaterWaveViewExtension()
 
 void FGerstnerWaterWaveViewExtension::Initialize()
 {
+	check(GEngine);
 	if (UGerstnerWaterWaveSubsystem* GerstnerWaterWaveSubsystem = GEngine->GetEngineSubsystem<UGerstnerWaterWaveSubsystem>())
 	{
 		GerstnerWaterWaveSubsystem->Register(this);
@@ -29,9 +30,12 @@ void FGerstnerWaterWaveViewExtension::Deinitialize()
 		[WaveGPUData=WaveGPUData](FRHICommandListImmediate& RHICmdList){}
 	);
 
-	if (UGerstnerWaterWaveSubsystem* GerstnerWaterWaveSubsystem = GEngine->GetEngineSubsystem<UGerstnerWaterWaveSubsystem>())
+	if (GEngine)
 	{
-		GerstnerWaterWaveSubsystem->Unregister(this);
+		if (UGerstnerWaterWaveSubsystem* GerstnerWaterWaveSubsystem = GEngine->GetEngineSubsystem<UGerstnerWaterWaveSubsystem>())
+		{
+			GerstnerWaterWaveSubsystem->Unregister(this);
+		}
 	}
 }
 
