Fixed crashes when undoing and then redoing a water body actor deletion transaction:
PostEditUndo is called on the "deleted" actor, while it just had its internal references nullified by GC. For water body (island) actors, this means that it's possible to have a null SplineComp (root component) or WaterSplineMetadata during any call made by PostEditUndo

diff --git a/Engine/Plugins/Experimental/Water/Source/Runtime/Private/WaterBodyActor.cpp b/Engine/Plugins/Experimental/Water/Source/Runtime/Private/WaterBodyActor.cpp
--- a/Engine/Plugins/Experimental/Water/Source/Runtime/Private/WaterBodyActor.cpp
+++ b/Engine/Plugins/Experimental/Water/Source/Runtime/Private/WaterBodyActor.cpp
@@ -334,13 +334,16 @@ bool AWaterBody::IsNavigationRelevant() const
 
 float AWaterBody::FindInputKeyClosestToWorldLocation(const FVector& WorldLocation) const
 {
-	return GetWaterSpline()->FindInputKeyClosestToWorldLocation(WorldLocation);
+	UWaterSplineComponent* WaterSpline = GetWaterSpline();
+	return WaterSpline ? WaterSpline->FindInputKeyClosestToWorldLocation(WorldLocation) : 0.0f;
 }
 
 float AWaterBody::GetConstantSurfaceZ() const
 {
+	UWaterSplineComponent* WaterSpline = GetWaterSpline();
+
 	// A single Z doesn't really make sense for non-flat water bodies, but it can be useful for when using FixedZ post process for example. Take the first spline key in that case : 
-	float WaterSurfaceZ = IsFlatSurface() ? GetActorLocation().Z : GetWaterSpline()->GetLocationAtSplineInputKey(0.0f, ESplineCoordinateSpace::World).Z;
+	float WaterSurfaceZ = (IsFlatSurface() || WaterSpline == nullptr) ? GetActorLocation().Z : WaterSpline->GetLocationAtSplineInputKey(0.0f, ESplineCoordinateSpace::World).Z;
 	
 	// Apply body height offset if applicable (ocean)
 	if (IsHeightOffsetSupported())
@@ -354,7 +357,8 @@ float AWaterBody::GetConstantSurfaceZ() const
 float AWaterBody::GetConstantDepth() const
 {
 	// Only makes sense when you consider the water depth to be constant for the whole water body, in which case we just use the first spline key's : 
-	return GetWaterSpline()->GetFloatPropertyAtSplineInputKey(0.0f, GET_MEMBER_NAME_CHECKED(UWaterSplineMetadata, Depth));
+	UWaterSplineComponent* WaterSpline = GetWaterSpline();
+	return WaterSpline ? WaterSpline->GetFloatPropertyAtSplineInputKey(0.0f, GET_MEMBER_NAME_CHECKED(UWaterSplineMetadata, Depth)) : 0.0f;
 }
 
 void AWaterBody::GetSurfaceMinMaxZ(float& OutMinZ, float& OutMaxZ) const
@@ -582,20 +586,21 @@ FWaterBodyQueryResult AWaterBody::QueryWaterInfoClosestToWorldLocation(const FVe
 
 float AWaterBody::GetWaterVelocityAtSplineInputKey(float InKey) const
 {
-	return WaterSplineMetadata->WaterVelocityScalar.Eval(InKey, 0.f);
+	return WaterSplineMetadata ? WaterSplineMetadata->WaterVelocityScalar.Eval(InKey, 0.f) : 0.0f;
 }
 
 FVector AWaterBody::GetWaterVelocityVectorAtSplineInputKey(float InKey) const
 {
+	UWaterSplineComponent* WaterSpline = GetWaterSpline();
 	float WaterVelocityScalar = GetWaterVelocityAtSplineInputKey(InKey);
-	const FVector SplineDirection = GetWaterSpline()->GetDirectionAtSplineInputKey(InKey, ESplineCoordinateSpace::World);
+	const FVector SplineDirection = WaterSpline ? WaterSpline->GetDirectionAtSplineInputKey(InKey, ESplineCoordinateSpace::World) : FVector::ZeroVector;
 	return SplineDirection * WaterVelocityScalar;
 }
 
 
 float AWaterBody::GetAudioIntensityAtSplineInputKey(float InKey) const
 {
-	return WaterSplineMetadata->AudioIntensity.Eval(InKey, 0.f);
+	return WaterSplineMetadata ? WaterSplineMetadata->AudioIntensity.Eval(InKey, 0.f) : 0.0f;
 }
 
 TArray<AWaterBodyIsland*> AWaterBody::GetIslands() const
@@ -714,7 +719,7 @@ void AWaterBody::UpdateMaterialInstances()
 bool AWaterBody::UpdateWaterHeight()
 {
 	bool bWaterBodyChanged = false;
-	if (IsFlatSurface())
+	if (IsFlatSurface() && SplineComp)
 	{
 		const int32 NumSplinePoints = SplineComp->GetNumberOfSplinePoints();
 
@@ -1198,12 +1203,18 @@ void AWaterBody::UpdateSplineComponent()
 
 void AWaterBody::OnWaterBodyChanged(bool bShapeOrPositionChanged, bool bWeightmapSettingsChanged)
 {
-	UpdateAll(bShapeOrPositionChanged);
-
-	// Some of the spline parameters need to be transferred to the underwater post process MID, if any : 
-	if (bShapeOrPositionChanged)
+	// It's possible to get called without a water spline after the Redo of a water body deletion (i.e. the water body actor gets deleted again, hence its SplineComp is restored to nullptr)
+	//  This is a very-edgy case that needs to be checked everywhere that UpdateAll might hook into so it's simpler to just skip it all. The actor is in limbo by then anyway (it only survives because
+	//  of the editor transaction) :
+	if (GetWaterSpline())
 	{
-		SetDynamicParametersOnUnderwaterPostProcessMID(UnderwaterPostProcessMID);
+		UpdateAll(bShapeOrPositionChanged);
+
+		// Some of the spline parameters need to be transferred to the underwater post process MID, if any : 
+		if (bShapeOrPositionChanged)
+		{
+			SetDynamicParametersOnUnderwaterPostProcessMID(UnderwaterPostProcessMID);
+		}
 	}
 
 #if WITH_EDITOR
diff --git a/Engine/Plugins/Experimental/Water/Source/Runtime/Private/WaterBodyIslandActor.cpp b/Engine/Plugins/Experimental/Water/Source/Runtime/Private/WaterBodyIslandActor.cpp
--- a/Engine/Plugins/Experimental/Water/Source/Runtime/Private/WaterBodyIslandActor.cpp
+++ b/Engine/Plugins/Experimental/Water/Source/Runtime/Private/WaterBodyIslandActor.cpp
@@ -68,16 +68,19 @@ void AWaterBodyIsland::GetBrushRenderDependencies(TSet<UObject*>& OutDependencie
 
 void AWaterBodyIsland::UpdateHeight()
 {
-	const int32 NumSplinePoints = SplineComp->GetNumberOfSplinePoints();
+	if (SplineComp)
+	{
+		const int32 NumSplinePoints = SplineComp->GetNumberOfSplinePoints();
 
-	const float ActorZ = GetActorLocation().Z;
+		const float ActorZ = GetActorLocation().Z;
 
-	for (int32 PointIndex = 0; PointIndex < NumSplinePoints; ++PointIndex)
-	{
-		FVector WorldLoc = SplineComp->GetLocationAtSplinePoint(PointIndex, ESplineCoordinateSpace::World);
+		for (int32 PointIndex = 0; PointIndex < NumSplinePoints; ++PointIndex)
+		{
+			FVector WorldLoc = SplineComp->GetLocationAtSplinePoint(PointIndex, ESplineCoordinateSpace::World);
 
-		WorldLoc.Z = ActorZ;
-		SplineComp->SetLocationAtSplinePoint(PointIndex, WorldLoc, ESplineCoordinateSpace::World);
+			WorldLoc.Z = ActorZ;
+			SplineComp->SetLocationAtSplinePoint(PointIndex, WorldLoc, ESplineCoordinateSpace::World);
+		}
 	}
 }
 
@@ -161,10 +164,13 @@ void AWaterBodyIsland::UpdateOverlappingWaterBodies()
 {
 	TArray<FOverlapResult> Overlaps;
 
-	FCollisionShape OverlapShape;
-	// Expand shape in Z to ensure we get overlaps for islands slighty above or below water level
-	OverlapShape.SetBox(SplineComp->Bounds.BoxExtent+FVector(0,0,10000));
-	GetWorld()->OverlapMultiByObjectType(Overlaps, SplineComp->Bounds.Origin, FQuat::Identity, FCollisionObjectQueryParams::AllObjects, OverlapShape);
+	if (SplineComp)
+	{
+		FCollisionShape OverlapShape;
+		// Expand shape in Z to ensure we get overlaps for islands slighty above or below water level
+		OverlapShape.SetBox(SplineComp->Bounds.BoxExtent+FVector(0,0,10000));
+		GetWorld()->OverlapMultiByObjectType(Overlaps, SplineComp->Bounds.Origin, FQuat::Identity, FCollisionObjectQueryParams::AllObjects, OverlapShape);
+	}
 
 	// Find any new overlapping bodies and notify them that this island influences them
 	TSet<AWaterBody*> ExistingOverlappingBodies;
diff --git a/Engine/Plugins/Experimental/Water/Source/Runtime/Private/WaterBodyLakeActor.cpp b/Engine/Plugins/Experimental/Water/Source/Runtime/Private/WaterBodyLakeActor.cpp
--- a/Engine/Plugins/Experimental/Water/Source/Runtime/Private/WaterBodyLakeActor.cpp
+++ b/Engine/Plugins/Experimental/Water/Source/Runtime/Private/WaterBodyLakeActor.cpp
@@ -110,52 +110,53 @@ void ULakeGenerator::OnUpdateBody(bool bWithExclusionVolumes)
 		LakeCollision = nullptr;
 	}
 
-	UWaterSplineComponent* WaterSpline = OwnerBody->GetWaterSpline();
-
-	UStaticMesh* WaterMesh = OwnerBody->GetWaterMeshOverride() ? OwnerBody->GetWaterMeshOverride() : UWaterSubsystem::StaticClass()->GetDefaultObject<UWaterSubsystem>()->DefaultLakeMesh;
+	if (UWaterSplineComponent* WaterSpline = OwnerBody->GetWaterSpline())
+	{
+		UStaticMesh* WaterMesh = OwnerBody->GetWaterMeshOverride() ? OwnerBody->GetWaterMeshOverride() : UWaterSubsystem::StaticClass()->GetDefaultObject<UWaterSubsystem>()->DefaultLakeMesh;
 
-	const FVector SplineExtent = WaterSpline->Bounds.BoxExtent;
+		const FVector SplineExtent = WaterSpline->Bounds.BoxExtent;
 
-	FVector WorldLoc(WaterSpline->Bounds.Origin);
-	WorldLoc.Z = OwnerBody->GetActorLocation().Z;
+		FVector WorldLoc(WaterSpline->Bounds.Origin);
+		WorldLoc.Z = OwnerBody->GetActorLocation().Z;
 
-	if (WaterMesh)
-	{
-		FTransform MeshCompToWorld = WaterSpline->GetComponentToWorld();
-		// Scale the water mesh so that it is the size of the bounds
-		FVector MeshExtent = WaterMesh->GetBounds().BoxExtent;
-		MeshExtent.Z = 1.0f;
-
-		const FVector LocalSplineExtent = WaterSpline->Bounds.TransformBy(MeshCompToWorld.Inverse()).BoxExtent;
-
-		const FVector ScaleRatio = SplineExtent / MeshExtent;
-		LakeMeshComp->SetWorldScale3D(FVector(ScaleRatio.X, ScaleRatio.Y, 1));
-		LakeMeshComp->SetWorldLocation(WorldLoc);
-		LakeMeshComp->SetWorldRotation(FQuat::Identity);
-		LakeMeshComp->SetAbsolute(false, false, true);
-		LakeMeshComp->SetStaticMesh(WaterMesh);
-		LakeMeshComp->SetMaterial(0, OwnerBody->GetWaterMaterial());
-		LakeMeshComp->SetCastShadow(false);
-		LakeMeshComp->SetCollisionEnabled(ECollisionEnabled::NoCollision);
-	}
+		if (WaterMesh)
+		{
+			FTransform MeshCompToWorld = WaterSpline->GetComponentToWorld();
+			// Scale the water mesh so that it is the size of the bounds
+			FVector MeshExtent = WaterMesh->GetBounds().BoxExtent;
+			MeshExtent.Z = 1.0f;
+
+			const FVector LocalSplineExtent = WaterSpline->Bounds.TransformBy(MeshCompToWorld.Inverse()).BoxExtent;
+
+			const FVector ScaleRatio = SplineExtent / MeshExtent;
+			LakeMeshComp->SetWorldScale3D(FVector(ScaleRatio.X, ScaleRatio.Y, 1));
+			LakeMeshComp->SetWorldLocation(WorldLoc);
+			LakeMeshComp->SetWorldRotation(FQuat::Identity);
+			LakeMeshComp->SetAbsolute(false, false, true);
+			LakeMeshComp->SetStaticMesh(WaterMesh);
+			LakeMeshComp->SetMaterial(0, OwnerBody->GetWaterMaterial());
+			LakeMeshComp->SetCastShadow(false);
+			LakeMeshComp->SetCollisionEnabled(ECollisionEnabled::NoCollision);
+		}
 
-	LakeMeshComp->SetMobility(OwnerBody->GetRootComponent()->Mobility);
+		LakeMeshComp->SetMobility(OwnerBody->GetRootComponent()->Mobility);
 
-	if (LakeCollision)
-	{
-		check(OwnerBody->bGenerateCollisions);
-		LakeCollision->bFillCollisionUnderneathForNavmesh = OwnerBody->bFillCollisionUnderWaterBodiesForNavmesh;
-		LakeCollision->SetMobility(OwnerBody->GetRootComponent()->Mobility);
-		LakeCollision->SetCollisionProfileName(OwnerBody->GetCollisionProfileName());
-		LakeCollision->SetGenerateOverlapEvents(true);
-
-		const float Depth = OwnerBody->GetChannelDepth() / 2;
-		FVector Scale = OwnerBody->GetActorScale();
-		// Avoid dividing by 0 but keep the scale's sign :
-		Scale = Scale.GetSignVector() * FVector::Max(Scale.GetAbs(), FVector::OneVector * 0.001f);
-		FVector LakeCollisionExtent = FVector(SplineExtent.X, SplineExtent.Y, Depth) / Scale;
-		LakeCollision->SetWorldLocation(WorldLoc + FVector(0, 0, -Depth));
-		LakeCollision->UpdateCollision(LakeCollisionExtent, true);
+		if (LakeCollision)
+		{
+			check(OwnerBody->bGenerateCollisions);
+			LakeCollision->bFillCollisionUnderneathForNavmesh = OwnerBody->bFillCollisionUnderWaterBodiesForNavmesh;
+			LakeCollision->SetMobility(OwnerBody->GetRootComponent()->Mobility);
+			LakeCollision->SetCollisionProfileName(OwnerBody->GetCollisionProfileName());
+			LakeCollision->SetGenerateOverlapEvents(true);
+
+			const float Depth = OwnerBody->GetChannelDepth() / 2;
+			FVector Scale = OwnerBody->GetActorScale();
+			// Avoid dividing by 0 but keep the scale's sign :
+			Scale = Scale.GetSignVector() * FVector::Max(Scale.GetAbs(), FVector::OneVector * 0.001f);
+			FVector LakeCollisionExtent = FVector(SplineExtent.X, SplineExtent.Y, Depth) / Scale;
+			LakeCollision->SetWorldLocation(WorldLoc + FVector(0, 0, -Depth));
+			LakeCollision->UpdateCollision(LakeCollisionExtent, true);
+		}
 	}
 }
 
