* Fixed crash when a BP with a nested water body is added to a map : in that case, PostRegisterAllComponents can be called multiple times in a row (e.g. the parent BP instance will register first, with all its child components, which will trigger registration of the child water body actor, and then the water body actor will also get registered independently as a "standard" actor)

diff --git a/Engine/Plugins/Experimental/Water/Source/Runtime/Private/WaterBodyActor.cpp b/Engine/Plugins/Experimental/Water/Source/Runtime/Private/WaterBodyActor.cpp
--- a/Engine/Plugins/Experimental/Water/Source/Runtime/Private/WaterBodyActor.cpp
+++ b/Engine/Plugins/Experimental/Water/Source/Runtime/Private/WaterBodyActor.cpp
@@ -1425,9 +1425,11 @@ void AWaterBody::PostRegisterAllComponents()
 	}
 #endif // WITH_EDITOR
 
-	check(WaterBodyIndex == INDEX_NONE); // Must not register if already registered 
+	// We must check for WaterBodyIndex to see if we have already been registered because PostRegisterAllComponents can be called multiple times in a row (e.g. if the actor is a child 
+	//  actor of another BP, the parent BP instance will register first, with all its child components, which will trigger registration of the child water body actor, and then 
+	//  the water body actor will also get registered independently as a "standard" actor) :
 	FWaterBodyManager* Manager = UWaterSubsystem::GetWaterBodyManager(GetWorld());
-	if (Manager && !IsTemplate())
+	if (Manager && !IsTemplate() && (WaterBodyIndex == INDEX_NONE))
 	{
 		WaterBodyIndex = Manager->AddWaterBody(this);
 	}
@@ -1441,8 +1443,9 @@ void AWaterBody::PostUnregisterAllComponents()
 {
 	Super::PostUnregisterAllComponents();
 
+	// We must check for WaterBodyIndex because PostUnregisterAllComponents can be called multiple times in a row by PostEditChangeProperty, etc.
 	FWaterBodyManager* Manager = UWaterSubsystem::GetWaterBodyManager(GetWorld());
-	if (Manager && !IsTemplate() && WaterBodyIndex != INDEX_NONE)
+	if (Manager && !IsTemplate() && (WaterBodyIndex != INDEX_NONE))
 	{
 		Manager->RemoveWaterBody(this);
 	}
