Integrate 16197883 (+ build fix 16197886) :
* Removed dependency of WaterBodyData on MPC_Landscape and removed LandscapeWaterInfo : all MPC values to control water rendering should now be controlled by the WaterBrushManager
* All maps with a water velocity height texture need to be converted (update water brush manager) and a MapCheck displays a warning to do so with an hyperlink to trigger the action

#tests editor, PIE

diff --git a/Engine/Plugins/Experimental/Water/Source/Editor/Private/WaterBrushManager.cpp b/Engine/Plugins/Experimental/Water/Source/Editor/Private/WaterBrushManager.cpp
--- a/Engine/Plugins/Experimental/Water/Source/Editor/Private/WaterBrushManager.cpp
+++ b/Engine/Plugins/Experimental/Water/Source/Editor/Private/WaterBrushManager.cpp
@@ -28,6 +28,12 @@
 #include "Curves/CurveFloat.h"
 #include "EngineUtils.h"
 #include "Landscape.h"
+#include "UObject/FortniteMainBranchObjectVersion.h"
+#include "Misc/UObjectToken.h"
+#include "Misc/MapErrors.h"
+#include "Logging/MessageLog.h"
+
+#define LOCTEXT_NAMESPACE "WaterBrushManager"
 
 AWaterBrushManager::AWaterBrushManager(const FObjectInitializer& ObjectInitializer)
 	: Super()
@@ -59,6 +65,7 @@ void AWaterBrushManager::Serialize(FArchive& Ar)
 	Super::Serialize(Ar);
 
 	Ar.UsingCustomVersion(FWaterCustomVersion::GUID);
+	Ar.UsingCustomVersion(FFortniteMainBranchObjectVersion::GUID);
 }
 
 void AWaterBrushManager::PostLoad()
@@ -133,6 +140,18 @@ void AWaterBrushManager::PostLoad()
 		});
 
 	}
+
+	if (!IsTemplate() 
+		&& (bNeedsForceUpdate || (GetLinkerCustomVersion(FFortniteMainBranchObjectVersion::GUID) < FFortniteMainBranchObjectVersion::RemoveLandscapeWaterInfo)))
+	{
+		// The removal of LandscapeWaterInfo is accompanied by a change in how the water velocity height texture is encoded so we need to regenerate it :
+		bNeedsForceUpdate = true;
+
+		ShowForceUpdateMapCheckError();
+
+		// Show MapCheck window
+		FMessageLog("MapCheck").Open(EMessageSeverity::Warning);
+	}
 #endif // WITH_EDITOR
 }
 
@@ -316,11 +335,19 @@ void AWaterBrushManager::BlueprintOnRenderTargetTexturesUpdated_Native(UTexture2
 
 void AWaterBrushManager::ForceUpdate()
 {
-	bKillCache = true;
+#if WITH_EDITOR
+	if (bNeedsForceUpdate)
+	{
+		// We need to mark our own package as dirty to force save the water brush manager and stop dispaying the ForceUpdate message
+		Modify();
+	}
+#endif // WITH_EDITOR
+
+	bKillCache = true;	
 	ClearCurveCache();
 	ALandscapeBlueprintBrushBase::RequestLandscapeUpdate();
 	// Regenerate the water depth velocity RT : 
-	MarkRenderTargetsDirty();
+	MarkRenderTargetsDirty();	
 }
 
 void AWaterBrushManager::SingleJumpStep()
@@ -657,12 +684,7 @@ void AWaterBrushManager::UpdateCurveCacheKeys()
 			if (ensure(ElevationCurveAsset != nullptr))
 			{
 				FWaterBodyBrushCache* WaterBrushCache = BrushCurveRTCache.Find(ElevationCurveAsset);
-				if (WaterBrushCache != nullptr)
-				{
-					// TODO [jonathan.bard] : there are some repros where this can happen : fix this : 
-					check(WaterBrushCache->CacheRenderTarget != nullptr);
-				}
-				else
+				if ((WaterBrushCache == nullptr) || (WaterBrushCache->CacheRenderTarget == nullptr))
 				{
 					UTextureRenderTarget2D* CurveRT = FWaterUtils::GetOrCreateTransientRenderTarget2D(nullptr, TEXT("CurveRT"), FIntPoint(256, 1), ETextureRenderTargetFormat::RTF_R16f);
 					BrushCurveRTCache.Add(ElevationCurveAsset, FWaterBodyBrushCache{ CurveRT, false });
@@ -1075,6 +1097,24 @@ bool AWaterBrushManager::DeprecateWaterLandscapeInfo(FVector& OutRTWorldLocation
 #endif // WITH_EDITOR
 }
 
+#if WITH_EDITOR
+
+void AWaterBrushManager::ShowForceUpdateMapCheckError()
+{
+	FFormatNamedArguments Arguments;
+	Arguments.Add(TEXT("WaterBrush"), FText::FromString(GetName()));
+	Arguments.Add(TEXT("Outer"), FText::FromString(GetPackage()->GetPathName()));
+	
+	FMessageLog("MapCheck").Warning()
+		->AddToken(FUObjectToken::Create(this))
+		->AddToken(FTextToken::Create(FText::Format(LOCTEXT("AWaterBrushManager_ShowForceUpdateMapCheckError_Message_ForceUpdateNeeded", "Water brush {WaterBrush} in package {Outer} is out of date and needs updating for water to render properly."), Arguments)))
+		->AddToken(FMapErrorToken::Create(TEXT("AWaterBrushManager_ShowForceUpdateMapCheckError_MapError_ForceUpdateNeeded")))
+		->AddToken(FActionToken::Create(LOCTEXT("AWaterBrushManager_ShowForceUpdateMapCheckError_ActionName_ForceUpdateNeeded", "Update water brush"), FText(),
+			FOnActionTokenExecuted::CreateUObject(this, &AWaterBrushManager::ForceUpdate), true));
+}
+
+#endif // WITH_EDITOR
+
 void AWaterBrushManager::SetMPCParams()
 {
 	UWorld* World = GetWorld();
@@ -1390,6 +1430,8 @@ UTextureRenderTarget2D* AWaterBrushManager::Render_Native(bool InIsHeightmap, UT
 
 	bKillCache = false;
 
+	bNeedsForceUpdate = false;
+
 	return ReturnRT;
 }
 
@@ -1417,4 +1459,20 @@ void AWaterBrushManager::UpdateBrushCacheKeys()
 	}
 }
 
+
+#if WITH_EDITOR
+
+void AWaterBrushManager::CheckForErrors()
+{
+	Super::CheckForErrors();
+
+	// If a force update action was requested but hasn't been performed yet, display the message again : 
+	if (bNeedsForceUpdate)
+	{
+		ShowForceUpdateMapCheckError();
+	}
+}
+
+#endif // WITH_EDITOR
+
 #undef LOCTEXT_NAMESPACE
diff --git a/Engine/Plugins/Experimental/Water/Source/Editor/Private/WaterBrushManager.h b/Engine/Plugins/Experimental/Water/Source/Editor/Private/WaterBrushManager.h
--- a/Engine/Plugins/Experimental/Water/Source/Editor/Private/WaterBrushManager.h
+++ b/Engine/Plugins/Experimental/Water/Source/Editor/Private/WaterBrushManager.h
@@ -188,7 +188,10 @@ public:
 		
 	UPROPERTY(EditAnywhere, BlueprintReadWrite, meta=(Category="Debug"))
 	bool DisableBrushTextureEffects = false;
-			
+
+	UPROPERTY()
+	bool bNeedsForceUpdate = false;
+
 	AWaterBrushManager(const FObjectInitializer& ObjectInitializer = FObjectInitializer::Get());
 
 	virtual void Serialize(FArchive& Ar) override;
@@ -288,6 +291,10 @@ private:
 	virtual void ApplyWeightmapSettings(const FBrushRenderContext& BrushRenderContext, const FBrushActorRenderContext& BrushActorRenderContext, const FWaterBodyWeightmapSettings& WMSettings);
 	virtual void ApplyToCompositeWaterBodyTexture(FBrushRenderContext& BrushRenderContext, const FBrushActorRenderContext& BrushActorRenderContext);
 
+#if WITH_EDITOR
+	virtual void CheckForErrors() override;
+#endif // WITH_EDITOR
+
 	UCurveFloat* GetElevationCurveAsset(const FWaterCurveSettings& CurveSettings);
 	void ClearCurveCache();
 	void OnCurveUpdated(UCurveBase* Curve, EPropertyChangeType::Type ChangeType);
@@ -296,6 +303,10 @@ private:
 	// HACK [jonathan.bard] : this is only needed for data deprecation, when LandscapeTransform and LandscapeRTRes were not serialized: 
 	bool DeprecateWaterLandscapeInfo(FVector& OutRTWorldLocation, FVector& OutRTWorldSizeVector);
 
+#if WITH_EDITOR
+	void ShowForceUpdateMapCheckError();
+#endif // WITH_EDITOR
+
 	UTextureRenderTarget2D* VelocityPingPongRead(const FBrushRenderContext& BrushRenderContext) const;
 	UTextureRenderTarget2D* VelocityPingPongWrite(const FBrushRenderContext& BrushRenderContext) const;
 
diff --git a/Engine/Plugins/Experimental/Water/Source/Editor/Private/WaterEditorSubsystem.cpp b/Engine/Plugins/Experimental/Water/Source/Editor/Private/WaterEditorSubsystem.cpp
--- a/Engine/Plugins/Experimental/Water/Source/Editor/Private/WaterEditorSubsystem.cpp
+++ b/Engine/Plugins/Experimental/Water/Source/Editor/Private/WaterEditorSubsystem.cpp
@@ -42,6 +42,7 @@ void UpdateSingleTexture(UTexture2D*& DestTexure, UTextureRenderTarget2D* SrcRen
 		break;
 	}
 
+	DestTexure->Modify();
 	DestTexure->PreEditChange(nullptr); // Ensures synchronization with TextureCompilingManager.
 	DestTexure->LODGroup = GetDefault<UWaterEditorSettings>()->TextureGroupForGeneratedTextures;
 	DestTexure->MipGenSettings = TMGS_NoMipmaps;
diff --git a/Engine/Source/Runtime/Core/Public/UObject/FortniteMainBranchObjectVersion.h b/Engine/Source/Runtime/Core/Public/UObject/FortniteMainBranchObjectVersion.h
--- a/Engine/Source/Runtime/Core/Public/UObject/FortniteMainBranchObjectVersion.h
+++ b/Engine/Source/Runtime/Core/Public/UObject/FortniteMainBranchObjectVersion.h
@@ -149,6 +149,9 @@ struct CORE_API FFortniteMainBranchObjectVersion
 		// (Merged from //UE4/Main)
 		ChaosConvexVariableStructureDataAndVerticesArray,
 
+		// Remove the WaterVelocityHeightTexture dependency on MPC_Landscape and LandscapeWaterIndo 
+		RemoveLandscapeWaterInfo,
+
 		// -----<new versions can be added above this line>-------------------------------------------------
 		VersionPlusOne,
 		LatestVersion = VersionPlusOne - 1
