* Added spline resampling functions to spline component in order to iteratively transform splines into polylines until the polyline segments are close enough to the original spline (within a distance threshold)
* Implemented more accurate transformation of lake and ocean splines into polygons:
** Lake collision component's convex hulls are generated using the new spline resampling method, leading to much more accurate collision volumes when spline points use long tangents
** Water mesh generation now relies on the lakes' collision component's convex hulls if present (in order to save some computing time) or, if not, resample the lake spline themselves. This reduces the amount of water tiles that were improperly absent from the lakes' surfaces
** For oceans, water mesh generation now always relies on spline resampling to generate more accurate water tiles
* Fixed crash with lake and ocean collision display that could access game-thread modified data on the render thread

diff --git a/Engine/Plugins/Experimental/Water/Source/Runtime/Private/LakeCollisionComponent.cpp b/Engine/Plugins/Experimental/Water/Source/Runtime/Private/LakeCollisionComponent.cpp
--- a/Engine/Plugins/Experimental/Water/Source/Runtime/Private/LakeCollisionComponent.cpp
+++ b/Engine/Plugins/Experimental/Water/Source/Runtime/Private/LakeCollisionComponent.cpp
@@ -12,6 +12,15 @@
 #include "PhysicsEngine/ConvexElem.h"
 #include "AI/NavigationSystemHelpers.h"
 #include "WaterUtils.h"
+#include "Algo/Transform.h"
+
+
+// ----------------------------------------------------------------------------------
+
+extern TAutoConsoleVariable<float> CVarWaterSplineResampleMaxDistance;
+
+
+// ----------------------------------------------------------------------------------
 
 ULakeCollisionComponent::ULakeCollisionComponent(const FObjectInitializer& ObjectInitializer)
 	: UPrimitiveComponent(ObjectInitializer)
@@ -44,9 +53,11 @@ void ULakeCollisionComponent::UpdateCollision(FVector InBoxExtent, bool bSplineP
 	}
 }
 
+#if !(UE_BUILD_SHIPPING || UE_BUILD_TEST)
+
 FPrimitiveSceneProxy* ULakeCollisionComponent::CreateSceneProxy()
 {
-	/** Represents a UBoxComponent to the scene manager. */
+	/** Represents a ULakeCollisionComponent to the scene manager. */
 	class FLakeCollisionSceneProxy final : public FPrimitiveSceneProxy
 	{
 	public:
@@ -58,9 +69,14 @@ FPrimitiveSceneProxy* ULakeCollisionComponent::CreateSceneProxy()
 
 		FLakeCollisionSceneProxy(const ULakeCollisionComponent* InComponent)
 			: FPrimitiveSceneProxy(InComponent)
-			, BodySetup(InComponent->CachedBodySetup)
 		{
 			bWillEverBeLit = false;
+
+			if (InComponent->CachedBodySetup)
+			{
+				// copy the geometry for being able to access it on the render thread : 
+				AggregateGeom = InComponent->CachedBodySetup->AggGeom;
+			}
 		}
 
 		virtual void GetDynamicMeshElements(const TArray<const FSceneView*>& Views, const FSceneViewFamily& ViewFamily, uint32 VisibilityMap, FMeshElementCollector& Collector) const override
@@ -70,26 +86,23 @@ FPrimitiveSceneProxy* ULakeCollisionComponent::CreateSceneProxy()
 
 			const bool bDrawCollision = ViewFamily.EngineShowFlags.Collision && IsCollisionEnabled();
 
-#if !(UE_BUILD_SHIPPING || UE_BUILD_TEST)
 			for (int32 ViewIndex = 0; ViewIndex < Views.Num(); ViewIndex++)
 			{
 				if (VisibilityMap & (1 << ViewIndex))
 				{
 					const FSceneView* View = Views[ViewIndex];
 
-					if (bDrawCollision && BodySetup && AllowDebugViewmodes())
+					if (bDrawCollision && AllowDebugViewmodes())
 					{
 						FColor CollisionColor(157, 149, 223, 255);
 						const bool bPerHullColor = false;
 						const bool bDrawSolid = false;
-						BodySetup->AggGeom.GetAggGeom(LocalToWorldTransform, GetSelectionColor(CollisionColor, IsSelected(), IsHovered()).ToFColor(true), nullptr, bPerHullColor, bDrawSolid, DrawsVelocity(), ViewIndex, Collector);
+						AggregateGeom.GetAggGeom(LocalToWorldTransform, GetSelectionColor(CollisionColor, IsSelected(), IsHovered()).ToFColor(true), nullptr, bPerHullColor, bDrawSolid, DrawsVelocity(), ViewIndex, Collector);
 					}
 	
 					RenderBounds(Collector.GetPDI(ViewIndex), View->Family->EngineShowFlags, GetBounds(), IsSelected());
-
 				}
 			}
-#endif
 		}
 
 		virtual FPrimitiveViewRelevance GetViewRelevance(const FSceneView* View) const override
@@ -108,11 +121,12 @@ FPrimitiveSceneProxy* ULakeCollisionComponent::CreateSceneProxy()
 		uint32 GetAllocatedSize(void) const { return(FPrimitiveSceneProxy::GetAllocatedSize()); }
 
 	private:
-		UBodySetup* BodySetup;
+		FKAggregateGeom AggregateGeom;
 	};
 
 	return new FLakeCollisionSceneProxy(this);
 }
+#endif // !(UE_BUILD_SHIPPING || UE_BUILD_TEST)
 
 FBoxSphereBounds ULakeCollisionComponent::CalcBounds(const FTransform& LocalToWorld) const
 {
@@ -340,6 +354,8 @@ bool TriangulateSimpleXYPlanarPolygon(const TArray<FVector>& VertexPositions, TA
 
 void ULakeCollisionComponent::UpdateBodySetup()
 {
+	TRACE_CPUPROFILER_EVENT_SCOPE(ULakeCollisionComponent::UpdateBodySetup);
+
 	CreateLakeBodySetupIfNeeded();
 	
 	FGuid PreviousBodySetupGuid = CachedBodySetup->BodySetupGuid;
@@ -363,46 +379,19 @@ void ULakeCollisionComponent::UpdateBodySetup()
 		const float MinZ = BoxExtent.Z;
 
 		// Generate planes
-		const int32 NumIterationsPerSegment = 3;
 		const int32 NumPoints = SplineComp->GetNumberOfSplinePoints();
 		// lakes are closed loops so add 1 to the end 
 		const int32 NumSteps = NumPoints + 1;
-		const double FinalStopDist = SplineComp->GetDistanceAlongSplineAtSplinePoint(NumPoints-1);
-
-		SplineVerts.Reserve(NumSteps * NumIterationsPerSegment);
 
-		for (int32 StepIndex = 0; StepIndex < NumSteps; ++StepIndex)
 		{
-			const int32 StartIndex = StepIndex;
-			const int32 StopIndex = StartIndex + 1;
-
-			const double StartDist = SplineComp->GetDistanceAlongSplineAtSplinePoint(StartIndex);
-			const double StopDist = SplineComp->GetDistanceAlongSplineAtSplinePoint(StopIndex);
-
-			double SubstepSize = FMath::Abs(StopDist - StartDist) / NumIterationsPerSegment;
-			if (SubstepSize == 0.0)
-			{
-				// Make sure there's at least 1 sub-step so that we get a single spline vertex for constant interpolation : 
-				SubstepSize = StopDist + 1.0;
-			}
-
-			int32 SubstepIndex = 0;
-			double SubstepDist = StartDist;
-			while (SubstepDist <= StopDist)
-			{
-				// Get world location since local space would be relative to the spline comp
-				FVector Location = SplineComp->GetLocationAtDistanceAlongSpline(SubstepDist, ESplineCoordinateSpace::World);
+			TRACE_CPUPROFILER_EVENT_SCOPE(ResampleSpline);
 
-				// Transform to local space of this component
-				FVector LocalSpace = GetComponentToWorld().InverseTransformPosition(Location);
-
-				SplineVerts.Add(FVector2D(LocalSpace));
-
-				SubstepDist += SubstepSize;
-			}
+			TArray<FVector> PolyLineVertices;
+			SplineComp->ConvertSplineToPolyLine(ESplineCoordinateSpace::World, FMath::Square(CVarWaterSplineResampleMaxDistance.GetValueOnGameThread()), PolyLineVertices);
+			// Transform to local space of this component :
+			Algo::Transform(PolyLineVertices, SplineVerts, [this](const FVector& Vertex) { return FVector2D(GetComponentToWorld().InverseTransformPosition(Vertex)); });
 		}
 
-
 		TArray<FVector2D> CorrectedSplineVertices;
 		FGeomTools2D::CorrectPolygonWinding(CorrectedSplineVertices, SplineVerts, false);
 
@@ -415,27 +404,30 @@ void ULakeCollisionComponent::UpdateBodySetup()
 		TArray<TArray<FVector2D>> ConvexHulls;
 		FGeomTools2D::GenerateConvexPolygonsFromTriangles(ConvexHulls, OutCleanTris);
 
-		for (const auto& Hull : ConvexHulls)
 		{
-			TArray<FVector> Hull3DVerts;
-			Hull3DVerts.Reserve(Hull.Num());
-
-			for (int32 PointIdx = 0; PointIdx < Hull.Num(); ++PointIdx)
+			TRACE_CPUPROFILER_EVENT_SCOPE(GenerateHull);
+			for (const auto& Hull : ConvexHulls)
 			{
-				Hull3DVerts.Emplace(FVector(Hull[PointIdx], 0));
-			}
+				TArray<FVector> Hull3DVerts;
+				Hull3DVerts.Reserve(Hull.Num());
 
-			if(Hull3DVerts.Num() > 2)
-			{
-				TArray<FVector> ExtrudedVerts;
-				TArray<FIntVector> Indices;
-				ExtrudeZSimplePolygon(Hull3DVerts, MinZ, MaxZ, ExtrudedVerts, Indices);
+				for (int32 PointIdx = 0; PointIdx < Hull.Num(); ++PointIdx)
+				{
+					Hull3DVerts.Emplace(FVector(Hull[PointIdx], 0));
+				}
+
+				if(Hull3DVerts.Num() > 2)
+				{
+					TArray<FVector> ExtrudedVerts;
+					TArray<FIntVector> Indices;
+					ExtrudeZSimplePolygon(Hull3DVerts, MinZ, MaxZ, ExtrudedVerts, Indices);
 
-				FKConvexElem Convex;
-				Convex.VertexData = MoveTemp(ExtrudedVerts);
-				Convex.UpdateElemBox();
+					FKConvexElem Convex;
+					Convex.VertexData = MoveTemp(ExtrudedVerts);
+					Convex.UpdateElemBox();
 
-				CachedBodySetup->AggGeom.ConvexElems.Add(Convex);
+					CachedBodySetup->AggGeom.ConvexElems.Add(Convex);
+				}
 			}
 		}
 
diff --git a/Engine/Plugins/Experimental/Water/Source/Runtime/Private/OceanCollisionComponent.cpp b/Engine/Plugins/Experimental/Water/Source/Runtime/Private/OceanCollisionComponent.cpp
--- a/Engine/Plugins/Experimental/Water/Source/Runtime/Private/OceanCollisionComponent.cpp
+++ b/Engine/Plugins/Experimental/Water/Source/Runtime/Private/OceanCollisionComponent.cpp
@@ -115,19 +115,12 @@ void UOceanCollisionComponent::UpdateBodySetup(const TArray<FKConvexElem>& Conve
 	UpdateBounds();
 }
 
-
 UBodySetup* UOceanCollisionComponent::GetBodySetup()
 {
 	return CachedBodySetup;
 }
 
-
-
-
-
-
-
-
+#if !(UE_BUILD_SHIPPING || UE_BUILD_TEST)
 
 FPrimitiveSceneProxy* UOceanCollisionComponent::CreateSceneProxy()
 {
@@ -143,9 +136,14 @@ FPrimitiveSceneProxy* UOceanCollisionComponent::CreateSceneProxy()
 
 		FOceanCollisionSceneProxy(const UOceanCollisionComponent* InComponent)
 			: FPrimitiveSceneProxy(InComponent)
-			, BodySetup(InComponent->CachedBodySetup)
 		{
 			bWillEverBeLit = false;
+
+			if (InComponent->CachedBodySetup)
+			{
+				// copy the geometry for being able to access it on the render thread : 
+				AggregateGeom = InComponent->CachedBodySetup->AggGeom;
+			}
 		}
 
 		virtual void GetDynamicMeshElements(const TArray<const FSceneView*>& Views, const FSceneViewFamily& ViewFamily, uint32 VisibilityMap, FMeshElementCollector& Collector) const override
@@ -155,26 +153,23 @@ FPrimitiveSceneProxy* UOceanCollisionComponent::CreateSceneProxy()
 
 			const bool bDrawCollision = ViewFamily.EngineShowFlags.Collision && IsCollisionEnabled();
 
-#if !(UE_BUILD_SHIPPING || UE_BUILD_TEST)
 			for (int32 ViewIndex = 0; ViewIndex < Views.Num(); ViewIndex++)
 			{
 				if (VisibilityMap & (1 << ViewIndex))
 				{
 					const FSceneView* View = Views[ViewIndex];
 
-					if (bDrawCollision && BodySetup && AllowDebugViewmodes())
+					if (bDrawCollision && AllowDebugViewmodes())
 					{
 						FColor CollisionColor(157, 149, 223, 255);
 						const bool bPerHullColor = false;
 						const bool bDrawSolid = false;
-						BodySetup->AggGeom.GetAggGeom(LocalToWorldTransform, GetSelectionColor(CollisionColor, IsSelected(), IsHovered()).ToFColor(true), nullptr, bPerHullColor, bDrawSolid, DrawsVelocity(), ViewIndex, Collector);
+						AggregateGeom.GetAggGeom(LocalToWorldTransform, GetSelectionColor(CollisionColor, IsSelected(), IsHovered()).ToFColor(true), nullptr, bPerHullColor, bDrawSolid, DrawsVelocity(), ViewIndex, Collector);
 					}
 
 					RenderBounds(Collector.GetPDI(ViewIndex), View->Family->EngineShowFlags, GetBounds(), IsSelected());
-
 				}
 			}
-#endif
 		}
 
 		virtual FPrimitiveViewRelevance GetViewRelevance(const FSceneView* View) const override
@@ -193,12 +188,14 @@ FPrimitiveSceneProxy* UOceanCollisionComponent::CreateSceneProxy()
 		uint32 GetAllocatedSize(void) const { return(FPrimitiveSceneProxy::GetAllocatedSize()); }
 
 	private:
-		UBodySetup* BodySetup;
+		FKAggregateGeom AggregateGeom;
 	};
 
 	return new FOceanCollisionSceneProxy(this);
 }
 
+#endif // !(UE_BUILD_SHIPPING || UE_BUILD_TEST)
+
 
 
 
diff --git a/Engine/Plugins/Experimental/Water/Source/Runtime/Private/WaterBodyActor.cpp b/Engine/Plugins/Experimental/Water/Source/Runtime/Private/WaterBodyActor.cpp
--- a/Engine/Plugins/Experimental/Water/Source/Runtime/Private/WaterBodyActor.cpp
+++ b/Engine/Plugins/Experimental/Water/Source/Runtime/Private/WaterBodyActor.cpp
@@ -56,6 +56,13 @@ TAutoConsoleVariable<float> CVarWaterOceanFallbackDepth(
 	TEXT("Depth to report for the ocean when no terrain is found under the query location. Not used when <= 0."),
 	ECVF_Default);
 
+TAutoConsoleVariable<float> CVarWaterSplineResampleMaxDistance(
+	TEXT("r.Water.WaterSplineResampleMaxDistance"),
+	50.0f,
+	TEXT("Maximum distance between the sample segments and the spline when converting the water spline shapes to polygons (as distance discreases, the number of vertices will increase, \
+		, the physics shapes will be more accurate, the water tiles will match more closely, but the computational cost will also increase)."),
+	ECVF_Default);
+
 const FName AWaterBody::WaterBodyIndexParamName(TEXT("WaterBodyIndex"));
 const FName AWaterBody::WaterVelocityAndHeightName(TEXT("WaterVelocityAndHeight"));
 const FName AWaterBody::GlobalOceanHeightName(TEXT("GlobalOceanHeight"));
diff --git a/Engine/Plugins/Experimental/Water/Source/Runtime/Private/WaterMeshComponent.cpp b/Engine/Plugins/Experimental/Water/Source/Runtime/Private/WaterMeshComponent.cpp
--- a/Engine/Plugins/Experimental/Water/Source/Runtime/Private/WaterMeshComponent.cpp
+++ b/Engine/Plugins/Experimental/Water/Source/Runtime/Private/WaterMeshComponent.cpp
@@ -15,6 +15,7 @@
 #include "WaterSubsystem.h"
 #include "WaterModule.h"
 #include "Math/NumericLimits.h"
+#include "Algo/Transform.h"
 
 /** Scalability CVars*/
 static TAutoConsoleVariable<int32> CVarWaterMeshLODCountBias(
@@ -54,6 +55,9 @@ static TAutoConsoleVariable<int32> CVarWaterMeshEnabled(
 	ECVF_RenderThreadSafe
 );
 
+extern TAutoConsoleVariable<float> CVarWaterSplineResampleMaxDistance;
+
+
 // ----------------------------------------------------------------------------------
 
 bool IsWaterMeshEnabled(bool bIsRenderThread)
@@ -331,42 +335,84 @@ void UWaterMeshComponent::RebuildWaterMesh(float InTileSize, const FIntPoint& In
 			TRACE_CPUPROFILER_EVENT_SCOPE(Lake);
 
 			const UWaterSplineComponent* SplineComp = WaterBody->GetWaterSpline();
-			const int32 NumOriginaSplinePoints = SplineComp->GetNumberOfSplinePoints();
-			
+			const int32 NumOriginalSplinePoints = SplineComp->GetNumberOfSplinePoints();
+			float ConstantZ = SplineComp->GetLocationAtDistanceAlongSpline(0.0f, ESplineCoordinateSpace::World).Z;
+
+			FBox LakeBounds = WaterBody->GetComponentsBoundingBox();
+			LakeBounds.Max.Z += WaterBody->GetMaxWaveHeight();
+
 			// Skip lakes with less than 3 spline points
-			if (NumOriginaSplinePoints < 3)
+			if (NumOriginalSplinePoints < 3)
 			{
 				break;
 			}
 
-			const float SplineLength = SplineComp->GetSplineLength();
-			// LeafSize * 1.5 is roughly the max distance along the spline we aim to sample at
-			const int32 NumSamplePoints = FMath::Max(FMath::FloorToInt(SplineLength / (WaterQuadTree.GetLeafSize() * 1.5f)), NumOriginaSplinePoints);
-			const float DistDelta = SplineLength / NumSamplePoints;
+			TArray<TArray<FVector2D>> PolygonBatches;
+			// Reuse the convex hulls generated for the physics shape because the work has already been done, but we can fallback to a simple spline evaluation method in case there's no physics:
+			bool bUseFallbackMethod = true;
 
-			TArray<FVector2D> SplinePoints;
+			TArray<UPrimitiveComponent*> CollisionComponents = WaterBody->GetCollisionComponents();
+			if (CollisionComponents.Num() > 0)
+			{
+				UPrimitiveComponent* Comp = CollisionComponents[0];
+				if (UBodySetup* BodySetup = (Comp != nullptr) ? Comp->GetBodySetup() : nullptr)
+				{
+					FTransform CompTransform = Comp->GetComponentTransform();
+					// Go through all sub shapes on the bodysetup to get a tight fit along water body
+					for (const FKConvexElem& ConvElem : BodySetup->AggGeom.ConvexElems)
+					{
+						TRACE_CPUPROFILER_EVENT_SCOPE(Add);
+
+						// Vertex data contains the bottom vertices first, then the top ones :
+						int32 TotalNumVertices = ConvElem.VertexData.Num();
+						check(TotalNumVertices % 2 == 0);
+						int32 NumVertices = TotalNumVertices / 2;
+						if (NumVertices > 0)
+						{ 
+							bUseFallbackMethod = false;
+
+							// Because the physics shape is made of multiple convex hulls, we cannot simply add their vertices to one big list but have to have 1 batch of polygons per convex hull
+							TArray<FVector2D>& Polygon = PolygonBatches.Emplace_GetRef();
+							Polygon.SetNum(NumVertices);
+							for (int32 i = 0; i < NumVertices; ++i)
+							{
+								// Gather the top vertices :
+								Polygon[i] = FVector2D(CompTransform.TransformPosition(ConvElem.VertexData[NumVertices + i]));
+							}
+						}
+					}
+				}
+			}
 
-			// Allocate the number of spline points we'll be using and fill it from the actual spline
-			SplinePoints.SetNum(NumSamplePoints);
-			for (int32 i = 0; i < NumSamplePoints; i++)
+			if (bUseFallbackMethod)
 			{
-				SplinePoints[i] = FVector2D(SplineComp->GetLocationAtDistanceAlongSpline(i*DistDelta, ESplineCoordinateSpace::World));
-				
+				TArray<FVector> PolyLineVertices;
+				SplineComp->ConvertSplineToPolyLine(ESplineCoordinateSpace::World, FMath::Square(CVarWaterSplineResampleMaxDistance.GetValueOnGameThread()), PolyLineVertices);
+
+				TArray<FVector2D>& Polygon = PolygonBatches.Emplace_GetRef();
+				Polygon.Reserve(PolyLineVertices.Num());
+				Algo::Transform(PolyLineVertices, Polygon, [](const FVector& Vertex) { return FVector2D(Vertex); });
+			}
+
+			for (const TArray<FVector2D>& Polygon : PolygonBatches)
+			{
+				WaterQuadTree.AddLake(Polygon, LakeBounds, WaterBodyRenderDataIndex);
+
 #if !(UE_BUILD_SHIPPING || UE_BUILD_TEST)
 				if (!!CVarWaterMeshShowTileGenerationGeometry.GetValueOnGameThread())
 				{
-					FVector Point0 = SplineComp->GetLocationAtDistanceAlongSpline(i*DistDelta, ESplineCoordinateSpace::World);
-					FVector Point1 = SplineComp->GetLocationAtDistanceAlongSpline(((i + 1) % NumSamplePoints)*DistDelta, ESplineCoordinateSpace::World);
-					DrawDebugLine(GetWorld(), Point0, Point1, FColor::Green);
+					float Z = SplineComp->GetLocationAtDistanceAlongSpline(0.0f, ESplineCoordinateSpace::World).Z;
+					int32 NumVertices = Polygon.Num();
+					for (int32 i = 0; i < NumVertices; i++)
+					{
+						const FVector2D& Point0 = Polygon[i];
+						const FVector2D& Point1 = Polygon[(i + 1) % NumVertices];
+						DrawDebugLine(GetWorld(), FVector(Point0.X, Point0.Y, Z), FVector(Point1.X, Point1.Y, Z), FColor::Green);
+					}
 				}
-#endif
+#endif // !(UE_BUILD_SHIPPING || UE_BUILD_TEST)
 			}
 
-			FBox LakeBounds = WaterBody->GetComponentsBoundingBox();
-			LakeBounds.Max.Z += WaterBody->GetMaxWaveHeight();
-
-			WaterQuadTree.AddLake(SplinePoints, LakeBounds, WaterBodyRenderDataIndex);
-
 			break;
 		}
 		case EWaterBodyType::Ocean:
@@ -383,35 +429,38 @@ void UWaterMeshComponent::RebuildWaterMesh(float InTileSize, const FIntPoint& In
 			else
 			{
 				const UWaterSplineComponent* SplineComp = WaterBody->GetWaterSpline();
-				const int32 NumSamplePoints = SplineComp->GetNumberOfSplinePoints();
+				const int32 NumOriginalSplinePoints = SplineComp->GetNumberOfSplinePoints();
 
 				// Skip oceans with less than 3 spline points
-				if (NumSamplePoints < 3)
+				if (NumOriginalSplinePoints < 3)
 				{
 					break;
 				}
 
-				TArray<FVector2D> SplinePoints;
+				TArray<FVector2D> Polygon;
 
-				// Allocate the number of spline points we'll be using and fill it from the actual spline
-				SplinePoints.SetNum(NumSamplePoints);
-				for (int32 i = 0; i < NumSamplePoints; i++)
-				{
-					SplinePoints[i] = FVector2D(SplineComp->GetLocationAtSplinePoint(i, ESplineCoordinateSpace::World));
+				TArray<FVector> PolyLineVertices;
+				SplineComp->ConvertSplineToPolyLine(ESplineCoordinateSpace::World, FMath::Square(CVarWaterSplineResampleMaxDistance.GetValueOnGameThread()), PolyLineVertices);
+
+				Polygon.Reserve(PolyLineVertices.Num());
+				Algo::Transform(PolyLineVertices, Polygon, [](const FVector& Vertex) { return FVector2D(Vertex); });
 
 #if !(UE_BUILD_SHIPPING || UE_BUILD_TEST)
-					if (!!CVarWaterMeshShowTileGenerationGeometry.GetValueOnGameThread())
+				if (!!CVarWaterMeshShowTileGenerationGeometry.GetValueOnGameThread())
+				{
+					float Z = SplineComp->GetLocationAtDistanceAlongSpline(0.0f, ESplineCoordinateSpace::World).Z;
+					int32 NumVertices = Polygon.Num();
+					for (int32 i = 0; i < NumVertices; i++)
 					{
-						FVector Point0 = SplineComp->GetLocationAtSplinePoint(i, ESplineCoordinateSpace::World);
-						FVector Point1 = SplineComp->GetLocationAtSplinePoint(((i + 1) % NumSamplePoints), ESplineCoordinateSpace::World);
-						DrawDebugLine(GetWorld(), Point0, Point1, FColor::Blue);
+						const FVector2D& Point0 = Polygon[i];
+						const FVector2D& Point1 = Polygon[(i + 1) % NumVertices];
+						DrawDebugLine(GetWorld(), FVector(Point0.X, Point0.Y, Z), FVector(Point1.X, Point1.Y, Z), FColor::Blue);
 					}
-#endif
 				}
+#endif //!(UE_BUILD_SHIPPING || UE_BUILD_TEST)
 
 				const FBox OceanBounds = WaterBody->GetComponentsBoundingBox();
-
-				WaterQuadTree.AddOcean(SplinePoints, FVector2D(OceanBounds.Min.Z, OceanBounds.Max.Z + WaterBody->GetMaxWaveHeight()), WaterBodyRenderDataIndex);
+				WaterQuadTree.AddOcean(Polygon, FVector2D(OceanBounds.Min.Z, OceanBounds.Max.Z + WaterBody->GetMaxWaveHeight()), WaterBodyRenderDataIndex);
 			}
 
 			// Place far mesh height just below the ocean level
diff --git a/Engine/Plugins/Experimental/Water/Source/Runtime/Public/LakeCollisionComponent.h b/Engine/Plugins/Experimental/Water/Source/Runtime/Public/LakeCollisionComponent.h
--- a/Engine/Plugins/Experimental/Water/Source/Runtime/Public/LakeCollisionComponent.h
+++ b/Engine/Plugins/Experimental/Water/Source/Runtime/Public/LakeCollisionComponent.h
@@ -17,10 +17,14 @@ public:
 	void UpdateCollision(FVector InBoxExtent, bool bSplinePointsChanged);
 	
 	virtual bool IsZeroExtent() const override { return BoxExtent.IsZero(); }
-	virtual FPrimitiveSceneProxy* CreateSceneProxy() override;
 	virtual FBoxSphereBounds CalcBounds(const FTransform& LocalToWorld) const;
 	virtual UBodySetup* GetBodySetup() override;
 
+#if !(UE_BUILD_SHIPPING || UE_BUILD_TEST)
+	// The scene proxy is only for debug purposes :
+	virtual FPrimitiveSceneProxy* CreateSceneProxy() override;
+#endif // !(UE_BUILD_SHIPPING || UE_BUILD_TEST)
+
 	/** Collects custom navigable geometry of component.
     *   Substract the MaxWaveHeight to the Lake collision so nav mesh geometry is exported a ground level
 	*	@return true if regular navigable geometry exporting should be run as well */
diff --git a/Engine/Plugins/Experimental/Water/Source/Runtime/Public/OceanCollisionComponent.h b/Engine/Plugins/Experimental/Water/Source/Runtime/Public/OceanCollisionComponent.h
--- a/Engine/Plugins/Experimental/Water/Source/Runtime/Public/OceanCollisionComponent.h
+++ b/Engine/Plugins/Experimental/Water/Source/Runtime/Public/OceanCollisionComponent.h
@@ -20,10 +20,14 @@ public:
 	void InitializeFromConvexElements(const TArray<FKConvexElem>& ConvexElements);
 
 	virtual bool IsZeroExtent() const override { return BoundingBox.GetExtent().IsZero(); }
-	virtual FPrimitiveSceneProxy* CreateSceneProxy() override;
 	virtual FBoxSphereBounds CalcBounds(const FTransform& LocalToWorld) const;
 	virtual UBodySetup* GetBodySetup() override;
 
+#if !(UE_BUILD_SHIPPING || UE_BUILD_TEST)
+	// The scene proxy is only for debug purposes :
+	virtual FPrimitiveSceneProxy* CreateSceneProxy() override;
+#endif // !(UE_BUILD_SHIPPING || UE_BUILD_TEST)
+
 	/** Collects custom navigable geometry of component.
 	*   Substract the MaxWaveHeight to the Ocean collision so nav mesh geometry is exported a ground level
 	*	@return true if regular navigable geometry exporting should be run as well */
diff --git a/Engine/Source/Runtime/Engine/Classes/Components/SplineComponent.h b/Engine/Source/Runtime/Engine/Classes/Components/SplineComponent.h
--- a/Engine/Source/Runtime/Engine/Classes/Components/SplineComponent.h
+++ b/Engine/Source/Runtime/Engine/Classes/Components/SplineComponent.h
@@ -752,6 +752,17 @@ public:
 	UFUNCTION(BlueprintCallable, Category = Spline)
 	FTransform FindTransformClosestToWorldLocation(const FVector& WorldLocation, ESplineCoordinateSpace::Type CoordinateSpace, bool bUseScale = false) const;
 
+	/** Given a threshold, recursively sub-divides the spline section until the list of segments (polyline) matches the spline shape. */
+	UFUNCTION(BlueprintCallable, Category = Spline)
+	bool DivideSplineIntoPolylineRecursive(float StartDistanceAlongSpline, float EndDistanceAlongSpline, ESplineCoordinateSpace::Type CoordinateSpace, const float MaxSquareDistanceFromSpline, TArray<FVector>& OutPoints) const;
+
+	/** Given a threshold, returns a list of vertices along the spline segment that, treated as a list of segments (polyline), matches the spline shape. */
+	UFUNCTION(BlueprintCallable, Category = Spline)
+	bool ConvertSplineSegmentToPolyLine(int32 SplinePointStartIndex, ESplineCoordinateSpace::Type CoordinateSpace, const float MaxSquareDistanceFromSpline, TArray<FVector>& OutPoints) const;
+
+	/** Given a threshold, returns a list of vertices along the spline that, treated as a list of segments (polyline), matches the spline shape. */
+	UFUNCTION(BlueprintCallable, Category = Spline)
+	bool ConvertSplineToPolyLine(ESplineCoordinateSpace::Type CoordinateSpace, const float MaxSquareDistanceFromSpline, TArray<FVector>& OutPoints) const;
 
 private:
 	/** The dummy value used for queries when there are no point in a spline */
diff --git a/Engine/Source/Runtime/Engine/Private/Components/SplineComponent.cpp b/Engine/Source/Runtime/Engine/Private/Components/SplineComponent.cpp
--- a/Engine/Source/Runtime/Engine/Private/Components/SplineComponent.cpp
+++ b/Engine/Source/Runtime/Engine/Private/Components/SplineComponent.cpp
@@ -1514,6 +1514,105 @@ FTransform USplineComponent::FindTransformClosestToWorldLocation(const FVector&
 	return GetTransformAtSplineInputKey(Param, CoordinateSpace, bUseScale);
 }
 
+bool USplineComponent::DivideSplineIntoPolylineRecursive(float StartDistanceAlongSpline, float EndDistanceAlongSpline, ESplineCoordinateSpace::Type CoordinateSpace, const float MaxSquareDistanceFromSpline, TArray<FVector>& OutPoints) const
+{
+	double Dist = EndDistanceAlongSpline - StartDistanceAlongSpline;
+	if (Dist <= 0.0f)
+	{
+		return false;
+	}
+	double MiddlePointDistancAlongSpline = StartDistanceAlongSpline + Dist / 2.0f;
+
+	FVector Samples[3];
+	Samples[0] = GetLocationAtDistanceAlongSpline(StartDistanceAlongSpline, CoordinateSpace);
+	Samples[1] = GetLocationAtDistanceAlongSpline(MiddlePointDistancAlongSpline, CoordinateSpace);
+	Samples[2] = GetLocationAtDistanceAlongSpline(EndDistanceAlongSpline, CoordinateSpace);
+
+	if (FMath::PointDistToSegmentSquared(Samples[1], Samples[0], Samples[2]) > MaxSquareDistanceFromSpline)
+	{
+		TArray<FVector> NewPoints[2];
+		DivideSplineIntoPolylineRecursive(StartDistanceAlongSpline, MiddlePointDistancAlongSpline, CoordinateSpace, MaxSquareDistanceFromSpline, NewPoints[0]);
+		DivideSplineIntoPolylineRecursive(MiddlePointDistancAlongSpline, EndDistanceAlongSpline, CoordinateSpace, MaxSquareDistanceFromSpline, NewPoints[1]);
+		if ((NewPoints[0].Num() > 0) && (NewPoints[1].Num() > 0))
+		{
+			check(NewPoints[0].Last() == NewPoints[1][0]);
+			NewPoints[0].RemoveAt(NewPoints[0].Num() - 1);
+		}
+		NewPoints[0].Append(NewPoints[1]);
+		OutPoints.Append(NewPoints[0]);
+	}
+	else
+	{
+		// The middle point is close enough to the other 2 points, let's keep those and stop the recursion :
+		OutPoints.Add(Samples[0]);
+		OutPoints.Add(Samples[2]);
+	}
+
+	return (OutPoints.Num() > 0);
+}
+
+bool USplineComponent::ConvertSplineSegmentToPolyLine(int32 SplinePointStartIndex, ESplineCoordinateSpace::Type CoordinateSpace, const float MaxSquareDistanceFromSpline, TArray<FVector>& OutPoints) const
+{
+	OutPoints.Empty();
+
+	const double StartDist = GetDistanceAlongSplineAtSplinePoint(SplinePointStartIndex);
+	const double StopDist = GetDistanceAlongSplineAtSplinePoint(SplinePointStartIndex + 1);
+
+	const int32 NumLines = 2; // Dichotomic subdivision of the spline segment
+	double Dist = StopDist - StartDist;
+	double SubstepSize = Dist / NumLines;
+	if (SubstepSize == 0.0)
+	{
+		// Make sure there's at least 1 sub-step so that we get a single spline vertex for constant interpolation : 
+		SubstepSize = StopDist + 1.0;
+	}
+
+	double SubstepStartDist = StartDist;
+	for (int32 i = 0; i < NumLines; ++i)
+	{
+		double SubstepEndDist = SubstepStartDist + SubstepSize;
+		TArray<FVector> NewPoints;
+		// Recursively sub-divide each segment until the requested precision is reached :
+		if (DivideSplineIntoPolylineRecursive(SubstepStartDist, SubstepEndDist, CoordinateSpace, MaxSquareDistanceFromSpline, NewPoints))
+		{
+			if (OutPoints.Num() > 0)
+			{
+				check(OutPoints.Last() == NewPoints[0]); // our last point must be the same as the new segment's first
+				OutPoints.RemoveAt(OutPoints.Num() - 1);
+			}
+			OutPoints.Append(NewPoints);
+		}
+
+		SubstepStartDist = SubstepEndDist;
+	}
+
+	return (OutPoints.Num() > 0);
+}
+
+
+bool USplineComponent::ConvertSplineToPolyLine(ESplineCoordinateSpace::Type CoordinateSpace, const float MaxSquareDistanceFromSpline, TArray<FVector>& OutPoints) const
+{
+	int32 NumSegments = GetNumberOfSplineSegments();
+	OutPoints.Empty();
+	OutPoints.Reserve(NumSegments * 2); // We sub-divide each segment in at least 2 sub-segments, so let's start with this amount of points
+
+	TArray<FVector> SegmentPoints;
+	for (int32 SegmentIndex = 0; SegmentIndex < NumSegments; ++SegmentIndex)
+	{
+		if (ConvertSplineSegmentToPolyLine(SegmentIndex, CoordinateSpace, MaxSquareDistanceFromSpline, SegmentPoints))
+		{
+			if (OutPoints.Num() > 0)
+			{
+				check(OutPoints.Last() == SegmentPoints[0]); // our last point must be the same as the new segment's first
+				OutPoints.RemoveAt(OutPoints.Num() - 1);
+			}
+			OutPoints.Append(SegmentPoints);
+		}
+	}
+
+	return (OutPoints.Num() > 0);
+}
+
 template<class T>
 T GetPropertyValueAtSplinePoint(const USplineMetadata* Metadata, int32 Index, FName PropertyName)
 {
diff --git a/Engine/Source/Runtime/Engine/Private/GeomTools.cpp b/Engine/Source/Runtime/Engine/Private/GeomTools.cpp
--- a/Engine/Source/Runtime/Engine/Private/GeomTools.cpp
+++ b/Engine/Source/Runtime/Engine/Private/GeomTools.cpp
@@ -1281,6 +1281,7 @@ TArray<TArray<FVector2D>> FGeomTools2D::ReducePolygons(const TArray<TArray<FVect
 
 void FGeomTools2D::CorrectPolygonWinding(TArray<FVector2D>& OutVertices, const TArray<FVector2D>& Vertices, const bool bNegativeWinding)
 {
+	TRACE_CPUPROFILER_EVENT_SCOPE(FGeomTools2D::CorrectPolygonWinding);
 	if (Vertices.Num() >= 3)
 	{
 		// Make sure the polygon winding is correct
@@ -1419,6 +1420,7 @@ static bool AreEdgesMergeable(const FVector2D& V0, const FVector2D& V1, const FV
 // Expected input - PolygonVertices in CCW order, not overlapping
 bool FGeomTools2D::TriangulatePoly(TArray<FVector2D>& OutTris, const TArray<FVector2D>& InPolyVerts, bool bKeepColinearVertices)
 {
+	TRACE_CPUPROFILER_EVENT_SCOPE(FGeomTools2D::TriangulatePoly);
 	// Can't work if not enough verts for 1 triangle
 	if (InPolyVerts.Num() < 3)
 	{
@@ -1522,6 +1524,7 @@ bool FGeomTools2D::TriangulatePoly(TArray<FVector2D>& OutTris, const TArray<FVec
 // 2D version of GeomTools RemoveRedundantTriangles
 void FGeomTools2D::RemoveRedundantTriangles(TArray<FVector2D>& OutTriangles, const TArray<FVector2D>& InTriangleVertices)
 {
+	TRACE_CPUPROFILER_EVENT_SCOPE(FGeomTools2D::RemoveRedundantTriangles);
 	struct FLocalTriangle
 	{
 		int VertexA, VertexB, VertexC;
@@ -1569,6 +1572,7 @@ void FGeomTools2D::RemoveRedundantTriangles(TArray<FVector2D>& OutTriangles, con
 // Find convex polygons from triangle soup
 void FGeomTools2D::GenerateConvexPolygonsFromTriangles(TArray<TArray<FVector2D>>& OutPolygons, const TArray<FVector2D>& InTriangleVertices)
 {
+	TRACE_CPUPROFILER_EVENT_SCOPE(FGeomTools2D::GenerateConvexPolygonsFromTriangles);
 	struct FLocalTriangle
 	{
 		int VertexA, VertexB, VertexC;
