Fixed memory corruption induced by invalid water materials:
- When filling the instance buffers, the WaterInstanceData.BucketInstanceCounts member is reused for storing the actual offset in the buffer, which means that we cannot simply continue to the next mesh batch, because that would lead to this value not being written. The correct fix would be to simply use another array to store the instance data offset, rather than repurposing BucketInstanceCounts but for optimization reasons, we keep the current behavior and simply avoid continuing when encountering an invalid water material
- Also changed a raw pointer to the instance buffer's memory to a TArrayView, in order to have proper bounds checking when filling it (that would have made this memory stomp way more obvious)

#tests editor, PIE, cooked game Win64

diff --git a/Engine/Plugins/Experimental/Water/Source/Runtime/Private/WaterMeshSceneProxy.cpp b/Engine/Plugins/Experimental/Water/Source/Runtime/Private/WaterMeshSceneProxy.cpp
--- a/Engine/Plugins/Experimental/Water/Source/Runtime/Private/WaterMeshSceneProxy.cpp
+++ b/Engine/Plugins/Experimental/Water/Source/Runtime/Private/WaterMeshSceneProxy.cpp
@@ -314,14 +314,16 @@ void FWaterMeshSceneProxy::GetDynamicMeshElements(const TArray<const FSceneView*
 
 					TRACE_CPUPROFILER_EVENT_SCOPE(DensityBucket);
 
-					FMaterialRenderProxy* MaterialRenderProxy = (WireframeMaterialInstance != nullptr) ? WireframeMaterialInstance : WaterQuadTree.GetWaterMaterials()[MaterialIndex];
+					const FMaterialRenderProxy* MaterialRenderProxy = (WireframeMaterialInstance != nullptr) ? WireframeMaterialInstance : WaterQuadTree.GetWaterMaterials()[MaterialIndex];
 					check (MaterialRenderProxy != nullptr);
-					const FMaterial* BucketMaterial = MaterialRenderProxy->GetMaterialNoFallback(GetScene().GetFeatureLevel());
 
-					// If the material is not ready for render, just skip :
-					if (BucketMaterial == nullptr)
+					bool bUseForDepthPass = false;
+
+					// If there's a valid material, use that to figure out the depth pass status
+					if (const FMaterial* BucketMaterial = MaterialRenderProxy->GetMaterialNoFallback(GetScene().GetFeatureLevel()))
 					{
-						continue;
+						// Preemptively turn off depth rendering for this mesh batch if the material doesn't need it
+						bUseForDepthPass = !BucketMaterial->GetShadingModels().HasShadingModel(MSM_SingleLayerWater) && BucketMaterial->GetBlendMode() != EBlendMode::BLEND_Translucent;
 					}
 
 					bMaterialDrawn = true;
@@ -339,7 +341,7 @@ void FWaterMeshSceneProxy::GetDynamicMeshElements(const TArray<const FSceneView*
 						Mesh.bUseForMaterial = true;
 						Mesh.CastShadow = false;
 						// Preemptively turn off depth rendering for this mesh batch if the material doesn't need it
-						Mesh.bUseForDepthPass = !BucketMaterial->GetShadingModels().HasShadingModel(MSM_SingleLayerWater) && BucketMaterial->GetBlendMode() != EBlendMode::BLEND_Translucent;
+						Mesh.bUseForDepthPass = bUseForDepthPass;
 						Mesh.bUseAsOccluder = false;
 
 #if WITH_WATER_SELECTION_SUPPORT
@@ -383,6 +385,8 @@ void FWaterMeshSceneProxy::GetDynamicMeshElements(const TArray<const FSceneView*
 						}
 					}
 
+					// Note : we're repurposing the BucketInstanceCounts array here for storing the actual offset in the buffer. This means that effectively from this point on, BucketInstanceCounts doesn't actually 
+					//  contain the number of instances anymore : 
 					WaterInstanceData.BucketInstanceCounts[BucketIndex] = InstanceDataOffset;
 					InstanceDataOffset += InstanceCount;
 				}
@@ -409,7 +413,9 @@ void FWaterMeshSceneProxy::GetDynamicMeshElements(const TArray<const FSceneView*
 
 				for (int32 StreamIdx = 0; StreamIdx < WaterInstanceDataBuffersType::NumBuffers; ++StreamIdx)
 				{
-					FMemory::Memcpy(WaterInstanceDataBuffers->GetBufferMemory(StreamIdx) + WriteStartOffset, FarDistanceWaterInstanceData.Streams[StreamIdx].GetData(), NumFarInstances * sizeof(FVector4));
+					TArrayView<FVector4> BufferMemory = WaterInstanceDataBuffers->GetBufferMemory(StreamIdx);
+					check(WriteStartOffset + NumFarInstances <= BufferMemory.Num());
+					FMemory::Memcpy(BufferMemory.GetData() + WriteStartOffset, FarDistanceWaterInstanceData.Streams[StreamIdx].GetData(), NumFarInstances * sizeof(FVector4));
 				}
 			}
 		}
diff --git a/Engine/Plugins/Experimental/Water/Source/Runtime/Public/WaterInstanceDataBuffer.h b/Engine/Plugins/Experimental/Water/Source/Runtime/Public/WaterInstanceDataBuffer.h
--- a/Engine/Plugins/Experimental/Water/Source/Runtime/Public/WaterInstanceDataBuffer.h
+++ b/Engine/Plugins/Experimental/Water/Source/Runtime/Public/WaterInstanceDataBuffer.h
@@ -3,6 +3,7 @@
 #pragma once
 
 #include "CoreMinimal.h"
+#include "Containers/ArrayView.h"
 
 template <bool bWithWaterSelectionSupport>
 class TWaterInstanceDataBuffers
@@ -23,7 +24,7 @@ public:
 				for (int32 i = 0; i < NumBuffers; ++i)
 				{
 					Buffer[i] = RHICreateVertexBuffer(SizeInBytes, BUF_Dynamic, CreateInfo);
-					BufferMemory[i] = nullptr;
+					BufferMemory[i] = TArrayView<FVector4>();
 				}
 			}
 		);
@@ -50,7 +51,7 @@ public:
 		for (int32 i = 0; i < NumBuffers; ++i)
 		{
 			Unlock(i);
-			BufferMemory[i] = nullptr;
+			BufferMemory[i] = TArrayView<FVector4>();
 		}
 	}
 
@@ -59,14 +60,14 @@ public:
 		return Buffer[InBufferID];
 	}
 
-	FVector4* GetBufferMemory(int32 InBufferID) const
+	TArrayView<FVector4> GetBufferMemory(int32 InBufferID) const
 	{
-		check(BufferMemory[InBufferID]);
+		check(!BufferMemory[InBufferID].IsEmpty());
 		return BufferMemory[InBufferID];
 	}
 
 private:
-	FVector4* Lock(int32 InInstanceCount, int32 InBufferID)
+	TArrayView<FVector4> Lock(int32 InInstanceCount, int32 InBufferID)
 	{
 		check(IsInRenderingThread());
 
@@ -84,7 +85,8 @@ private:
 			Buffer[InBufferID] = RHICreateVertexBuffer(AlignedSizeInBytes, BUF_Dynamic, CreateInfo);
 		}
 
-		return reinterpret_cast<FVector4*>(RHILockBuffer(Buffer[InBufferID], 0, SizeInBytes, RLM_WriteOnly));
+		FVector4* Data = reinterpret_cast<FVector4*>(RHILockBuffer(Buffer[InBufferID], 0, SizeInBytes, RLM_WriteOnly));
+		return TArrayView<FVector4>(Data, InInstanceCount);
 	}
 
 	void Unlock(int32 InBufferID)
@@ -93,5 +95,5 @@ private:
 	}
 
 	FBufferRHIRef Buffer[NumBuffers];
-	FVector4* BufferMemory[NumBuffers];
+	TArrayView<FVector4> BufferMemory[NumBuffers];
 };
