Fix for no unity issue
#rnx

diff --git a/Engine/Plugins/Experimental/Water/Source/Runtime/Private/WaterBooleanUtils.h b/Engine/Plugins/Experimental/Water/Source/Runtime/Private/WaterBooleanUtils.h
--- a/Engine/Plugins/Experimental/Water/Source/Runtime/Private/WaterBooleanUtils.h
+++ b/Engine/Plugins/Experimental/Water/Source/Runtime/Private/WaterBooleanUtils.h
@@ -12,6 +12,8 @@
 using namespace UE::Geometry;
 
 class AWaterBodyExclusionVolume;
+class AVolume;
+
 
 class FWaterBooleanUtils
 {
