More no unity fixes
#rnx

diff --git a/Engine/Plugins/Experimental/Water/Source/Runtime/Private/WaterBooleanUtils.cpp b/Engine/Plugins/Experimental/Water/Source/Runtime/Private/WaterBooleanUtils.cpp
--- a/Engine/Plugins/Experimental/Water/Source/Runtime/Private/WaterBooleanUtils.cpp
+++ b/Engine/Plugins/Experimental/Water/Source/Runtime/Private/WaterBooleanUtils.cpp
@@ -1,7 +1,5 @@
 // Copyright Epic Games, Inc. All Rights Reserved.
 
-#pragma once
-
 #include "WaterBooleanUtils.h"
 #include "TransformTypes.h"
 #include "VectorUtil.h"
diff --git a/Engine/Plugins/Experimental/Water/Source/Runtime/Private/WaterBooleanUtils.h b/Engine/Plugins/Experimental/Water/Source/Runtime/Private/WaterBooleanUtils.h
--- a/Engine/Plugins/Experimental/Water/Source/Runtime/Private/WaterBooleanUtils.h
+++ b/Engine/Plugins/Experimental/Water/Source/Runtime/Private/WaterBooleanUtils.h
@@ -13,6 +13,7 @@ using namespace UE::Geometry;
 
 class AWaterBodyExclusionVolume;
 class AVolume;
+struct FKConvexElem;
 
 
 class FWaterBooleanUtils
