Fixed missing UFUNCTION on WaterBodyComponent::GetWaterWaves

diff --git a/Engine/Plugins/Experimental/Water/Source/Runtime/Public/WaterBodyComponent.h b/Engine/Plugins/Experimental/Water/Source/Runtime/Public/WaterBodyComponent.h
--- a/Engine/Plugins/Experimental/Water/Source/Runtime/Public/WaterBodyComponent.h
+++ b/Engine/Plugins/Experimental/Water/Source/Runtime/Public/WaterBodyComponent.h
@@ -141,6 +141,7 @@ public:
 	UFUNCTION(BlueprintCallable, Category = Water)
 	UWaterSplineComponent* GetWaterSpline() const;
 
+	UFUNCTION(BlueprintCallable, Category = Water)
 	UWaterWavesBase* GetWaterWaves() const;
 
 	/** Returns the unique id of this water body for accessing data in GPU buffers */
