Public-export RecomputeWaves for non-engine code to use when manually providing gerstner waves data

#rnx

diff --git a/Engine/Plugins/Experimental/Water/Source/Runtime/Public/GerstnerWaterWaves.h b/Engine/Plugins/Experimental/Water/Source/Runtime/Public/GerstnerWaterWaves.h
--- a/Engine/Plugins/Experimental/Water/Source/Runtime/Public/GerstnerWaterWaves.h
+++ b/Engine/Plugins/Experimental/Water/Source/Runtime/Public/GerstnerWaterWaves.h
@@ -206,7 +206,7 @@ public:
 #endif
 
 	/** Call RecomputeWaves whenever wave data changes, so that all cached data can be recomputed (do not call OnPostLoad... can call BP script internally) */
-	void RecomputeWaves(bool bAllowBPScript);
+	WATER_API void RecomputeWaves(bool bAllowBPScript);
 
 private:
 	FVector GetWaveOffsetAtPosition(const FGerstnerWave& InWaveParams, const FVector& InPosition, float InTime, FVector& OutNormal, float& OutOffset1D) const;
