Fixed Word Partitioned actors not being properly registered to the water brush

diff --git a/Engine/Plugins/Experimental/Water/Source/Editor/Private/WaterLandscapeBrush.cpp b/Engine/Plugins/Experimental/Water/Source/Editor/Private/WaterLandscapeBrush.cpp
--- a/Engine/Plugins/Experimental/Water/Source/Editor/Private/WaterLandscapeBrush.cpp
+++ b/Engine/Plugins/Experimental/Water/Source/Editor/Private/WaterLandscapeBrush.cpp
@@ -274,21 +274,20 @@ void AWaterLandscapeBrush::PostInitProperties()
 			}
 		});
 
-		OnLevelActorAddedHandle = GEngine->OnLevelActorAdded().AddLambda([this](AActor* Actor)
-		{
-			const UWorld* ThisWorld = GetWorld();
-			const bool bTriggerEvent = true;
-			const bool bModify = true;
-			AddActorInternal(Actor, ThisWorld, nullptr, bTriggerEvent, bModify);
-		});
+		OnLevelActorAddedHandle = GEngine->OnLevelActorAdded().AddUObject(this, &AWaterLandscapeBrush::OnLevelActorAdded);
+		OnLevelActorDeletedHandle = GEngine->OnLevelActorDeleted().AddUObject(this, &AWaterLandscapeBrush::OnLevelActorRemoved);
 
-		OnLevelActorDeletedHandle = GEngine->OnLevelActorDeleted().AddLambda([this](AActor* Actor)
+#if WITH_EDITOR
+		// in world partition, actors don't belong to levels and the on loaded/removed callbacks are different :
+		if (UWorld* World = GetWorld())
 		{
-			if (IsActorAffectingLandscape(Actor))
+			if (World->PersistentLevel)
 			{
-				RemoveActorInternal(Actor);
+				OnLoadedActorAddedToLevelEventHandle = World->PersistentLevel->OnLoadedActorAddedToLevelEvent.AddLambda([this](AActor& InActor) { OnLevelActorAdded(&InActor); });
+				OnLoadedActorRemovedFromLevelEventHandle = World->PersistentLevel->OnLoadedActorRemovedFromLevelEvent.AddLambda([this](AActor& InActor) { OnLevelActorRemoved(&InActor); });
 			}
-		});
+		}
+#endif // WITH_EDITOR
 
 		OnActorMovedHandle = GEngine->OnActorMoved().AddLambda([this](AActor* Actor)
 		{
@@ -305,6 +304,25 @@ void AWaterLandscapeBrush::PostInitProperties()
 	UpdateActors(!GIsEditorLoadingPackage);
 }
 
+void AWaterLandscapeBrush::OnLevelActorAdded(AActor* InActor)
+{
+	if (InActor->GetWorld() == GetWorld())
+	{
+		AddActorInternal(InActor, GetWorld(), /*InCache = */nullptr, /*bTriggerEvent = */true, /*bModify = */true);
+	}
+}
+
+void AWaterLandscapeBrush::OnLevelActorRemoved(AActor* InActor)
+{
+	if (InActor->GetWorld() == GetWorld())
+	{
+		if (IsActorAffectingLandscape(InActor))
+		{
+			RemoveActorInternal(InActor);
+		}
+	}
+}
+
 void AWaterLandscapeBrush::OnConstruction(const FTransform& Transform)
 {
 	Super::OnConstruction(Transform);
@@ -343,6 +361,20 @@ void AWaterLandscapeBrush::BeginDestroy()
 		GEngine->OnActorMoved().Remove(OnActorMovedHandle);
 		OnActorMovedHandle.Reset();
 
+#if WITH_EDITOR
+		// Unregister callbacks
+		if (UWorld* World = GetWorld())
+		{
+			if (World->PersistentLevel)
+			{
+				World->PersistentLevel->OnLoadedActorAddedToLevelEvent.Remove(OnLoadedActorAddedToLevelEventHandle);
+				World->PersistentLevel->OnLoadedActorRemovedFromLevelEvent.Remove(OnLoadedActorRemovedFromLevelEventHandle);
+			}
+		}
+		OnLoadedActorAddedToLevelEventHandle.Reset();
+		OnLoadedActorRemovedFromLevelEventHandle.Reset();
+#endif // WITH_EDITOR
+
 		IWaterBrushActorInterface::GetOnWaterBrushActorChangedEvent().RemoveAll(this);
 	}
 }
diff --git a/Engine/Plugins/Experimental/Water/Source/Editor/Private/WaterLandscapeBrush.h b/Engine/Plugins/Experimental/Water/Source/Editor/Private/WaterLandscapeBrush.h
--- a/Engine/Plugins/Experimental/Water/Source/Editor/Private/WaterLandscapeBrush.h
+++ b/Engine/Plugins/Experimental/Water/Source/Editor/Private/WaterLandscapeBrush.h
@@ -121,6 +121,8 @@ private:
 	void OnWaterBrushActorChanged(const IWaterBrushActorInterface::FWaterBrushActorChangedEventParams& InParams);
 	void OnActorChanged(AActor* Actor, bool bWeightmapSettingsChanged, bool bRebuildMesh);
 	void OnActorsAffectingLandscapeChanged();
+	void OnLevelActorAdded(AActor* InActor);
+	void OnLevelActorRemoved(AActor* InActor);
 
 	TArray<TWeakInterfacePtr<IWaterBrushActorInterface>> ActorsAffectingLandscape;
 	FDelegateHandle OnWorldPostInitHandle;
@@ -129,6 +131,10 @@ private:
 	FDelegateHandle OnLevelActorAddedHandle;
 	FDelegateHandle OnLevelActorDeletedHandle;
 	FDelegateHandle OnActorMovedHandle;
+#if WITH_EDITOR
+	FDelegateHandle OnLoadedActorAddedToLevelEventHandle;
+	FDelegateHandle OnLoadedActorRemovedFromLevelEventHandle;
+#endif // WITH_EDITOR
 
 	UPROPERTY(Transient, DuplicateTransient, VisibleAnywhere, AdvancedDisplay, meta = (Category = "Debug"))
 	TMap<TWeakObjectPtr<AActor>, UObject*> Cache;
