PR #8365: CVarOverrideWavesTime should be a float instead of an integer (Contributed by [REDACTED_PERSON])

#rnx

diff --git a/Engine/Plugins/Experimental/Water/Source/Runtime/Private/WaterSubsystem.cpp b/Engine/Plugins/Experimental/Water/Source/Runtime/Private/WaterSubsystem.cpp
--- a/Engine/Plugins/Experimental/Water/Source/Runtime/Private/WaterSubsystem.cpp
+++ b/Engine/Plugins/Experimental/Water/Source/Runtime/Private/WaterSubsystem.cpp
@@ -51,7 +51,7 @@ static FAutoConsoleVariableRef CVarFreezeWaves(
 	ECVF_Cheat
 );
 
-static TAutoConsoleVariable<int32> CVarOverrideWavesTime(
+static TAutoConsoleVariable<float> CVarOverrideWavesTime(
 	TEXT("r.Water.OverrideWavesTime"),
 	-1.0f,
 	TEXT("Forces the time used for waves if >= 0.0"),
@@ -803,4 +803,4 @@ UWaterSubsystem::FScopedAllowWaterSubsystemOnPreviewWorld::~FScopedAllowWaterSub
 
 // ----------------------------------------------------------------------------------
 
-#undef LOCTEXT_NAMESPACE
\ No newline at end of file
+#undef LOCTEXT_NAMESPACE
