Implement SetMaterial on WaterMeshComponent to make sure Override Materials aren't being set on it. This is not supported since the materials of this component are auto-populated from Water Bodies and not mutable.

diff --git a/Engine/Plugins/Experimental/Water/Source/Runtime/Private/WaterMeshComponent.cpp b/Engine/Plugins/Experimental/Water/Source/Runtime/Private/WaterMeshComponent.cpp
--- a/Engine/Plugins/Experimental/Water/Source/Runtime/Private/WaterMeshComponent.cpp
+++ b/Engine/Plugins/Experimental/Water/Source/Runtime/Private/WaterMeshComponent.cpp
@@ -108,6 +108,11 @@ void UWaterMeshComponent::GetUsedMaterials(TArray<UMaterialInterface*>& OutMater
 	}
 }
 
+void UWaterMeshComponent::SetMaterial(int32 ElementIndex, UMaterialInterface* Material)
+{
+	UE_LOG(LogWater, Warning, TEXT("SetMaterial is not compatible with UWaterMeshComponent since all materials on this component are auto-populated from the Water Bodies contained within it."));
+}
+
 bool UWaterMeshComponent::ShouldRenderSelected() const
 {
 	if (bSelectable)
diff --git a/Engine/Plugins/Experimental/Water/Source/Runtime/Public/WaterMeshComponent.h b/Engine/Plugins/Experimental/Water/Source/Runtime/Public/WaterMeshComponent.h
--- a/Engine/Plugins/Experimental/Water/Source/Runtime/Public/WaterMeshComponent.h
+++ b/Engine/Plugins/Experimental/Water/Source/Runtime/Public/WaterMeshComponent.h
@@ -37,6 +37,7 @@ public:
 	//~ Begin UPrimitiveComponent Interface
 	virtual FPrimitiveSceneProxy* CreateSceneProxy() override;
 	virtual void GetUsedMaterials(TArray<UMaterialInterface*>& OutMaterials, bool bGetDebugMaterials = false) const override;
+	virtual void SetMaterial(int32 ElementIndex, UMaterialInterface* Material) override;
 	virtual bool ShouldRenderSelected() const override;
 	//~ End UPrimitiveComponent Interface
 
