Fixed race condition that leads to randomly giant water icons

diff --git a/Engine/Plugins/Experimental/Water/Source/Runtime/Private/WaterIconHelper.cpp b/Engine/Plugins/Experimental/Water/Source/Runtime/Private/WaterIconHelper.cpp
--- a/Engine/Plugins/Experimental/Water/Source/Runtime/Private/WaterIconHelper.cpp
+++ b/Engine/Plugins/Experimental/Water/Source/Runtime/Private/WaterIconHelper.cpp
@@ -46,9 +46,10 @@ void FWaterIconHelper::UpdateSpriteComponent(AActor* Actor, UTexture2D* InTextur
 		FVector ZOffset(0.0f, 0.0f, GetDefault<UWaterRuntimeSettings>()->WaterBodyIconWorldZOffset);
 		if (InTexture != nullptr)
 		{
-			int32 TextureSize = FMath::Max(InTexture->GetSizeX(), InTexture->GetSizeY());
-			float Scale = (float)TargetSize / TextureSize;
-			ActorIcon->SetRelativeScale3D(FVector((TextureSize > 0) ? Scale : 1.0f));
+			// Use the texture source's size as the texture might not having finished loading when this runs, in which case the default texture's size would be returned : 
+			int32 TextureSize = FMath::Max(InTexture->Source.GetSizeX(), InTexture->Source.GetSizeY());
+			float Scale = (TextureSize > 0) ? (TargetSize / (float)TextureSize) : 1.0f;
+			ActorIcon->SetRelativeScale3D(FVector(Scale));
 		}
 		ActorIcon->Sprite = InTexture;
 		ActorIcon->SetRelativeLocation(ZOffset);
