Setting the engine's bCanEverAffectNavigation from the WaterBody bCanAffectNavigation as when GEditor is false, the value isn't propagated through UWaterBodyComponent::UpdateAll -> ApplyNavigationSettings().
We should not port this change to an engine level but remove bCanAffectNavigation from the UWaterBodyComponent and use the engine UActorComponent's one bCanEverAffectNavigation.

#rnx

diff --git a/Engine/Plugins/Experimental/Water/Source/Runtime/Private/WaterBodyComponent.cpp b/Engine/Plugins/Experimental/Water/Source/Runtime/Private/WaterBodyComponent.cpp
--- a/Engine/Plugins/Experimental/Water/Source/Runtime/Private/WaterBodyComponent.cpp
+++ b/Engine/Plugins/Experimental/Water/Source/Runtime/Private/WaterBodyComponent.cpp
@@ -998,13 +998,16 @@ bool UWaterBodyComponent::IsNavigationRelevant() const
 	return CanAffectNavigation() && (GetCollisionComponents().Num() > 0);
 }
 
-void UWaterBodyComponent::ApplyNavigationSettings() const
+void UWaterBodyComponent::ApplyNavigationSettings()
 {
 	AActor* Owner = GetOwner();
 	if (Owner)
 	{
 		const bool bCanAffectNav = CanAffectNavigation();
 
+		// Make sure the engine's bCanEverAffectionNavigation matches the component's one.
+		SetCanEverAffectNavigation(bCanAffectNav);
+
 		// @todo_water: change this
 		TInlineComponentArray<UActorComponent*> Components;
 		Owner->GetComponents(Components);
diff --git a/Engine/Plugins/Experimental/Water/Source/Runtime/Public/WaterBodyComponent.h b/Engine/Plugins/Experimental/Water/Source/Runtime/Public/WaterBodyComponent.h
--- a/Engine/Plugins/Experimental/Water/Source/Runtime/Public/WaterBodyComponent.h
+++ b/Engine/Plugins/Experimental/Water/Source/Runtime/Public/WaterBodyComponent.h
@@ -357,7 +357,7 @@ protected:
 	void CreateOrUpdateWaterMID();
 	void CreateOrUpdateUnderwaterPostProcessMID();
 	void PrepareCurrentPostProcessSettings();
-	void ApplyNavigationSettings() const;
+	void ApplyNavigationSettings();
 	void RequestGPUWaveDataUpdate();
 	EObjectFlags GetTransientMIDFlags() const; 
 
