Water component bCanAffectNav should be reflected in the UPrimitiveComponent's version (bCanEverAffectNavigation)

diff --git a/Engine/Plugins/Experimental/Water/Source/Runtime/Private/WaterBodyComponent.cpp b/Engine/Plugins/Experimental/Water/Source/Runtime/Private/WaterBodyComponent.cpp
--- a/Engine/Plugins/Experimental/Water/Source/Runtime/Private/WaterBodyComponent.cpp
+++ b/Engine/Plugins/Experimental/Water/Source/Runtime/Private/WaterBodyComponent.cpp
@@ -1016,6 +1016,10 @@ void UWaterBodyComponent::ApplyNavigationSettings()
 		for (UActorComponent* ActorComp : Components)
 		{
 			UPrimitiveComponent* PrimComp = Cast<UPrimitiveComponent>(ActorComp);
+			if (PrimComp == this)
+			{
+				continue;
+			}
 			if (!PrimComp || LocalCollisionComponents.Find(PrimComp) == INDEX_NONE)
 			{
 				ActorComp->SetCanEverAffectNavigation(false);
