Fixed undo behavior for water component : river spline mesh components (amongst others) were not marked as RF_Transactional, which meant they were flagged as PendingKill when they were destroyed but weren't being un-flagged when they got restored by a transaction. This meant that a TWeakObjectPtr pointing to it would fail, leading to invalid renders in the water brush manager

#tests editor

diff --git a/Engine/Plugins/Experimental/Water/Source/Editor/Private/WaterBrushManager.cpp b/Engine/Plugins/Experimental/Water/Source/Editor/Private/WaterBrushManager.cpp
--- a/Engine/Plugins/Experimental/Water/Source/Editor/Private/WaterBrushManager.cpp
+++ b/Engine/Plugins/Experimental/Water/Source/Editor/Private/WaterBrushManager.cpp
@@ -1470,7 +1470,7 @@ void AWaterBrushManager::UpdateBrushCacheKeys()
 			// If no cache entry, create one and associate it to the actor : 
 			if (WaterBrushCacheContainer == nullptr)
 			{
-				WaterBrushCacheContainer = NewObject<UWaterBodyBrushCacheContainer>(this);
+				WaterBrushCacheContainer = NewObject<UWaterBodyBrushCacheContainer>(this, TEXT("WaterBodyBrushCacheContainer"), RF_Transactional);
 				check(!WaterBrushCacheContainer->Cache.CacheIsValid);
 				SetActorCache(Actor, WaterBrushCacheContainer);
 			}
diff --git a/Engine/Plugins/Experimental/Water/Source/Runtime/Private/WaterBodyActor.cpp b/Engine/Plugins/Experimental/Water/Source/Runtime/Private/WaterBodyActor.cpp
--- a/Engine/Plugins/Experimental/Water/Source/Runtime/Private/WaterBodyActor.cpp
+++ b/Engine/Plugins/Experimental/Water/Source/Runtime/Private/WaterBodyActor.cpp
@@ -426,7 +426,7 @@ void AWaterBody::DeprecateData()
 				{
 					void* OldPropertyOnWaveSpectrumSettings = OldWaveStructProperty->ContainerPtrToValuePtr<void>(this);
 					// We need to propagate object flags to the sub objects (if we deprecate an archetype's data, it is public and its sub-object need to be as well) :
-					EObjectFlags NewFlags = GetMaskedFlags(RF_PropagateToSubObjects);
+					EObjectFlags NewFlags = GetMaskedFlags(RF_PropagateToSubObjects) | RF_Transactional;
 					UGerstnerWaterWaves* GerstnerWaves = NewObject<UGerstnerWaterWaves>(this, MakeUniqueObjectName(this, UGerstnerWaterWaves::StaticClass(), TEXT("GestnerWaterWaves")), NewFlags);
 					UClass* NewGerstnerClass = UGerstnerWaterWaveGeneratorSimple::StaticClass();
 					UGerstnerWaterWaveGeneratorSimple* GerstnerWavesGenerator = NewObject<UGerstnerWaterWaveGeneratorSimple>(this, MakeUniqueObjectName(this, NewGerstnerClass, TEXT("GestnerWaterWavesGenerator")), NewFlags);
diff --git a/Engine/Plugins/Experimental/Water/Source/Runtime/Private/WaterBodyCustomComponent.cpp b/Engine/Plugins/Experimental/Water/Source/Runtime/Private/WaterBodyCustomComponent.cpp
--- a/Engine/Plugins/Experimental/Water/Source/Runtime/Private/WaterBodyCustomComponent.cpp
+++ b/Engine/Plugins/Experimental/Water/Source/Runtime/Private/WaterBodyCustomComponent.cpp
@@ -64,7 +64,7 @@ void UWaterBodyCustomComponent::OnUpdateBody(bool bWithExclusionVolumes)
 
 	if (!MeshComp)
 	{
-		MeshComp = NewObject<UStaticMeshComponent>(OwnerActor, TEXT("CustomMeshComponent"));
+		MeshComp = NewObject<UStaticMeshComponent>(OwnerActor, TEXT("CustomMeshComponent"), RF_Transactional);
 		MeshComp->SetNetAddressable(); // it's deterministically named so it's addressable over network (needed for collision)
 		MeshComp->SetupAttachment(this);
 		MeshComp->SetCollisionProfileName(GetCollisionProfileName());
diff --git a/Engine/Plugins/Experimental/Water/Source/Runtime/Private/WaterBodyLakeComponent.cpp b/Engine/Plugins/Experimental/Water/Source/Runtime/Private/WaterBodyLakeComponent.cpp
--- a/Engine/Plugins/Experimental/Water/Source/Runtime/Private/WaterBodyLakeComponent.cpp
+++ b/Engine/Plugins/Experimental/Water/Source/Runtime/Private/WaterBodyLakeComponent.cpp
@@ -73,7 +73,7 @@ void UWaterBodyLakeComponent::OnUpdateBody(bool bWithExclusionVolumes)
 	
 	if (!LakeMeshComp)
 	{
-		LakeMeshComp = NewObject<UStaticMeshComponent>(OwnerActor, TEXT("LakeMeshComponent"));
+		LakeMeshComp = NewObject<UStaticMeshComponent>(OwnerActor, TEXT("LakeMeshComponent"), RF_Transactional);
 		LakeMeshComp->SetupAttachment(this);
 		LakeMeshComp->RegisterComponent();
 	}
@@ -82,7 +82,7 @@ void UWaterBodyLakeComponent::OnUpdateBody(bool bWithExclusionVolumes)
 	{
 		if (!LakeCollision)
 		{
-			LakeCollision = NewObject<ULakeCollisionComponent>(OwnerActor, TEXT("LakeCollisionComponent"));
+			LakeCollision = NewObject<ULakeCollisionComponent>(OwnerActor, TEXT("LakeCollisionComponent"), RF_Transactional);
 			LakeCollision->SetNetAddressable(); // it's deterministically named so it's addressable over network (needed for collision)
 			LakeCollision->SetupAttachment(this);
 			LakeCollision->RegisterComponent();
diff --git a/Engine/Plugins/Experimental/Water/Source/Runtime/Private/WaterBodyOceanComponent.cpp b/Engine/Plugins/Experimental/Water/Source/Runtime/Private/WaterBodyOceanComponent.cpp
--- a/Engine/Plugins/Experimental/Water/Source/Runtime/Private/WaterBodyOceanComponent.cpp
+++ b/Engine/Plugins/Experimental/Water/Source/Runtime/Private/WaterBodyOceanComponent.cpp
@@ -153,7 +153,7 @@ void UWaterBodyOceanComponent::OnUpdateBody(bool bWithExclusionVolumes)
 			}
 			else
 			{
-				BoxComponent = NewObject<UOceanBoxCollisionComponent>(OwnerActor, Name);
+				BoxComponent = NewObject<UOceanBoxCollisionComponent>(OwnerActor, Name, RF_Transactional);
 				BoxComponent->SetupAttachment(this);
 				CollisionBoxes.Add(BoxComponent);
 			}
@@ -191,7 +191,7 @@ void UWaterBodyOceanComponent::OnUpdateBody(bool bWithExclusionVolumes)
 			}
 			else
 			{
-				CollisionComponent = NewObject<UOceanCollisionComponent>(OwnerActor, Name);
+				CollisionComponent = NewObject<UOceanCollisionComponent>(OwnerActor, Name, RF_Transactional);
 				CollisionComponent->SetupAttachment(this);
 				CollisionHullSets.Add(CollisionComponent);
 			}
diff --git a/Engine/Plugins/Experimental/Water/Source/Runtime/Private/WaterBodyRiverComponent.cpp b/Engine/Plugins/Experimental/Water/Source/Runtime/Private/WaterBodyRiverComponent.cpp
--- a/Engine/Plugins/Experimental/Water/Source/Runtime/Private/WaterBodyRiverComponent.cpp
+++ b/Engine/Plugins/Experimental/Water/Source/Runtime/Private/WaterBodyRiverComponent.cpp
@@ -179,7 +179,7 @@ void UWaterBodyRiverComponent::GenerateMeshes()
 		for (int32 SplinePtIndex = 0; SplinePtIndex < NumberOfMeshComponentsNeeded; ++SplinePtIndex)
 		{
 			FString Name = FString::Printf(TEXT("SplineMeshComponent_%d"), SplinePtIndex);
-			USplineMeshComponent* MeshComp = NewObject<USplineMeshComponent>(Owner, *Name);
+			USplineMeshComponent* MeshComp = NewObject<USplineMeshComponent>(Owner, *Name, RF_Transactional);
 			MeshComp->SetNetAddressable(); // it's deterministically named so it's addressable over network (needed for collision)
 			SplineMeshComponents.Add(MeshComp);
 			MeshComp->SetMobility(Owner->GetRootComponent()->Mobility);
