Expose current water bodies to blueprint via buoyancy component. Expose water information at location query to blueprint

diff --git a/Engine/Plugins/Experimental/Water/Source/Runtime/Private/WaterBodyComponent.cpp b/Engine/Plugins/Experimental/Water/Source/Runtime/Private/WaterBodyComponent.cpp
--- a/Engine/Plugins/Experimental/Water/Source/Runtime/Private/WaterBodyComponent.cpp
+++ b/Engine/Plugins/Experimental/Water/Source/Runtime/Private/WaterBodyComponent.cpp
@@ -45,9 +45,10 @@
 
 DECLARE_CYCLE_STAT(TEXT("WaterBody_ComputeWaterInfo"), STAT_WaterBody_ComputeWaterInfo, STATGROUP_Water);
 DECLARE_CYCLE_STAT(TEXT("WaterBody_ComputeWaterDepth"), STAT_WaterBody_ComputeWaterDepth, STATGROUP_Water);
+DECLARE_CYCLE_STAT(TEXT("WaterBody_ComputeWaterDepth"), STAT_WaterBody_ComputeLocation, STATGROUP_Water);
+DECLARE_CYCLE_STAT(TEXT("WaterBody_ComputeWaterDepth"), STAT_WaterBody_ComputeNormal, STATGROUP_Water);
 DECLARE_CYCLE_STAT(TEXT("WaterBody_ComputeLandscapeDepth"), STAT_WaterBody_ComputeLandscapeDepth, STATGROUP_Water);
 DECLARE_CYCLE_STAT(TEXT("WaterBody_ComputeWaveHeight"), STAT_WaterBody_ComputeWaveHeight, STATGROUP_Water);
-
 // ----------------------------------------------------------------------------------
 
 TAutoConsoleVariable<float> CVarWaterOceanFallbackDepth(
@@ -396,6 +397,7 @@ FWaterBodyQueryResult UWaterBodyComponent::QueryWaterInfoClosestToWorldLocation(
 	// Compute water plane location :
 	if (EnumHasAnyFlags(Result.GetQueryFlags(), EWaterBodyQueryFlags::ComputeLocation))
 	{
+		SCOPE_CYCLE_COUNTER(STAT_WaterBody_ComputeLocation);
 		FVector WaterPlaneLocation = InWorldLocation;
 		// If in exclusion volume, force the water plane location at the query location. It is technically invalid, but it's up to the caller to check whether we're in an exclusion volume. 
 		//  If the user fails to do so, at least it allows immersion depth to be 0.0f, which means the query location is NOT in water :
@@ -419,6 +421,7 @@ FWaterBodyQueryResult UWaterBodyComponent::QueryWaterInfoClosestToWorldLocation(
 	FVector WaterPlaneNormal = FVector::UpVector;
 	if (EnumHasAnyFlags(Result.GetQueryFlags(), EWaterBodyQueryFlags::ComputeNormal))
 	{
+		SCOPE_CYCLE_COUNTER(STAT_WaterBody_ComputeNormal);
 		// Default to Z up for the normal
 		if (!bFlatSurface)
 		{
@@ -462,7 +465,7 @@ FWaterBodyQueryResult UWaterBodyComponent::QueryWaterInfoClosestToWorldLocation(
 				}
 			}
 
-			// If the height is invalid, we either have invalid landscape data or we're under the landscape
+			// If the height is invalid, we either have invalid landscape data or we're under the 
 			if (!bValidLandscapeData || (WaterPlaneDepth < 0.0f))
 			{
 				if (GetWaterBodyType() == EWaterBodyType::Ocean)
@@ -549,6 +552,30 @@ FWaterBodyQueryResult UWaterBodyComponent::QueryWaterInfoClosestToWorldLocation(
 	}
 
 	return Result;
+
+}
+
+void UWaterBodyComponent::GetWaterSurfaceInfoAtLocation(const FVector& InLocation, FVector& OutWaterSurfaceLocation, FVector& OutWaterSurfaceNormal, FVector& OutWaterVelocity, float& OutWaterDepth, bool bIncludeDepth /* = false */) const
+{
+	EWaterBodyQueryFlags QueryFlags =
+		EWaterBodyQueryFlags::ComputeLocation
+		| EWaterBodyQueryFlags::ComputeNormal
+		| EWaterBodyQueryFlags::ComputeVelocity;
+
+	if (bIncludeDepth)
+	{
+		QueryFlags |= EWaterBodyQueryFlags::ComputeDepth;
+	}
+
+	FWaterBodyQueryResult QueryResult = QueryWaterInfoClosestToWorldLocation(InLocation, QueryFlags);
+	OutWaterSurfaceLocation = QueryResult.GetWaterSurfaceLocation();
+	OutWaterSurfaceNormal = QueryResult.GetWaterSurfaceNormal();
+	OutWaterVelocity = QueryResult.GetVelocity();
+
+	if (bIncludeDepth)
+	{
+		OutWaterDepth = QueryResult.GetWaterSurfaceDepth();
+	}
 }
 
 float UWaterBodyComponent::GetWaterVelocityAtSplineInputKey(float InKey) const
diff --git a/Engine/Plugins/Experimental/Water/Source/Runtime/Public/BuoyancyComponent.h b/Engine/Plugins/Experimental/Water/Source/Runtime/Public/BuoyancyComponent.h
--- a/Engine/Plugins/Experimental/Water/Source/Runtime/Public/BuoyancyComponent.h
+++ b/Engine/Plugins/Experimental/Water/Source/Runtime/Public/BuoyancyComponent.h
@@ -60,6 +60,7 @@ public:
 	void EnteredWaterBody(UWaterBodyComponent* WaterBodyComponent);
 	void ExitedWaterBody(UWaterBodyComponent* WaterBodyComponent);
 
+	UFUNCTION(BlueprintCallable, Category = Buoyancy)
 	const TArray<UWaterBodyComponent*>& GetCurrentWaterBodyComponents() const { return CurrentWaterBodyComponents; }
 	TArray<UWaterBodyComponent*>& GetCurrentWaterBodyComponents() { return CurrentWaterBodyComponents; }
 
diff --git a/Engine/Plugins/Experimental/Water/Source/Runtime/Public/WaterBodyComponent.h b/Engine/Plugins/Experimental/Water/Source/Runtime/Public/WaterBodyComponent.h
--- a/Engine/Plugins/Experimental/Water/Source/Runtime/Public/WaterBodyComponent.h
+++ b/Engine/Plugins/Experimental/Water/Source/Runtime/Public/WaterBodyComponent.h
@@ -216,6 +216,9 @@ public:
 	*/
 	virtual FWaterBodyQueryResult QueryWaterInfoClosestToWorldLocation(const FVector& InWorldLocation, EWaterBodyQueryFlags InQueryFlags, const TOptional<float>& InSplineInputKey = TOptional<float>()) const;
 
+	UFUNCTION(BlueprintCallable, Category = WaterBody)
+	void GetWaterSurfaceInfoAtLocation(const FVector& InLocation, FVector& OutWaterSurfaceLocation, FVector& OutWaterSurfaceNormal, FVector& OutWaterVelocity, float& OutWaterDepth, bool bIncludeDepth = false) const;
+
 	/** Spline query helper. It's faster to get the spline key once then query properties using that key, rather than querying repeatedly by location etc. */
 	float FindInputKeyClosestToWorldLocation(const FVector& WorldLocation) const;
 
