Fixed water brush manager mistakingly reusing render targets from another water body because of NewObject's reuse object policy, leading to completely invalid landscape

#tests editor
#rnx

diff --git a/Engine/Plugins/Experimental/Water/Source/Editor/Private/WaterBrushManager.cpp b/Engine/Plugins/Experimental/Water/Source/Editor/Private/WaterBrushManager.cpp
--- a/Engine/Plugins/Experimental/Water/Source/Editor/Private/WaterBrushManager.cpp
+++ b/Engine/Plugins/Experimental/Water/Source/Editor/Private/WaterBrushManager.cpp
@@ -1470,7 +1470,7 @@ void AWaterBrushManager::UpdateBrushCacheKeys()
 			// If no cache entry, create one and associate it to the actor : 
 			if (WaterBrushCacheContainer == nullptr)
 			{
-				WaterBrushCacheContainer = NewObject<UWaterBodyBrushCacheContainer>(this, TEXT("WaterBodyBrushCacheContainer"), RF_Transactional);
+				WaterBrushCacheContainer = NewObject<UWaterBodyBrushCacheContainer>(this, NAME_None, RF_Transactional);
 				check(!WaterBrushCacheContainer->Cache.CacheIsValid);
 				SetActorCache(Actor, WaterBrushCacheContainer);
 			}
