Fixed custom water bodies being invisible when dynamically spawned : UpdateComponentVisibility is called right before the level they belonged is made visible (via RegisterAllComponents). This means that calling IsVisible() on the water body component (to transfer the information to sub-components) will return false even if it is actually visible (but is only invisible because of its parent level being so). Now we simply use GetVisibileFlag() instead and let the primitive component call IsVisible whenever the level's visibility changes

#tests editor, PIE

diff --git a/Engine/Plugins/Experimental/Water/Source/Runtime/Private/WaterBodyComponent.cpp b/Engine/Plugins/Experimental/Water/Source/Runtime/Private/WaterBodyComponent.cpp
--- a/Engine/Plugins/Experimental/Water/Source/Runtime/Private/WaterBodyComponent.cpp
+++ b/Engine/Plugins/Experimental/Water/Source/Runtime/Private/WaterBodyComponent.cpp
@@ -759,7 +759,7 @@ void UWaterBodyComponent::UpdateComponentVisibility(bool bAllowWaterMeshRebuild)
 	 	const bool bIsWaterRenderingEnabled = FWaterUtils::IsWaterEnabled(/*bIsRenderThread = */false);
 	 
 		bool bIsRenderedByWaterMesh = ShouldGenerateWaterMeshTile();
-		bool bLocalVisible = bIsWaterRenderingEnabled && !bIsRenderedByWaterMesh && IsVisible();
+		bool bLocalVisible = bIsWaterRenderingEnabled && !bIsRenderedByWaterMesh && GetVisibleFlag();
 		bool bLocalHiddenInGame = !bIsWaterRenderingEnabled || bIsRenderedByWaterMesh || bHiddenInGame;
 
 	 	for (UPrimitiveComponent* Component : GetStandardRenderableComponents())
