Refactored Far Mesh code to now be part of the water quadtree instead. No more special case handling in the proxy needed anymore and some old dead code has been removed.

#rnx

diff --git a/Engine/Plugins/Experimental/Water/Source/Runtime/Private/WaterMeshComponent.cpp b/Engine/Plugins/Experimental/Water/Source/Runtime/Private/WaterMeshComponent.cpp
--- a/Engine/Plugins/Experimental/Water/Source/Runtime/Private/WaterMeshComponent.cpp
+++ b/Engine/Plugins/Experimental/Water/Source/Runtime/Private/WaterMeshComponent.cpp
@@ -501,52 +501,17 @@ void UWaterMeshComponent::RebuildWaterMesh(float InTileSize, const FIntPoint& In
 		return true;
 	});
 
-	WaterQuadTree.Unlock(true);
-
-	MarkRenderStateDirty();
-
-	// Build the far distance mesh instances. These are 8 tiles around the water quad tree
+	// Build the far distance mesh instances if needed
 	if (IsMaterialUsedWithWater(FarDistanceMaterial) && FarDistanceMeshExtent > 0.0f)
 	{
 		UsedMaterials.Add(FarDistanceMaterial);
 
-		for (int32 StreamIdx = 0; StreamIdx < FWaterTileInstanceData::NumStreams; ++StreamIdx)
-		{
-			FarDistanceWaterInstanceData.Streams[StreamIdx].SetNum(8);
-		}
+		WaterQuadTree.AddFarMesh(FarDistanceMaterial, FarDistanceMeshExtent, FarMeshHeight);
+	}
 
-		const FVector2D WaterCenter = WaterQuadTree.GetTileRegion().GetCenter();
-		const FVector2D WaterExtents = WaterQuadTree.GetTileRegion().GetExtent();
-		const FVector2D WaterSize = WaterQuadTree.GetTileRegion().GetSize();
-		const FVector2D TileOffets[] = { {-1.0f, 1.0f}, {0.0f, 1.0f}, {1.0f, 1.0f}, {1.0f, 0.0f}, {1.0f, -1.0f}, {0.0f, -1.0f}, {-1.0f, -1.0f}, {-1.0f, 0.0f} };
-		for (int32 i = 0; i < 8; i++)
-		{
-			const FVector2D TilePos = WaterCenter + TileOffets[i] * (WaterExtents + 0.5f * FarDistanceMeshExtent);
-			FVector2D TileScale;
-			TileScale.X = (TileOffets[i].X == 0.0f) ? WaterSize.X : FarDistanceMeshExtent;
-			TileScale.Y = (TileOffets[i].Y == 0.0f) ? WaterSize.Y : FarDistanceMeshExtent;
-
-			// Build instance data
-			FVector4 InstanceData[FWaterTileInstanceData::NumStreams];
-			InstanceData[0] = FVector4(TilePos, FVector2D(FarMeshHeight, 0.0f));
-			InstanceData[1] = FVector4(FVector2D::ZeroVector, TileScale);
-#if WITH_WATER_SELECTION_SUPPORT
-			InstanceData[2] = FHitProxyId::InvisibleHitProxyId.GetColor().ReinterpretAsLinear();
-#endif // WITH_WATER_SELECTION_SUPPORT
+	WaterQuadTree.Unlock(true);
 
-			for (int32 StreamIdx = 0; StreamIdx < FWaterTileInstanceData::NumStreams; ++StreamIdx)
-			{
-				FarDistanceWaterInstanceData.Streams[StreamIdx][i] = InstanceData[StreamIdx];
-			}
-		}
-	}
-	else
-	{
-		for(int32 StreamIdx = 0; StreamIdx < FWaterTileInstanceData::NumStreams; ++StreamIdx)
-		{
-			FarDistanceWaterInstanceData.Streams[StreamIdx].Empty();
-		}
-	}
+	MarkRenderStateDirty();
 }
 
 void UWaterMeshComponent::Update()
diff --git a/Engine/Plugins/Experimental/Water/Source/Runtime/Private/WaterMeshSceneProxy.cpp b/Engine/Plugins/Experimental/Water/Source/Runtime/Private/WaterMeshSceneProxy.cpp
--- a/Engine/Plugins/Experimental/Water/Source/Runtime/Private/WaterMeshSceneProxy.cpp
+++ b/Engine/Plugins/Experimental/Water/Source/Runtime/Private/WaterMeshSceneProxy.cpp
@@ -117,11 +117,7 @@ FWaterMeshSceneProxy::FWaterMeshSceneProxy(UWaterMeshComponent* Component)
 
 	WaterMeshUserDataBuffers = new WaterMeshUserDataBuffersType(WaterInstanceDataBuffers);
 
-	// Far distance mesh
-	FarDistanceWaterInstanceData = Component->GetFarDistanceInstanceData();
-	FarDistanceMaterial = Component->FarDistanceMaterial;
-
-	FarDistanceMaterialIndex = WaterQuadTree.BuildMaterialIndices(FarDistanceMaterial);
+	WaterQuadTree.BuildMaterialIndices();
 }
 
 FWaterMeshSceneProxy::~FWaterMeshSceneProxy()
@@ -178,9 +174,6 @@ void FWaterMeshSceneProxy::GetDynamicMeshElements(const TArray<const FSceneView*
 	}
 
 	const int32 NumBuckets = WaterQuadTree.GetWaterMaterials().Num() * DensityCount;
-	const int32 NumFarInstances = FarDistanceWaterInstanceData.Streams[0].Num();
-	const bool bHasFarWaterMesh = FarDistanceMaterial && NumFarInstances > 0;
-	const int32 FarBucketIndex = FarDistanceMaterialIndex * DensityCount + DensityCount - 1;
 
 	TArray<FWaterQuadTree::FTraversalOutput, TInlineAllocator<4>> WaterInstanceDataPerView;
 
@@ -248,14 +241,6 @@ void FWaterMeshSceneProxy::GetDynamicMeshElements(const TArray<const FSceneView*
 			WaterQuadTree.BuildWaterTileInstanceData(TraversalDesc, WaterInstanceData);
 
 			HistoricalMaxViewInstanceCount = FMath::Max(HistoricalMaxViewInstanceCount, WaterInstanceData.InstanceCount);
-
-			// Add far distance mesh data to the instance tile data to have it render instanced together with the qater quad tree and possibly merge if it has the same material
-			if (bHasFarWaterMesh)
-			{
-				check(FarDistanceMaterialIndex != INDEX_NONE);
-				WaterInstanceData.BucketInstanceCounts[FarBucketIndex] += NumFarInstances;
-				WaterInstanceData.InstanceCount += NumFarInstances;
-			}
 		}
 	}
 
@@ -396,19 +381,6 @@ void FWaterMeshSceneProxy::GetDynamicMeshElements(const TArray<const FSceneView*
 					WaterInstanceDataBuffers->GetBufferMemory(StreamIdx)[WriteIndex] = Data.Data[StreamIdx];
 				}
 			}
-
-			if (bHasFarWaterMesh)
-			{
-				check(FarDistanceWaterInstanceData.Streams[1].Num() == NumFarInstances);
-				const int32 WriteStartOffset = WaterInstanceData.BucketInstanceCounts[FarBucketIndex];
-
-				for (int32 StreamIdx = 0; StreamIdx < WaterInstanceDataBuffersType::NumBuffers; ++StreamIdx)
-				{
-					TArrayView<FVector4f> BufferMemory = WaterInstanceDataBuffers->GetBufferMemory(StreamIdx);
-					check(WriteStartOffset + NumFarInstances <= BufferMemory.Num());
-					FMemory::Memcpy(BufferMemory.GetData() + WriteStartOffset, FarDistanceWaterInstanceData.Streams[StreamIdx].GetData(), NumFarInstances * sizeof(FVector4f));
-				}
-			}
 		}
 	}
 
diff --git a/Engine/Plugins/Experimental/Water/Source/Runtime/Private/WaterQuadTree.cpp b/Engine/Plugins/Experimental/Water/Source/Runtime/Private/WaterQuadTree.cpp
--- a/Engine/Plugins/Experimental/Water/Source/Runtime/Private/WaterQuadTree.cpp
+++ b/Engine/Plugins/Experimental/Water/Source/Runtime/Private/WaterQuadTree.cpp
@@ -419,6 +419,8 @@ void FWaterQuadTree::InitTree(const FBox2D& InBounds, float InTileSize, FIntPoin
 	ensure(InExtentInTiles.X > 0);
 	ensure(InExtentInTiles.Y > 0);
 
+	FarMeshData.Clear();
+
 	// Maximum number of allocated leaf nodes for this config
 	MaxLeafCount = InExtentInTiles.X*InExtentInTiles.Y*4;
 	LeafSize = InTileSize;
@@ -529,12 +531,53 @@ void FWaterQuadTree::AddLake(const TArray<FVector2D>& InPoly, const FBox& InLake
 	AddLakeRecursive(InPoly, LakeBounds, FVector2D(InLakeBounds.Min.Z, InLakeBounds.Max.Z), true, TreeDepth * 2, InWaterBodyIndex);
 }
 
-int32 FWaterQuadTree::BuildMaterialIndices(UMaterialInterface* FarDistanceMaterial)
+void FWaterQuadTree::AddFarMesh(const UMaterialInterface* InFarMeshMaterial, double InFarDistanceMeshExtent, double InFarDistanceMeshHeight)
+{
+	// Checking for not being read only here to keep things consistent with the other Add functions. In reality the FarMesh isn't added to the QuadTree itself, so it could technically be done whenever.
+	ensure(!bIsReadOnly);
+	ensure(InFarMeshMaterial);
+
+	// Early out when there would be no far mesh rendering anyway
+	if (InFarMeshMaterial == nullptr || InFarDistanceMeshExtent == 0.0)
+	{
+		return;
+	}
+
+	// Far mesh is always 8 tiles around the quadtree region (marked as Q in diagram below) 
+	//  _ _ _
+	// |_|_|_|
+	// |_|Q|_|
+	// |_|_|_|
+	FarMeshData.InstanceData.SetNum(8);
+	FarMeshData.Material = InFarMeshMaterial;
+
+	const FVector2D WaterCenter = GetTileRegion().GetCenter();
+	const FVector2D WaterExtents = GetTileRegion().GetExtent();
+	const FVector2D WaterSize = GetTileRegion().GetSize();
+	const FVector2D TileOffets[] = { {-1.0, 1.0}, {0.0, 1.0}, {1.0, 1.0}, {1.0, 0.0}, {1.0, -1.0}, {0.0, -1.0}, {-1.0, -1.0}, {-1.0, 0.0} };
+
+	for (int32 i = 0; i < 8; i++)
+	{
+		const FVector2D TilePos = WaterCenter + TileOffets[i] * (WaterExtents + 0.5 * InFarDistanceMeshExtent);
+		FVector2D TileScale;
+		TileScale.X = (TileOffets[i].X == 0.0) ? WaterSize.X : InFarDistanceMeshExtent;
+		TileScale.Y = (TileOffets[i].Y == 0.0) ? WaterSize.Y : InFarDistanceMeshExtent;
+
+		// Build instance data
+		FarMeshData.InstanceData[i].Data[0] = FVector4(TilePos, FVector2D(InFarDistanceMeshHeight, 0.0));
+		FarMeshData.InstanceData[i].Data[1] = FVector4(FVector2D::ZeroVector, TileScale);
+#if WITH_WATER_SELECTION_SUPPORT
+		FarMeshData.InstanceData[i].Data[2] = FHitProxyId::InvisibleHitProxyId.GetColor().ReinterpretAsLinear();
+#endif // WITH_WATER_SELECTION_SUPPORT
+	}
+}
+
+void FWaterQuadTree::BuildMaterialIndices()
 {
 	int32 NextIdx = 0;
 	TMap<FMaterialRenderProxy*, int32> MatToIdxMap;
 
-	auto GetMatIdx = [&NextIdx, &MatToIdxMap](UMaterialInterface* Material)
+	auto GetMatIdx = [&NextIdx, &MatToIdxMap](const UMaterialInterface* Material)
 	{
 		if (!Material)
 		{
@@ -558,7 +601,8 @@ int32 FWaterQuadTree::BuildMaterialIndices(UMaterialInterface* FarDistanceMateri
 		Data.RiverToOceanMaterialIndex = GetMatIdx(Data.RiverToOceanMaterial);
 	}
 
-	const int32 FarDistMatIdx = GetMatIdx(FarDistanceMaterial);
+	// Special case handling for Far Mesh
+	FarMeshData.MaterialIndex = GetMatIdx(FarMeshData.Material);
 
 	WaterMaterials.Empty(MatToIdxMap.Num());
 	WaterMaterials.AddUninitialized(MatToIdxMap.Num());
@@ -567,8 +611,6 @@ int32 FWaterQuadTree::BuildMaterialIndices(UMaterialInterface* FarDistanceMateri
 	{
 		WaterMaterials[It->Value] = It->Key;
 	}
-
-	return FarDistMatIdx;
 }
 
 void FWaterQuadTree::BuildWaterTileInstanceData(const FTraversalDesc& InTraversalDesc, FTraversalOutput& Output) const
@@ -576,6 +618,26 @@ void FWaterQuadTree::BuildWaterTileInstanceData(const FTraversalDesc& InTraversa
 	TRACE_CPUPROFILER_EVENT_SCOPE(BuildWaterTileInstanceData);
 	check(bIsReadOnly);
 	NodeData.Nodes[0].SelectLOD(NodeData, TreeDepth, InTraversalDesc, Output);
+
+	// Append Far Mesh tiles
+	if (FarMeshData.InstanceData.Num() > 0 && FarMeshData.MaterialIndex != INDEX_NONE)
+	{
+		const int32 FarMeshTileCount = FarMeshData.InstanceData.Num();
+		ensure(FarMeshTileCount == 8);
+
+		// Bucket index calculation is MaterialIndex*DensityCount+CurrentDensity. Since far mesh doesn't have any Density(aka LOD) steps and should render only using a 2 triangle quad, we enter it only into the last Density bucket (this always corresponds to a 2 triangle quad). 
+		const int32 BucketIndex = FarMeshData.MaterialIndex * InTraversalDesc.DensityCount + (InTraversalDesc.DensityCount - 1);
+		Output.BucketInstanceCounts[BucketIndex] += FarMeshTileCount;
+		Output.InstanceCount += FarMeshTileCount;
+
+		const int32 StartIndex = Output.StagingInstanceData.AddUninitialized(FarMeshTileCount);
+
+		for (int32 i = 0; i < FarMeshTileCount; i++)
+		{
+			Output.StagingInstanceData[StartIndex + i] = FarMeshData.InstanceData[i];
+			Output.StagingInstanceData[StartIndex + i].BucketIndex = BucketIndex;
+		}
+	}
 }
 
 bool FWaterQuadTree::QueryInterpolatedTileBaseHeightAtLocation(const FVector2D& InWorldLocationXY, float& OutHeight) const
diff --git a/Engine/Plugins/Experimental/Water/Source/Runtime/Public/WaterMeshComponent.h b/Engine/Plugins/Experimental/Water/Source/Runtime/Public/WaterMeshComponent.h
--- a/Engine/Plugins/Experimental/Water/Source/Runtime/Public/WaterMeshComponent.h
+++ b/Engine/Plugins/Experimental/Water/Source/Runtime/Public/WaterMeshComponent.h
@@ -48,8 +48,6 @@ public:
 
 	const TSet<UMaterialInterface*>& GetUsedMaterialsSet() const { return UsedMaterials; }
 
-	FWaterTileInstanceData GetFarDistanceInstanceData() const { return FarDistanceWaterInstanceData; }
-
 	void MarkWaterMeshGridDirty() { bNeedsRebuild = true; }
 
 	int32 GetTessellationFactor() const { return FMath::Clamp(TessellationFactor + TessFactorBiasScalability, 1, 12); }
@@ -128,9 +126,6 @@ private:
 	/** Cached CVarWaterMeshLODScaleBias to detect changes in scalability */
 	float LODScaleBiasScalability = 0.0f;
 
-	/** Instance data for the far distance mesh */
-	FWaterTileInstanceData FarDistanceWaterInstanceData;
-
 	/** Highest tessellation factor of a water tile. Max number of verts on the side of a tile will be (2^TessellationFactor)+1)  */
 	UPROPERTY(EditAnywhere, Category = Mesh, meta = (ClampMin = "1", ClampMax = "12"))
 	int32 TessellationFactor = 6;
diff --git a/Engine/Plugins/Experimental/Water/Source/Runtime/Public/WaterMeshSceneProxy.h b/Engine/Plugins/Experimental/Water/Source/Runtime/Public/WaterMeshSceneProxy.h
--- a/Engine/Plugins/Experimental/Water/Source/Runtime/Public/WaterMeshSceneProxy.h
+++ b/Engine/Plugins/Experimental/Water/Source/Runtime/Public/WaterMeshSceneProxy.h
@@ -80,9 +80,4 @@ private:
 	int32 ForceCollapseDensityLevel = TNumericLimits<int32>::Max();
 
 	mutable int32 HistoricalMaxViewInstanceCount = 0;
-
-	/** Instance data for the far distance mesh */
-	int32 FarDistanceMaterialIndex;
-	FWaterTileInstanceData FarDistanceWaterInstanceData;
-	UMaterialInterface* FarDistanceMaterial = nullptr;
 };
diff --git a/Engine/Plugins/Experimental/Water/Source/Runtime/Public/WaterQuadTree.h b/Engine/Plugins/Experimental/Water/Source/Runtime/Public/WaterQuadTree.h
--- a/Engine/Plugins/Experimental/Water/Source/Runtime/Public/WaterQuadTree.h
+++ b/Engine/Plugins/Experimental/Water/Source/Runtime/Public/WaterQuadTree.h
@@ -64,24 +64,16 @@ struct FWaterBodyRenderData
 	}
 };
 
-/** Instance data to feed the GPU, each array is bound as an instanced vertex stream */
-struct FWaterTileInstanceData 
-{
-	static constexpr int32 NumStreams = WITH_WATER_SELECTION_SUPPORT ? 3 : 2;
-
-	TArray<FVector4f> Streams[NumStreams];
-};
-
 struct FWaterQuadTree 
 {
 	enum { INVALID_PARENT = 0xFFFFFFF };
 
-	typedef TArray<FWaterTileInstanceData, TInlineAllocator<12>> FInstanceDataBuckets;
+	static constexpr int32 NumStreams = WITH_WATER_SELECTION_SUPPORT ? 3 : 2;
 
 	struct FStagingInstanceData
 	{
 		int32 BucketIndex;
-		FVector4f Data[FWaterTileInstanceData::NumStreams];
+		FVector4f Data[NumStreams];
 	};
 
 	/** Output of the quadtree when asking to traverse it for visible water tiles */
@@ -141,8 +133,11 @@ struct FWaterQuadTree
 	/** Add Lake by giving a closed spline that represents the lake */
 	void AddLake(const TArray<FVector2D>& InPoly, const FBox& InLakeBounds, uint32 InWaterBodyIndex);
 
-	/** Assign an index to each material and return the index of FarDistanceMaterial if not null. */
-	int32 BuildMaterialIndices(UMaterialInterface* FarDistanceMaterial);
+	/** Add an automatically generated mesh (8 quads) skirt around the main water quadtree which extends out InFarDistanceMeshExtent, is placed at Z value InFarDistanceMeshHeight and is rendered using InFarMeshMaterial */
+	void AddFarMesh(const UMaterialInterface* InFarMeshMaterial, double InFarDistanceMeshExtent, double InFarDistanceMeshHeight);
+
+	/** Assign an index to each material */
+	void BuildMaterialIndices();
 
 	/** Add water body render data to this tree. Returns the index in the array. Use this index to add tiles with this water body to the tree, see AddWaterTilesInsideBounds(..) */
 	uint32 AddWaterBodyRenderData(const FWaterBodyRenderData& InWaterBodyRenderData) { return NodeData.WaterBodyRenderData.Add(InWaterBodyRenderData); }
@@ -276,6 +271,27 @@ private:
 
 	TArray<FMaterialRenderProxy*> WaterMaterials;
 
+	/** Contains everything needed to render the far mesh. This data lives outside the quadtree structure itself */
+	struct FFarMeshData
+	{
+		/** Stored data for rendering all far mesh instances (8 or them if used). This is built oncce when the quadtree is built and then read each time the quadtree is traversed */
+		TArray<FStagingInstanceData> InstanceData;
+
+		/** Material for the Far Distance Mesh, its material render proxy will be cached in WaterMaterials when BuildMaterialIndices is called */
+		const UMaterialInterface* Material = nullptr;
+
+		/** Cached material index */
+		int16 MaterialIndex = INDEX_NONE;
+
+		void Clear()
+		{
+			InstanceData.Empty();
+			Material = nullptr;
+			MaterialIndex = INDEX_NONE;
+		}
+
+	} FarMeshData;
+
 	/** If true, the tree may not change */
 	bool bIsReadOnly = true;
 };
