[Water] Changed condition to early out of far mesh generation to include negative extents.

#rnx

diff --git a/Engine/Plugins/Experimental/Water/Source/Runtime/Private/WaterQuadTree.cpp b/Engine/Plugins/Experimental/Water/Source/Runtime/Private/WaterQuadTree.cpp
--- a/Engine/Plugins/Experimental/Water/Source/Runtime/Private/WaterQuadTree.cpp
+++ b/Engine/Plugins/Experimental/Water/Source/Runtime/Private/WaterQuadTree.cpp
@@ -538,7 +538,7 @@ void FWaterQuadTree::AddFarMesh(const UMaterialInterface* InFarMeshMaterial, dou
 	ensure(InFarMeshMaterial);
 
 	// Early out when there would be no far mesh rendering anyway
-	if (InFarMeshMaterial == nullptr || InFarDistanceMeshExtent == 0.0)
+	if (InFarMeshMaterial == nullptr || InFarDistanceMeshExtent <= 0.0)
 	{
 		return;
 	}
