* Added validation for invalid custom water bodies (no water mesh specified)
* Improved water body map check error messages
* Deprecated bOverrideWaterMesh (WaterMeshOverride is enough to override the water mesh, no need for an additional boolean)

#tests editor

diff --git a/Engine/Plugins/Experimental/Water/Source/Runtime/Private/WaterBodyActor.cpp b/Engine/Plugins/Experimental/Water/Source/Runtime/Private/WaterBodyActor.cpp
--- a/Engine/Plugins/Experimental/Water/Source/Runtime/Private/WaterBodyActor.cpp
+++ b/Engine/Plugins/Experimental/Water/Source/Runtime/Private/WaterBodyActor.cpp
@@ -512,7 +512,6 @@ void AWaterBody::DeprecateData()
 		WaterBodyComponent->LayerWeightmapSettings = LayerWeightmapSettings_DEPRECATED;
 		WaterBodyComponent->bAffectsLandscape = bAffectsLandscape_DEPRECATED;
 		WaterBodyComponent->bGenerateCollisions = bGenerateCollisions_DEPRECATED;
-		WaterBodyComponent->bOverrideWaterMesh = bOverrideWaterMesh_DEPRECATED;
 		WaterBodyComponent->WaterMeshOverride = WaterMeshOverride_DEPRECATED;
 		WaterBodyComponent->OverlapMaterialPriority = OverlapMaterialPriority_DEPRECATED;
 		WaterBodyComponent->CollisionProfileName = CollisionProfileName_DEPRECATED;
diff --git a/Engine/Plugins/Experimental/Water/Source/Runtime/Private/WaterBodyComponent.cpp b/Engine/Plugins/Experimental/Water/Source/Runtime/Private/WaterBodyComponent.cpp
--- a/Engine/Plugins/Experimental/Water/Source/Runtime/Private/WaterBodyComponent.cpp
+++ b/Engine/Plugins/Experimental/Water/Source/Runtime/Private/WaterBodyComponent.cpp
@@ -812,18 +812,6 @@ void UWaterBodyComponent::PostEditImport()
 	RequestGPUWaveDataUpdate();
 }
 
-bool UWaterBodyComponent::CanEditChange(const FProperty* InProperty) const
-{
-	if (InProperty)
-	{
-		if (InProperty->GetFName() == GET_MEMBER_NAME_CHECKED(UWaterBodyComponent, WaterMeshOverride))
-		{
-			return bOverrideWaterMesh || GetWaterBodyType() == EWaterBodyType::Transition;
-		}
-	}
-	return Super::CanEditChange(InProperty);
-}
-
 void UWaterBodyComponent::OnPostEditChangeProperty(FPropertyChangedEvent& PropertyChangedEvent, bool& bShapeOrPositionChanged, bool& bWeightmapSettingsChanged)
 {
 	const FName PropertyName = PropertyChangedEvent.GetPropertyName();
@@ -859,28 +847,25 @@ TArray<TSharedRef<FTokenizedMessage>> UWaterBodyComponent::CheckWaterBodyStatus(
 	TArray<TSharedRef<FTokenizedMessage>> Result;
 	if (!IsTemplate())
 	{
-		if (const UWorld* World = GetWorld())
+		if (const UWaterSubsystem* WaterSubsystem = UWaterSubsystem::GetWaterSubsystem(GetWorld()))
 		{
-			if (const UWaterSubsystem* WaterSubsystem = UWaterSubsystem::GetWaterSubsystem(World))
+			if (AffectsWaterMesh() && (WaterSubsystem->GetWaterZoneActor() == nullptr))
 			{
-				if (AffectsWaterMesh() && (WaterSubsystem->GetWaterZoneActor() == nullptr))
-				{
-					Result.Add(FTokenizedMessage::Create(
-						EMessageSeverity::Error,
-						FText::Format(
-							LOCTEXT("MapCheck_Message_MissingWaterZone", "Water body {0} requires a WaterZone actorr to be rendered. Please add one to the map. "),
-							FText::FromString(GetWaterBodyActor()->GetActorLabel()))));
-				}
+				Result.Add(FTokenizedMessage::Create(EMessageSeverity::Error)
+					->AddToken(FUObjectToken::Create(this))
+					->AddToken(FTextToken::Create(FText::Format(
+						LOCTEXT("MapCheck_Message_MissingWaterZone", "Water body {0} requires a WaterZone actor to be rendered. Please add one to the map. "),
+						FText::FromString(GetWaterBodyActor()->GetActorLabel())))));
 			}
+		}
 
-			if (AffectsLandscape() && FindLandscape() == nullptr)
-			{
-				Result.Add(FTokenizedMessage::Create(
-					EMessageSeverity::Error,
-					FText::Format(
-						LOCTEXT("MapCheck_Message_MissingLandscape", "Water body {0} requires a Landscape to be rendered. Please add one to the map. "),
-						FText::FromString(GetWaterBodyActor()->GetActorLabel()))));
-			}
+		if (AffectsLandscape() && (FindLandscape() == nullptr))
+		{
+			Result.Add(FTokenizedMessage::Create(EMessageSeverity::Error)
+				->AddToken(FUObjectToken::Create(this))
+				->AddToken(FTextToken::Create(FText::Format(
+					LOCTEXT("MapCheck_Message_MissingLandscape", "Water body {0} requires a Landscape to be rendered. Please add one to the map. "),
+					FText::FromString(GetWaterBodyActor()->GetActorLabel())))));
 		}
 	}
 	return Result;
@@ -1175,6 +1160,14 @@ void UWaterBodyComponent::PostLoad()
 {
 	Super::PostLoad();
 
+#if WITH_EDITORONLY_DATA
+	// WaterMeshOverride is now enough to override the water mesh (bOverrideWaterMesh_DEPRECATED was superfluous), so make sure to discard WaterMeshOverride (except on custom water bodies) when the boolean wasn't set :
+	if (!bOverrideWaterMesh_DEPRECATED && (WaterMeshOverride != nullptr) && (GetWaterBodyType() != EWaterBodyType::Transition))
+	{
+		WaterMeshOverride = nullptr;
+	}
+#endif
+
 #if WITH_EDITOR
 	RegisterOnUpdateWavesData(GetWaterWaves(), /* bRegister = */true);
 #endif
diff --git a/Engine/Plugins/Experimental/Water/Source/Runtime/Private/WaterBodyCustomComponent.cpp b/Engine/Plugins/Experimental/Water/Source/Runtime/Private/WaterBodyCustomComponent.cpp
--- a/Engine/Plugins/Experimental/Water/Source/Runtime/Private/WaterBodyCustomComponent.cpp
+++ b/Engine/Plugins/Experimental/Water/Source/Runtime/Private/WaterBodyCustomComponent.cpp
@@ -4,11 +4,14 @@
 #include "Components/StaticMeshComponent.h"
 #include "Materials/MaterialInstanceDynamic.h"
 #include "WaterSubsystem.h"
+#include "Misc/UObjectToken.h"
 
 #if WITH_EDITOR
 #include "WaterIconHelper.h"
 #endif
 
+#define LOCTEXT_NAMESPACE "Water"
+
 // ----------------------------------------------------------------------------------
 
 UWaterBodyCustomComponent::UWaterBodyCustomComponent(const FObjectInitializer& ObjectInitializer)
@@ -112,3 +115,25 @@ void UWaterBodyCustomComponent::BeginUpdateWaterBody()
 	}
 }
 
+#if WITH_EDITOR
+TArray<TSharedRef<FTokenizedMessage>> UWaterBodyCustomComponent::CheckWaterBodyStatus() const
+{
+	TArray<TSharedRef<FTokenizedMessage>> StatusMessages = Super::CheckWaterBodyStatus();
+
+	if (!IsTemplate())
+	{
+		if (WaterMeshOverride == nullptr)
+		{
+			StatusMessages.Add(FTokenizedMessage::Create(EMessageSeverity::Error)
+				->AddToken(FUObjectToken::Create(this))
+				->AddToken(FTextToken::Create(FText::Format(
+					LOCTEXT("MapCheck_Message_MissingCustomWaterMesh", "Custom water body {0} requires a static mesh to be rendered. Please set WaterMeshOverride to a valid static mesh. "),
+					FText::FromString(GetWaterBodyActor()->GetActorLabel())))));
+		}
+	}
+
+	return StatusMessages;
+}
+#endif // WITH_EDITOR
+
+#undef LOCTEXT_NAMESPACE
diff --git a/Engine/Plugins/Experimental/Water/Source/Runtime/Public/WaterBodyComponent.h b/Engine/Plugins/Experimental/Water/Source/Runtime/Public/WaterBodyComponent.h
--- a/Engine/Plugins/Experimental/Water/Source/Runtime/Public/WaterBodyComponent.h
+++ b/Engine/Plugins/Experimental/Water/Source/Runtime/Public/WaterBodyComponent.h
@@ -129,7 +129,7 @@ public:
 	/** Returns the additional water height added to the body (For internal use. Please use AWaterBodyOcean instead.) */
 	virtual float GetHeightOffset() const { return 0.f; }
 
-	/** Sets the water mesh (when bOverrideWaterMesh is true of the for custom water body actors) */
+	/** Sets a static mesh to use as a replacement for the water mesh (for water bodies that are being rendered by the water mesh) */
 	void SetWaterMeshOverride(UStaticMesh* InMesh) { WaterMeshOverride = InMesh; }
 
 	/** Returns River to lake transition material instance (For internal use. Please use AWaterBodyRiver instead.) */
@@ -158,7 +158,7 @@ public:
 	FName GetCollisionProfileName() const { return CollisionProfileName; }
 
 	/** Returns water mesh override */
-	UStaticMesh* GetWaterMeshOverride() const { return (bOverrideWaterMesh || GetWaterBodyType() == EWaterBodyType::Transition) ? WaterMeshOverride : nullptr; }
+	UStaticMesh* GetWaterMeshOverride() const { return WaterMeshOverride; }
 
 	/** Returns water material */
 	UFUNCTION(BlueprintCallable, Category = Rendering)
@@ -370,7 +370,6 @@ protected:
 	virtual bool MoveComponentImpl(const FVector& Delta, const FQuat& NewRotation, bool bSweep, FHitResult* Hit = NULL, EMoveComponentFlags MoveFlags = MOVECOMP_NoFlags, ETeleportType Teleport = ETeleportType::None) override;
 
 #if WITH_EDITOR
-	virtual bool CanEditChange(const FProperty* InProperty) const override;
 	virtual void PreEditUndo() override;
 	virtual void PostEditUndo() override;
 	virtual void PostEditImport() override;
@@ -458,11 +457,7 @@ protected:
 	UPROPERTY(Transient, DuplicateTransient, NonTransactional, VisibleAnywhere, BlueprintReadOnly, Category = Water)
 	int32 WaterBodyIndex = INDEX_NONE;
 
-	// TODO [jonathan.bard] : make sure override water mesh works for all types 
 	UPROPERTY(Category = Rendering, EditAnywhere, BlueprintReadOnly)
-	bool bOverrideWaterMesh;
-
-	UPROPERTY(Category = Rendering, EditAnywhere, AdvancedDisplay, BlueprintReadOnly)
 	UStaticMesh* WaterMeshOverride;
 
 	/** Higher number is higher priority. If two water bodies overlap and they don't have a transition material specified, this will be used to determine which water body to use the material from. Valid range is -8192 to 8191 */
@@ -500,4 +495,9 @@ protected:
 	// The navigation area class that will be generated on nav mesh
 	UPROPERTY(EditAnywhere, BlueprintReadWrite, Category = Navigation, meta = (EditCondition = "bCanAffectNavigation && bGenerateCollisions"))
 	TSubclassOf<UNavAreaBase> WaterNavAreaClass;
+
+#if WITH_EDITORONLY_DATA
+	UPROPERTY()
+	bool bOverrideWaterMesh_DEPRECATED;
+#endif
 };
diff --git a/Engine/Plugins/Experimental/Water/Source/Runtime/Public/WaterBodyCustomComponent.h b/Engine/Plugins/Experimental/Water/Source/Runtime/Public/WaterBodyCustomComponent.h
--- a/Engine/Plugins/Experimental/Water/Source/Runtime/Public/WaterBodyCustomComponent.h
+++ b/Engine/Plugins/Experimental/Water/Source/Runtime/Public/WaterBodyCustomComponent.h
@@ -27,6 +27,11 @@ protected:
 	virtual void BeginUpdateWaterBody() override;
 	virtual void OnUpdateBody(bool bWithExclusionVolumes) override;
 
+#if WITH_EDITOR
+	virtual TArray<TSharedRef<FTokenizedMessage>> CheckWaterBodyStatus() const override;
+#endif // WITH_EDITOR
+
+protected:
 	UPROPERTY(NonPIEDuplicateTransient)
 	UStaticMeshComponent* MeshComp;
 };
\ No newline at end of file
