Changed WaterSubsystem to avoid keeping a reference to a WaterZoneActor to properly support unloading a sublevel containing water data.
Also modified WaterSubsystem to return a WaterZoneActor based on a level.
This fixes the problem where a water body MID would reference a the water velocity texture outered to a water zone actor part of a sublevel being unloaded.

diff --git a/Engine/Plugins/Experimental/Water/Source/Runtime/Private/WaterBodyComponent.cpp b/Engine/Plugins/Experimental/Water/Source/Runtime/Private/WaterBodyComponent.cpp
--- a/Engine/Plugins/Experimental/Water/Source/Runtime/Private/WaterBodyComponent.cpp
+++ b/Engine/Plugins/Experimental/Water/Source/Runtime/Private/WaterBodyComponent.cpp
@@ -1213,7 +1213,8 @@ bool UWaterBodyComponent::SetDynamicParametersOnMID(UMaterialInstanceDynamic* In
 	InMID->SetScalarParameterValue(FixedWaterDepthName, GetConstantDepth());
 	InMID->SetVectorParameterValue(FixedVelocityName, GetConstantVelocity());
 
-	if (AWaterZone* WaterZone = WaterSubsystem->GetWaterZoneActor())
+	// Use WaterZone actor of the same level
+	if (AWaterZone* WaterZone = WaterSubsystem->GetWaterZoneActor(GetTypedOuter<ULevel>()))
 	{
 		InMID->SetTextureParameterValue(WaterVelocityAndHeightName, WaterZone->WaterVelocityTexture);
 
diff --git a/Engine/Plugins/Experimental/Water/Source/Runtime/Private/WaterSubsystem.cpp b/Engine/Plugins/Experimental/Water/Source/Runtime/Private/WaterSubsystem.cpp
--- a/Engine/Plugins/Experimental/Water/Source/Runtime/Private/WaterSubsystem.cpp
+++ b/Engine/Plugins/Experimental/Water/Source/Runtime/Private/WaterSubsystem.cpp
@@ -202,9 +202,9 @@ void UWaterSubsystem::Tick(float DeltaTime)
 	SetMPCTime(MPCTime, PrevWorldTimeSeconds);
 	PrevWorldTimeSeconds = MPCTime;
 
-	if (WaterZoneActor)
+	if (AWaterZone* LocalWaterZoneActor = GetWaterZoneActor())
 	{
-		WaterZoneActor->Update();
+		LocalWaterZoneActor->Update();
 	}
 
 	if (!bUnderWaterForAudio && CachedDepthUnderwater > 0.0f)
@@ -479,14 +479,24 @@ void UWaterSubsystem::SetOceanFloodHeight(float InFloodHeight)
 	}
 }
 
-AWaterZone* UWaterSubsystem::GetWaterZoneActor() const
+AWaterZone* UWaterSubsystem::GetWaterZoneActor(ULevel* InPreferredOuterLevel) const
 {
 	if (UWorld* World = GetWorld())
 	{
+		// First, try to find the AWaterZone outered to the provided level
+		if (InPreferredOuterLevel)
+		{
+			for (TActorIterator<AWaterZone> It(World); It; ++It)
+			{
+				if (It && (It->GetTypedOuter<ULevel>() == InPreferredOuterLevel))
+				{
+					return *It;
+				}
+			}
+		}
 		// #todo_water: this assumes only one water zone actor right now.  In the future we may need to associate a water mesh actor with a water body more directly
 		TActorIterator<AWaterZone> It(World);
-		WaterZoneActor = It ? *It : nullptr;
-		return WaterZoneActor;
+		return It ? *It : nullptr;
 	}
 
 	return nullptr;
diff --git a/Engine/Plugins/Experimental/Water/Source/Runtime/Public/WaterSubsystem.h b/Engine/Plugins/Experimental/Water/Source/Runtime/Public/WaterSubsystem.h
--- a/Engine/Plugins/Experimental/Water/Source/Runtime/Public/WaterSubsystem.h
+++ b/Engine/Plugins/Experimental/Water/Source/Runtime/Public/WaterSubsystem.h
@@ -101,7 +101,7 @@ public:
 
 	FWaterBodyManager WaterBodyManager;
 
-	AWaterZone* GetWaterZoneActor() const;
+	AWaterZone* GetWaterZoneActor(ULevel* InPreferredOuterLevel = nullptr) const;
 
 	ABuoyancyManager* GetBuoyancyManager() const { return BuoyancyManager; }
 
@@ -220,8 +220,6 @@ public:
 	UStaticMesh* DefaultLakeMesh;
 
 private:
-	UPROPERTY()
-	mutable AWaterZone* WaterZoneActor;
 
 	TWeakObjectPtr<UWaterBodyComponent> OceanBodyComponent;
 
