Buoyancy Sim:
Enable debug draw in the async code path.
Add cvars to control debug points.

diff --git a/Engine/Plugins/Experimental/Water/Source/Runtime/Private/BuoyancyComponent.cpp b/Engine/Plugins/Experimental/Water/Source/Runtime/Private/BuoyancyComponent.cpp
--- a/Engine/Plugins/Experimental/Water/Source/Runtime/Private/BuoyancyComponent.cpp
+++ b/Engine/Plugins/Experimental/Water/Source/Runtime/Private/BuoyancyComponent.cpp
@@ -18,6 +18,18 @@ TAutoConsoleVariable<int32> CVarWaterDebugBuoyancy(
 	TEXT("Enable debug drawing for water interactions."),
 	ECVF_Default);
 
+TAutoConsoleVariable<int32> CVarWaterBuoyancyDebugPoints(
+	TEXT("r.Water.BuoyancyDebugPoints"),
+	10,
+	TEXT("Number of points in one dimension for buoyancy debug."),
+	ECVF_Default);
+
+TAutoConsoleVariable<int32> CVarWaterBuoyancyDebugSize(
+	TEXT("r.Water.BuoyancyDebugSize"),
+	1000,
+	TEXT("Side length of square for buoyancy debug."),
+	ECVF_Default);
+
 TAutoConsoleVariable<int32> CVarWaterUseSplineKeyOptimization(
 	TEXT("r.Water.UseSplineKeyOptimization"),
 	1,
@@ -468,13 +480,17 @@ int32 UBuoyancyComponent::UpdatePontoons(float DeltaTime, float ForwardSpeed, fl
 #if ENABLE_DRAW_DEBUG
 		if (CVarWaterDebugBuoyancy.GetValueOnAnyThread())
 		{
+			const float NumPoints = CVarWaterBuoyancyDebugPoints.GetValueOnAnyThread();
+			const float Size = CVarWaterBuoyancyDebugSize.GetValueOnAnyThread();
+			const float StartOffset = NumPoints * 0.5f;
+			const float Scale = Size / NumPoints;
 			TMap<const UWaterBodyComponent*, float> DebugSplineKeyMap;
 			TMap<const UWaterBodyComponent*, float> DebugSplineSegmentsMap;
-			for (int i = 0; i < 10; ++i)
+			for (int i = 0; i < NumPoints; ++i)
 			{
-				for (int j = 0; j < 10; ++j)
+				for (int j = 0; j < NumPoints; ++j)
 				{
-					FVector Location = PrimitiveComponent->GetComponentLocation() + (FVector::RightVector * (i - 5) * 90) + (FVector::ForwardVector * (j - 5) * 90);
+					FVector Location = PrimitiveComponent->GetComponentLocation() + (FVector::RightVector * (i - StartOffset) * Scale) + (FVector::ForwardVector * (j - StartOffset) * Scale);
 					GetWaterSplineKey(Location, DebugSplineKeyMap, DebugSplineSegmentsMap);
 					FVector Point(Location.X, Location.Y, GetWaterHeight(Location - FVector::UpVector * 200.f, DebugSplineKeyMap, GetOwner()->GetActorLocation().Z));
 					DrawDebugPoint(GetWorld(), Point, 5.f, IsOverlappingWaterBody() ? FColor::Green : FColor::Red, false, -1.f, 0);
diff --git a/Engine/Plugins/Experimental/Water/Source/Runtime/Public/BuoyancyComponentSimulation.h b/Engine/Plugins/Experimental/Water/Source/Runtime/Public/BuoyancyComponentSimulation.h
--- a/Engine/Plugins/Experimental/Water/Source/Runtime/Public/BuoyancyComponentSimulation.h
+++ b/Engine/Plugins/Experimental/Water/Source/Runtime/Public/BuoyancyComponentSimulation.h
@@ -4,6 +4,7 @@
 
 #include "WaterBodyTypes.h"
 #include "BuoyancyTypes.h"
+#include "Chaos/DebugDrawQueue.h"
 #include "Chaos/Particle/ParticleUtilities.h"
 #include "Kismet/KismetMathLibrary.h"
 
@@ -413,12 +414,6 @@ public:
 					Pontoon.ImmersionDepth = 0.f;
 				}
 
-				//#if ENABLE_DRAW_DEBUG
-				//					if (CVarWaterDebugBuoyancy.GetValueOnAnyThread())
-				//					{
-				//						DrawDebugSphere(GetWorld(), Pontoon.CenterLocation, Pontoon.Radius, 16, FColor::Red, false, -1.f, 0, 1.f);
-				//					}
-				//#endif
 				ComputeBuoyancy(BuoyancyData, Pontoon, State.ForwardSpeedKmh, State.LinearVelocity.Z);
 
 				if (Pontoon.bIsInWater && !bPrevIsInWater)
@@ -437,23 +432,52 @@ public:
 			}
 			PontoonIndex++;
 		}
-		//#if ENABLE_DRAW_DEBUG
-		//			if (CVarWaterDebugBuoyancy.GetValueOnAnyThread())
-		//			{
-		//				TMap<const AWaterBody*, float> DebugSplineKeyMap;
-		//				TMap<const AWaterBody*, float> DebugSplineSegmentsMap;
-		//				for (int i = 0; i < 30; ++i)
-		//				{
-		//					for (int j = 0; j < 30; ++j)
-		//					{
-		//						FVector Location = PrimitiveComponent->GetComponentLocation() + (FVector::RightVector * (i - 15) * 30) + (FVector::ForwardVector * (j - 15) * 30);
-		//						GetWaterSplineKey(Location, DebugSplineKeyMap, DebugSplineSegmentsMap);
-		//						FVector Point(Location.X, Location.Y, GetWaterHeight(Location - FVector::UpVector * 200.f, DebugSplineKeyMap, GetOwner()->GetActorLocation().Z));
-		//						DrawDebugPoint(GetWorld(), Point, 5.f, IsOverlappingWaterBody() ? FColor::Green : FColor::Red, false, -1.f, 0);
-		//					}
-		//				}
-		//			}
-		//#endif
+#if ENABLE_DRAW_DEBUG
+		if (CVarWaterDebugBuoyancy.GetValueOnAnyThread())
+		{
+			const float NumPoints = CVarWaterBuoyancyDebugPoints.GetValueOnAnyThread();
+			const float Size = CVarWaterBuoyancyDebugSize.GetValueOnAnyThread();
+			const float StartOffset = NumPoints * 0.5f;
+			const float Scale = Size / NumPoints;
+			TMap<const FSolverSafeWaterBodyData*, float> DebugSplineKeyMap;
+			Chaos::FDebugDrawQueue::GetInstance().SetEnabled(true);
+			for (int i = 0; i < NumPoints; ++i)
+			{
+				for (int j = 0; j < NumPoints; ++j)
+				{
+					FVector Location = TParticleUtilities::GetCoMWorldPosition(Body) + (FVector::RightVector * (i - StartOffset) * Scale) + (FVector::ForwardVector * (j - StartOffset) * Scale);
+
+					FSphericalPontoon DummyPontoon;
+					for (const FSolverSafeWaterBodyData* WaterBody : SolverWaterBodies)
+					{
+						if (WaterBody && WaterBody->WaterBodyType == EWaterBodyType::River)
+						{
+							float SplineInputKey;
+							//if (CVarWaterUseSplineKeyOptimization.GetValueOnAnyThread())
+							//{
+							//	SplineInputKey = GetWaterSplineKeyFast(Location, WaterBody, OutSegmentMap);
+							//}
+							//else
+							{
+								SplineInputKey = WaterBody->WaterSpline.FindInputKeyClosestToWorldLocation(Location);
+							}
+							DummyPontoon.SolverSplineInputKeys.Add(WaterBody, SplineInputKey);
+						}
+					}
+					
+
+					const float WaterHeight = GetWaterHeight(SolverWaterBodies, Location - FVector::UpVector * 100.f, Aux.SmoothedWorldTimeSeconds, DummyPontoon.SolverSplineInputKeys, /*DefaultHeight*/-100000.f, DummyPontoon.SolverWaterBody, DummyPontoon.WaterDepth, DummyPontoon.WaterPlaneLocation, DummyPontoon.WaterPlaneNormal, DummyPontoon.WaterSurfacePosition, DummyPontoon.WaterVelocity, DummyPontoon.WaterBodyIndex);
+					const FVector DebugPoint(Location.X, Location.Y, WaterHeight);
+					Chaos::FDebugDrawQueue::GetInstance().DrawDebugPoint(DebugPoint, FColor::Green, false, -1.f, 0, 5.f);
+				}
+			}
+
+			for (FSphericalPontoon& Pontoon : Aux.Pontoons)
+			{
+				Chaos::FDebugDrawQueue::GetInstance().DrawDebugSphere(Pontoon.CenterLocation, Pontoon.Radius, 16, Pontoon.bIsInWater ? FColor::Blue : FColor::Red, false, -1.f, 0, 1.f);
+			}
+		}
+#endif
 		
 		State.bIsInWaterBody = State.NumPontoonsInWater > 0;
 	}
diff --git a/Engine/Plugins/Experimental/Water/Source/Runtime/Public/BuoyancyTypes.h b/Engine/Plugins/Experimental/Water/Source/Runtime/Public/BuoyancyTypes.h
--- a/Engine/Plugins/Experimental/Water/Source/Runtime/Public/BuoyancyTypes.h
+++ b/Engine/Plugins/Experimental/Water/Source/Runtime/Public/BuoyancyTypes.h
@@ -11,6 +11,10 @@
 class UBuoyancyComponent;
 class AWaterBody;
 
+extern WATER_API TAutoConsoleVariable<int32> CVarWaterDebugBuoyancy;
+extern WATER_API TAutoConsoleVariable<int32> CVarWaterBuoyancyDebugPoints;
+extern WATER_API TAutoConsoleVariable<int32> CVarWaterBuoyancyDebugSize;
+
 USTRUCT(Blueprintable)
 struct FSphericalPontoon
 {
