Include missed header.

diff --git a/Engine/Plugins/Experimental/Water/Source/Runtime/Public/BuoyancyComponentSimulation.h b/Engine/Plugins/Experimental/Water/Source/Runtime/Public/BuoyancyComponentSimulation.h
--- a/Engine/Plugins/Experimental/Water/Source/Runtime/Public/BuoyancyComponentSimulation.h
+++ b/Engine/Plugins/Experimental/Water/Source/Runtime/Public/BuoyancyComponentSimulation.h
@@ -5,6 +5,7 @@
 #include "WaterBodyTypes.h"
 #include "BuoyancyTypes.h"
 #include "Chaos/DebugDrawQueue.h"
+#include "DrawDebugHelpers.h"
 #include "Chaos/Particle/ParticleUtilities.h"
 #include "Kismet/KismetMathLibrary.h"
 
