Adding LWC support for the water plugin mesh rendering

#rnx

diff --git a/Engine/Plugins/Experimental/Water/Shaders/Private/WaterMeshVertexFactory.ush b/Engine/Plugins/Experimental/Water/Shaders/Private/WaterMeshVertexFactory.ush
--- a/Engine/Plugins/Experimental/Water/Shaders/Private/WaterMeshVertexFactory.ush
+++ b/Engine/Plugins/Experimental/Water/Shaders/Private/WaterMeshVertexFactory.ush
@@ -67,9 +67,8 @@ struct FPositionAndNormalOnlyVertexFactoryInput
  */
 struct FVertexFactoryIntermediates
 {
-	float4 MorphedWorldPos;
-	float3 MorphedWorldPosRaw;
-	float3 OriginalWorldPos;
+	FLWCVector3 MorphedWorldPos;
+	float3 MorphedTranslatedWorldPos;
 	
 	uint WaveParamIndex;
 	/** Cached primitive and instance data */
@@ -158,7 +157,8 @@ FMaterialVertexParameters GetMaterialVertexParameters(FVertexFactoryInput Input,
 	UNROLL
 	for(int CoordinateIndex = 0; CoordinateIndex < NUM_MATERIAL_TEXCOORDS_VERTEX; CoordinateIndex++)
 	{
-		Result.TexCoords[CoordinateIndex] = Intermediates.MorphedWorldPosRaw.xy;
+		// LWC_TODO: Using the LWC Offset here will result in seams at the tile edges.
+		Result.TexCoords[CoordinateIndex] = Intermediates.MorphedWorldPos.Offset.xy;
 	}
 #endif  //NUM_MATERIAL_TEXCOORDS_VERTEX
 
@@ -169,37 +169,36 @@ FMaterialVertexParameters GetMaterialVertexParameters(FVertexFactoryInput Input,
 	return Result;
 }
 
-float3 MorphWorldPosition(float3 OriginalWorldPos, float2 InMorphOrigin, float LODLevel, float InLODScale, float2 QuadSize, float HeightLODFactor, out float LODFactor)
+float3 MorphTranslatedWorldPosition(float3 OriginalTranslatedWorldPos, float2 InMorphOriginTWS, float LODLevel, float InLODScale, float2 QuadSize, float HeightLODFactor, out float LODFactor)
 {
-	float3 WorldPos = OriginalWorldPos;
-
-	// Origin world position  
-	float3 TessWorldOrigin = LWCHackToFloat(ResolvedView.WorldCameraOrigin);
+	float3 TranslatedWorldPos = OriginalTranslatedWorldPos;
 	
-	float DistanceToVert2D = distance(WorldPos.xy, TessWorldOrigin.xy);
+	float DistanceToVert2D = distance(TranslatedWorldPos.xy, ResolvedView.TranslatedWorldCameraOrigin.xy);
 
 	LODFactor = saturate(DistanceToVert2D / (InLODScale * pow(2.0f, LODLevel)) - 1.0f);
 	LODFactor = saturate(HeightLODFactor + LODFactor);
 
 	float2 Grid4 = 4.0f * QuadSize;
 
-	// The offset is calculated with the assumption that the origin of the space is in the center of the water quad tree. Move the vertex position into water quadtree space here.
-	const float2 MorphLocalPos = WorldPos.xy - InMorphOrigin;
-	float2 Offset = frac(MorphLocalPos / Grid4) - float2(0.5f, 0.5f);
+	// InMorphOriginTWS is the corner of the node being rendered
+	const float2 MorphLocalPos = TranslatedWorldPos.xy - InMorphOriginTWS;
+
+	// Offset is like a UV offset (-0.5 to 0.5) within a quadtree node. This is the scalar used to move the vertex
+	float2 Offset = frac(MorphLocalPos / Grid4) - float2(0.5, 0.5);
 
 	// Move every other vert either positive or negative towards a neighboring vert
 	// This is the smooth continuous sliding which moved to the next LOD level
 	const float MinRadius = 0.26f; 
 	if (abs(Offset.x) < MinRadius) 
 	{
-		WorldPos.x += Offset.x * LODFactor * Grid4.x;
+		TranslatedWorldPos.x += Offset.x * LODFactor * Grid4.x;
 	}
 	if (abs(Offset.y) < MinRadius) 
 	{
-		WorldPos.y += Offset.y * LODFactor * Grid4.y;
+		TranslatedWorldPos.y += Offset.y * LODFactor * Grid4.y;
 	}
 
-	return WorldPos;
+	return TranslatedWorldPos;
 }
 
 FVertexFactoryIntermediates GetVertexFactoryIntermediates(FVertexFactoryInput Input)
@@ -220,7 +219,7 @@ FVertexFactoryIntermediates GetVertexFactoryIntermediates(FVertexFactoryInput In
 	const uint ShouldMorph = ((PackedDataChannel >> 8) & 0x1); 
 
 	// Calculate the world pos
-	Intermediates.OriginalWorldPos = float3(Input.Position.xy, 0.0f) * Scale + Translation;
+	float3 TranslatedWorldPosition = float3(Input.Position.xy, 0.0f) * Scale + Translation;
 
 	Intermediates.SceneData = VF_GPUSCENE_GET_INTERMEDIATES(Input);
 	
@@ -230,32 +229,33 @@ FVertexFactoryIntermediates GetVertexFactoryIntermediates(FVertexFactoryInput In
 		// Then this reaches 1.0, the lowest LOD will pop out and be replaced by the next LOD, starting over at 0.0. Rinse and repeat.
 		//const float HeightLODFactor = 0.0f;
 		const float HeightLODFactor = Data1.y; 
-
+	
 		const float2 QuadSize = Scale.xy / WaterVF.NumQuadsPerTileSide;
-
+	
 		float OutLODFactor = 0.0f;
-		float3 WorldPos = MorphWorldPosition(Intermediates.OriginalWorldPos, WaterVF.MorphOrigin, LODLevel, WaterVF.LODScale, QuadSize, HeightLODFactor, OutLODFactor);
-
+		float3 TempTranslatedWorldPosition = MorphTranslatedWorldPosition(TranslatedWorldPosition, Translation.xy - Scale.xy*0.5f, LODLevel, WaterVF.LODScale, QuadSize, HeightLODFactor, OutLODFactor);
+	
 		// If the vert is fully morphed, we morph it again as if it was a vert from the next LOD. This effectively means one tile can morph between 3 different LOD levels. 
 		// This is needed because the current LOD level tile will stick in to the next LOD level since we move the observer around. Make sure to not add any heightLODFactor to this, since that only applies to lowest LOD
-		if (OutLODFactor >= 1.0f)
+		const uint CanMorphTwice = ((PackedDataChannel >> 9) & 0x1); 
+		if (OutLODFactor >= 1.0f && CanMorphTwice == 1)
 		{
-			WorldPos = MorphWorldPosition(WorldPos, WaterVF.MorphOrigin, LODLevel + 1, WaterVF.LODScale, QuadSize * 2.0f, 0.0f, OutLODFactor);
+			TempTranslatedWorldPosition = MorphTranslatedWorldPosition(TempTranslatedWorldPosition, Translation.xy - Scale.xy*0.5f, LODLevel + 1, WaterVF.LODScale, QuadSize * 2.0f, 0.0f, OutLODFactor);
 		}
-
-		Intermediates.MorphedWorldPosRaw = WorldPos;
+	
+		Intermediates.MorphedTranslatedWorldPos = TempTranslatedWorldPosition;
 	}
 	else
 	{
-		Intermediates.MorphedWorldPosRaw = Intermediates.OriginalWorldPos;
+		Intermediates.MorphedTranslatedWorldPos = TranslatedWorldPosition;
 	}
 
-	Intermediates.MorphedWorldPos = float4(Intermediates.MorphedWorldPosRaw + LWCHackToFloat(ResolvedView.PreViewTranslation), 1);
+	Intermediates.MorphedWorldPos = LWCSubtract(Intermediates.MorphedTranslatedWorldPos, ResolvedView.PreViewTranslation);
 
 #if WITH_WATER_SELECTION_SUPPORT_VF
 	float SelectedValue = Input.InstanceData2.w;
 	float IsVisible = WaterVF.bRenderSelected * SelectedValue + WaterVF.bRenderUnselected * (1-SelectedValue);
-	Intermediates.MorphedWorldPos *= IsVisible;
+	Intermediates.MorphedWorldPos = LWCMultiply(Intermediates.MorphedWorldPos, IsVisible);
 #endif // WITH_WATER_SELECTION_SUPPORT_VF
 
 	return Intermediates;
@@ -287,7 +287,7 @@ half3x3 VertexFactoryGetTangentToLocal( FVertexFactoryInput Input, FVertexFactor
 // @return translated world position
 float4 VertexFactoryGetWorldPosition(FVertexFactoryInput Input, FVertexFactoryIntermediates Intermediates)
 {
-	return Intermediates.MorphedWorldPos;
+	return float4(Intermediates.MorphedTranslatedWorldPos, 1);
 }
 
 float4 VertexFactoryGetRasterizedWorldPosition(FVertexFactoryInput Input, FVertexFactoryIntermediates Intermediates, float4 InWorldPosition)
@@ -340,7 +340,7 @@ float4 VertexFactoryGetWorldPosition(FPositionOnlyVertexFactoryInput Input)
 // @return previous translated world position
 float4 VertexFactoryGetPreviousWorldPosition(FVertexFactoryInput Input, FVertexFactoryIntermediates Intermediates)
 {
-	return float4(Intermediates.MorphedWorldPosRaw + LWCHackToFloat(ResolvedView.PrevPreViewTranslation), 1);
+	return float4(LWCToFloat(LWCAdd(Intermediates.MorphedWorldPos, ResolvedView.PrevPreViewTranslation)), 1);
 }
 
 float4 VertexFactoryGetTranslatedPrimitiveVolumeBounds(FVertexFactoryInterpolantsVSToPS Interpolants)
diff --git a/Engine/Plugins/Experimental/Water/Source/Runtime/Private/WaterMeshSceneProxy.cpp b/Engine/Plugins/Experimental/Water/Source/Runtime/Private/WaterMeshSceneProxy.cpp
--- a/Engine/Plugins/Experimental/Water/Source/Runtime/Private/WaterMeshSceneProxy.cpp
+++ b/Engine/Plugins/Experimental/Water/Source/Runtime/Private/WaterMeshSceneProxy.cpp
@@ -97,7 +97,7 @@ FWaterMeshSceneProxy::FWaterMeshSceneProxy(UWaterMeshComponent* Component)
 	WaterVertexFactories.Reserve(WaterQuadTree.GetTreeDepth());
 	for (uint8 i = 0; i < WaterQuadTree.GetTreeDepth(); i++)
 	{
-		WaterVertexFactories.Add(new WaterVertexFactoryType(GetScene().GetFeatureLevel(), NumQuads, LODScale, FVector2D(WaterQuadTree.GetBounds().GetCenter())));
+		WaterVertexFactories.Add(new WaterVertexFactoryType(GetScene().GetFeatureLevel(), NumQuads, LODScale));
 		BeginInitResource(WaterVertexFactories.Last());
 
 		NumQuads /= 2;
@@ -229,6 +229,7 @@ void FWaterMeshSceneProxy::GetDynamicMeshElements(const TArray<const FSceneView*
 			TraversalDesc.ForceCollapseDensityLevel = ForceCollapseDensityLevel;
 			TraversalDesc.Frustum = View->ViewFrustum;
 			TraversalDesc.ObserverPosition = ObserverPosition;
+			TraversalDesc.PreViewTranslation = View->ViewMatrices.GetPreViewTranslation();
 			TraversalDesc.LODScale = LODScale;
 			TraversalDesc.bLODMorphingEnabled = !!CVarWaterMeshLODMorphEnabled.GetValueOnRenderThread();
 
diff --git a/Engine/Plugins/Experimental/Water/Source/Runtime/Private/WaterQuadTree.cpp b/Engine/Plugins/Experimental/Water/Source/Runtime/Private/WaterQuadTree.cpp
--- a/Engine/Plugins/Experimental/Water/Source/Runtime/Private/WaterQuadTree.cpp
+++ b/Engine/Plugins/Experimental/Water/Source/Runtime/Private/WaterQuadTree.cpp
@@ -15,7 +15,7 @@ void FWaterQuadTree::FNode::AddNodeForRender(const FNodeData& InNodeData, const
 	int32 TileDebugID = InWaterBodyRenderData.WaterBodyType;
 
 	// The base height of this tile comes either the top of the bounding box (for rivers) or the given base height (lakes and ocean)
-	float BaseHeight = InWaterBodyRenderData.IsRiver() ? Bounds.Max.Z : InWaterBodyRenderData.SurfaceBaseHeight;
+	double BaseHeight = InWaterBodyRenderData.IsRiver() ? Bounds.Max.Z : InWaterBodyRenderData.SurfaceBaseHeight;
 
 	// If there's a transition water body
 	if (TransitionWaterBodyIndex > 0)
@@ -46,30 +46,34 @@ void FWaterQuadTree::FNode::AddNodeForRender(const FNodeData& InNodeData, const
 		}
 	}
 
+	const float BaseHeightTWS = BaseHeight + InTraversalDesc.PreViewTranslation.Z;
+
 	const int32 DensityIndex = FMath::Min(InDensityLevel, InTraversalDesc.DensityCount - 1);
 	const int32 BucketIndex = MaterialIndex * InTraversalDesc.DensityCount + DensityIndex;
 	
 	++Output.BucketInstanceCounts[BucketIndex];
 
-	const FVector Position(Bounds.GetCenter());
+	const FVector TranslatedWorldPosition(Bounds.GetCenter() + InTraversalDesc.PreViewTranslation);
 	const FVector2D Scale(Bounds.GetSize());
 	FStagingInstanceData& StagingData = Output.StagingInstanceData[Output.StagingInstanceData.AddUninitialized()];
 
 	// Add the data to the bucket
 	StagingData.BucketIndex = BucketIndex;
-	StagingData.Data[0].X = Position.X;
-	StagingData.Data[0].Y = Position.Y;
-	StagingData.Data[0].Z = BaseHeight;
+	StagingData.Data[0].X = TranslatedWorldPosition.X;
+	StagingData.Data[0].Y = TranslatedWorldPosition.Y;
+	StagingData.Data[0].Z = BaseHeightTWS;
 	StagingData.Data[0].W = *(float*)&NodeWaterBodyIndex;
 
 	// Lowest LOD isn't always 0, this increases with the height distance 
 	const bool bIsLowestLOD = (InLODLevel == InTraversalDesc.LowestLOD);
 
 	// Only allow a tile to morph if it's not the last density level and not the last LOD level, sicne there is no next level to morph to
-	const uint32 bShouldMorph = (InTraversalDesc.bLODMorphingEnabled && (InDensityLevel < InTraversalDesc.DensityCount - 1) && (InLODLevel < InTraversalDesc.LODCount - 1)) ? 1 : 0;
+	const uint32 bShouldMorph = (InTraversalDesc.bLODMorphingEnabled && (DensityIndex != InTraversalDesc.DensityCount - 1)) ? 1 : 0;
+	// Tiles can morph twice to be able to morph between 3 LOD levels. Next to last density level can only morph once
+	const uint32 bCanMorphTwice = (DensityIndex < InTraversalDesc.DensityCount - 2) ? 1 : 0;
 
-	// Pack some of the data to save space. LOD level in the lower 8 bits and then bShouldMorph in the 9th bit
-	const uint32 BitPackedChannel = ((uint32)(InLODLevel) & 0xFF) | (bShouldMorph << 8);
+	// Pack some of the data to save space. LOD level in the lower 8 bits and then bShouldMorph in the 9th bit and bCanMorphTwice in the 10th bit
+	const uint32 BitPackedChannel = ((uint32)(InLODLevel) & 0xFF) | (bShouldMorph << 8) | (bCanMorphTwice << 9);
 
 	// Should morph
 	StagingData.Data[1].X = *(float*)&BitPackedChannel;
@@ -95,8 +99,8 @@ void FWaterQuadTree::FNode::AddNodeForRender(const FNodeData& InNodeData, const
 		FColor Color;
 		if (InTraversalDesc.DebugShowTypeColor)
 		{
-			static FColor WaterTypeColor[] = { FColor::Red, FColor::Green, FColor::Blue, FColor::Yellow, FColor::Purple };
-			Color = WaterTypeColor[TileDebugID];
+			//static FColor WaterTypeColor[] = { FColor::Red, FColor::Green, FColor::Blue, FColor::Yellow, FColor::Purple };
+			Color = GColorList.GetFColorByIndex(DensityIndex + 1);
 		}
 		else
 		{
diff --git a/Engine/Plugins/Experimental/Water/Source/Runtime/Public/WaterQuadTree.h b/Engine/Plugins/Experimental/Water/Source/Runtime/Public/WaterQuadTree.h
--- a/Engine/Plugins/Experimental/Water/Source/Runtime/Public/WaterQuadTree.h
+++ b/Engine/Plugins/Experimental/Water/Source/Runtime/Public/WaterQuadTree.h
@@ -24,7 +24,7 @@ struct FWaterBodyRenderData
 	UMaterialInterface* RiverToOceanMaterial = nullptr;
 
 	/** World Z position of the waterbody, this is where the tiles for this water body will be rendered*/
-	float SurfaceBaseHeight = 0.0f;
+	double SurfaceBaseHeight = 0.0;
 
 	/** Render priority. If two water bodies overlap, this will decide which water body is used for a tile */
 	int16 Priority = TNumericLimits<int16>::Min();
@@ -104,6 +104,7 @@ struct FWaterQuadTree
 		int32 ForceCollapseDensityLevel = TNumericLimits<int32>::Max();
 		float LODScale = 1.0;
 		FVector ObserverPosition = FVector::ZeroVector;
+		FVector PreViewTranslation = FVector::ZeroVector;
 		FConvexVolume Frustum;
 		bool bLODMorphingEnabled = true;
 
diff --git a/Engine/Plugins/Experimental/Water/Source/Runtime/Public/WaterVertexFactory.h b/Engine/Plugins/Experimental/Water/Source/Runtime/Public/WaterVertexFactory.h
--- a/Engine/Plugins/Experimental/Water/Source/Runtime/Public/WaterVertexFactory.h
+++ b/Engine/Plugins/Experimental/Water/Source/Runtime/Public/WaterVertexFactory.h
@@ -16,7 +16,6 @@
 BEGIN_GLOBAL_SHADER_PARAMETER_STRUCT(FWaterVertexFactoryParameters, )
 	SHADER_PARAMETER(float, LODScale)
 	SHADER_PARAMETER(int32, NumQuadsPerTileSide)
-	SHADER_PARAMETER(FVector2f, MorphOrigin)
 	SHADER_PARAMETER(int32, bRenderSelected)
 	SHADER_PARAMETER(int32, bRenderUnselected)
 END_GLOBAL_SHADER_PARAMETER_STRUCT()
@@ -158,7 +157,7 @@ public:
 	static constexpr int32 NumRenderGroups = bWithWaterSelectionSupport ? 3 : 1; // Must match EWaterMeshRenderGroupType
 	static constexpr int32 NumAdditionalVertexStreams = TWaterInstanceDataBuffers<bWithWaterSelectionSupport>::NumBuffers;
 
-	TWaterVertexFactory(ERHIFeatureLevel::Type InFeatureLevel, int32 InNumQuadsPerSide,	float InLODScale, FVector2D InMorphOrigin);
+	TWaterVertexFactory(ERHIFeatureLevel::Type InFeatureLevel, int32 InNumQuadsPerSide,	float InLODScale);
 	~TWaterVertexFactory();
 
 	/**
@@ -194,9 +193,6 @@ private:
 
 	const int32 NumQuadsPerSide = 0;
 	const float LODScale = 0.0f;
-
-	/** Center of the water quadtree. Used to ensure that the morphing algorithm works on all LOD levels */
-	const FVector2D MorphOrigin = FVector2D::ZeroVector;
 };
 
 
diff --git a/Engine/Plugins/Experimental/Water/Source/Runtime/Public/WaterVertexFactory.inl b/Engine/Plugins/Experimental/Water/Source/Runtime/Public/WaterVertexFactory.inl
--- a/Engine/Plugins/Experimental/Water/Source/Runtime/Public/WaterVertexFactory.inl
+++ b/Engine/Plugins/Experimental/Water/Source/Runtime/Public/WaterVertexFactory.inl
@@ -12,11 +12,10 @@
 // ----------------------------------------------------------------------------------
 
 template <bool bWithWaterSelectionSupport>
-TWaterVertexFactory<bWithWaterSelectionSupport>::TWaterVertexFactory(ERHIFeatureLevel::Type InFeatureLevel, int32 InNumQuadsPerSide, float InLODScale, FVector2D InMorphOrigin)
+TWaterVertexFactory<bWithWaterSelectionSupport>::TWaterVertexFactory(ERHIFeatureLevel::Type InFeatureLevel, int32 InNumQuadsPerSide, float InLODScale)
 	: FVertexFactory(InFeatureLevel)
 	, NumQuadsPerSide(InNumQuadsPerSide)
 	, LODScale(InLODScale)
-	, MorphOrigin(InMorphOrigin)
 {
 	VertexBuffer = new FWaterMeshVertexBuffer(NumQuadsPerSide);
 	IndexBuffer = new FWaterMeshIndexBuffer(NumQuadsPerSide);
@@ -119,7 +118,6 @@ void TWaterVertexFactory<bWithWaterSelectionSupport>::SetupUniformDataForGroup(E
 	FWaterVertexFactoryParameters UniformParams;
 	UniformParams.NumQuadsPerTileSide = NumQuadsPerSide;
 	UniformParams.LODScale = LODScale;
-	UniformParams.MorphOrigin = FVector2f(MorphOrigin);	// LWC_TODO: Precision loss
 	UniformParams.bRenderSelected = true;
 	UniformParams.bRenderUnselected = true;
 
