Add ability to specify which pontoons can support fx for buoyancy component. Expose bIsOverlappingWaterbody to blueprint for vfx

diff --git a/Engine/Plugins/Experimental/Water/Source/Runtime/Private/BuoyancyComponent.cpp b/Engine/Plugins/Experimental/Water/Source/Runtime/Private/BuoyancyComponent.cpp
--- a/Engine/Plugins/Experimental/Water/Source/Runtime/Private/BuoyancyComponent.cpp
+++ b/Engine/Plugins/Experimental/Water/Source/Runtime/Private/BuoyancyComponent.cpp
@@ -71,7 +71,7 @@ void UBuoyancyComponent::BeginPlay()
 	{
 		for (FSphericalPontoon& Pontoon : BuoyancyData.Pontoons)
 		{
-			if (Pontoon.CenterSocket != NAME_None)
+			if (Pontoon.CenterSocket != NAME_None && SimulatingComponent->DoesSocketExist(Pontoon.CenterSocket))
 			{
 				Pontoon.bUseCenterSocket = true;
 				Pontoon.SocketTransform = SimulatingComponent->GetSocketTransform(Pontoon.CenterSocket, RTS_Actor);
diff --git a/Engine/Plugins/Experimental/Water/Source/Runtime/Public/BuoyancyComponent.h b/Engine/Plugins/Experimental/Water/Source/Runtime/Public/BuoyancyComponent.h
--- a/Engine/Plugins/Experimental/Water/Source/Runtime/Public/BuoyancyComponent.h
+++ b/Engine/Plugins/Experimental/Water/Source/Runtime/Public/BuoyancyComponent.h
@@ -97,7 +97,7 @@ public:
 	void GetLastWaterSurfaceInfo(FVector& OutWaterPlaneLocation, FVector& OutWaterPlaneNormal,
 	FVector& OutWaterSurfacePosition, float& OutWaterDepth, int32& OutWaterBodyIdx, FVector& OutWaterVelocity);
 
-	UPROPERTY(EditDefaultsOnly, Category = Buoyancy)
+	UPROPERTY(BlueprintReadOnly, EditDefaultsOnly, Category = Buoyancy)
 	FBuoyancyData BuoyancyData;
 
 protected:
diff --git a/Engine/Plugins/Experimental/Water/Source/Runtime/Public/BuoyancyTypes.h b/Engine/Plugins/Experimental/Water/Source/Runtime/Public/BuoyancyTypes.h
--- a/Engine/Plugins/Experimental/Water/Source/Runtime/Public/BuoyancyTypes.h
+++ b/Engine/Plugins/Experimental/Water/Source/Runtime/Public/BuoyancyTypes.h
@@ -20,7 +20,7 @@ struct FSphericalPontoon
 {
 	GENERATED_BODY()
 
-	/** The socket to center this pontoon on */
+	/** The socket to center this pontoon on. Also used as the name of the pontoon for effects */
 	UPROPERTY(EditAnywhere, BlueprintReadOnly, Category = Buoyancy)
 	FName CenterSocket;
 
@@ -32,6 +32,10 @@ struct FSphericalPontoon
 	UPROPERTY(EditAnywhere, BlueprintReadOnly, Category = Buoyancy)
 	float Radius;
 
+	/** Should this pontoon be considered as a candidate location for visual/audio effects upon entering water for burst cues? To be implemented by user*/
+	UPROPERTY(EditAnywhere, BlueprintReadOnly, Category = Buoyancy)
+	bool bFXEnabled = true;
+
 	UPROPERTY(BlueprintReadOnly, Category = Buoyancy)
 	FVector LocalForce;
 
@@ -70,6 +74,9 @@ struct FSphericalPontoon
 	UPROPERTY(BlueprintReadOnly, Category = Buoyancy)
 	int32 WaterBodyIndex;
 
+	UPROPERTY(BlueprintReadOnly, Category = Buoyancy)
+	bool bIsInWater;
+	
 	FTransform SocketTransform;
 
 	TMap<const UWaterBodyComponent*, float> SplineInputKeys;
@@ -78,7 +85,6 @@ struct FSphericalPontoon
 	TMap<const FSolverSafeWaterBodyData*, float> SolverSplineInputKeys;
 	TMap<const FSolverSafeWaterBodyData*, float> SolverSplineSegments;
 
-	uint8 bIsInWater : 1;
 	uint8 bEnabled : 1;
 	uint8 bUseCenterSocket : 1;
 
@@ -103,8 +109,8 @@ struct FSphericalPontoon
 		, WaterSurfacePosition(FVector::ZeroVector)
 		, WaterVelocity(FVector::ZeroVector)
 		, WaterBodyIndex(0)
-		, SocketTransform(FTransform::Identity)
 		, bIsInWater(false)
+		, SocketTransform(FTransform::Identity)
 		, bEnabled(true)
 		, bUseCenterSocket(false)
 		, CurrentWaterBodyComponent(nullptr)
@@ -159,14 +165,12 @@ struct FSphericalPontoon
 		Ar << WaterSurfacePosition;
 		Ar << WaterVelocity;
 		Ar << WaterBodyIndex;
+		Ar << bIsInWater;
 		Ar << SocketTransform;
-		uint8 IsInWater = bIsInWater;
 		uint8 Enabled = bEnabled;
 		uint8 UseCenterSocket = bUseCenterSocket;
-		Ar << IsInWater;
 		Ar << Enabled;
 		Ar << UseCenterSocket;
-		bIsInWater = IsInWater;
 		bEnabled = Enabled;
 		bUseCenterSocket = UseCenterSocket;
 	}
