Fix BuoyancyManager code when Creating/Destroying physics states to only look at relevant  components.

diff --git a/Engine/Plugins/Experimental/Water/Source/Runtime/Private/BuoyancyManager.cpp b/Engine/Plugins/Experimental/Water/Source/Runtime/Private/BuoyancyManager.cpp
--- a/Engine/Plugins/Experimental/Water/Source/Runtime/Private/BuoyancyManager.cpp
+++ b/Engine/Plugins/Experimental/Water/Source/Runtime/Private/BuoyancyManager.cpp
@@ -33,7 +33,10 @@ void ABuoyancyManager::OnCreatePhysics(UActorComponent* Component)
 	{
 		if (UBuoyancyComponent* BuoyancyComp = OwningActor->FindComponentByClass<UBuoyancyComponent>())
 		{
-			InitializeAsyncAux(BuoyancyComp);
+			if (BuoyancyComp->GetSimulatingComponent() == Cast<UPrimitiveComponent>(Component))
+			{
+				InitializeAsyncAux(BuoyancyComp);
+			}
 		}
 	}
 #endif
@@ -46,22 +49,24 @@ void ABuoyancyManager::OnDestroyPhysics(UActorComponent* Component)
 	{
 		if (UBuoyancyComponent* BuoyancyComp = OwningActor->FindComponentByClass<UBuoyancyComponent>())
 		{
-			
-			ClearAsyncInputs(BuoyancyComp);
-			if (UPrimitiveComponent* SimulatingComp = BuoyancyComp->GetSimulatingComponent())
+			if (PhysicsInitializedSimulatingComponents.Contains(Cast<UPrimitiveComponent>(Component)))
 			{
-				if (AsyncCallback)
+				ClearAsyncInputs(BuoyancyComp);
+				if (UPrimitiveComponent* SimulatingComp = BuoyancyComp->GetSimulatingComponent())
 				{
-					if (FBodyInstance* BI = SimulatingComp->GetBodyInstance())
+					if (AsyncCallback)
 					{
-						if (auto ActorHandle = BI->ActorHandle)
+						if (FBodyInstance* BI = SimulatingComp->GetBodyInstance())
 						{
-							AsyncCallback->ClearAsyncAux_External(ActorHandle->GetGameThreadAPI().UniqueIdx());
+							if (auto ActorHandle = BI->ActorHandle)
+							{
+								AsyncCallback->ClearAsyncAux_External(ActorHandle->GetGameThreadAPI().UniqueIdx());
+							}
 						}
 					}
-				}
 
-				PhysicsInitializedSimulatingComponents.Remove(SimulatingComp);
+					PhysicsInitializedSimulatingComponents.Remove(SimulatingComp);
+				}
 			}
 		}
 	}
