Water Info Rendering:

- Decoupled the rendering of the water information texture (flow,depth,shape data) from the landscape completely.
- Implemented separate render passes leveraging the SceneRenderer interface to draw the water data at runtime.
- Water Body components now have custom scene proxies which contain the geometry describing that water body as well as a small dilation region to solve undersampling issues with the water info texture.
- Old WaterVelocityTexture has been deprecated and is no longer contains data

diff --git a/Engine/Plugins/Experimental/Water/Shaders/Private/WaterInfoFinalize.usf b/Engine/Plugins/Experimental/Water/Shaders/Private/WaterInfoFinalize.usf
new file mode 100644
--- /dev/null
+++ b/Engine/Plugins/Experimental/Water/Shaders/Private/WaterInfoFinalize.usf
@@ -0,0 +1,78 @@
+// Copyright Epic Games, Inc. All Rights Reserved.
+
+#include "/Engine/Private/Common.ush"
+#include "/Engine/Private/DeferredShadingCommon.ush"
+
+Texture2D<float4> WaterInfoTexture;
+
+float2 WaterHeightExtents;
+float GroundZMin;
+float CaptureZ;
+
+int BlurRadius;
+
+struct FWaterInfo
+{
+	float2 Velocity;
+	float NormalizedWaterZ;
+	float NormalizedGroundZ;
+};
+
+FWaterInfo MakeWaterInfo(float4 InSampleData)
+{
+	FWaterInfo Result;
+	Result.Velocity = InSampleData.rg;
+	Result.NormalizedWaterZ = InSampleData.b;
+	Result.NormalizedGroundZ = InSampleData.a;
+	return Result;
+}
+
+float2 VelocityBlur(float Radius, uint2 TextureCoordinates)
+{
+	float2 Accum = 0;
+
+	float Neighbors = 0;
+
+	int2 CurrentTextureCoordinates = int2(TextureCoordinates);
+
+	for (int X = -Radius; X <= Radius; X++)
+	{
+		for(int Y=-Radius; Y <= Radius; Y++)
+		{
+			CurrentTextureCoordinates = TextureCoordinates + int2(X, Y);
+
+			const FWaterInfo WaterInfo = MakeWaterInfo(WaterInfoTexture.Load(int3(CurrentTextureCoordinates, 0)));
+			
+			const float WaterZMin = WaterHeightExtents.x;
+			const float WaterZMax = WaterHeightExtents.y;
+
+			const float WaterZ = WaterInfo.NormalizedWaterZ * (WaterZMax - WaterZMin) + WaterZMin;
+			const float GroundZ = WaterInfo.NormalizedGroundZ * (WaterZMax - GroundZMin) + GroundZMin;
+
+			if (WaterZ > GroundZ)
+			{
+				Accum.xy += WaterInfo.Velocity;
+				Neighbors += 1;
+			}
+		}
+	}
+
+	if(Neighbors > 0)
+	{
+		Accum.xy /= Neighbors;
+	}
+	
+	return float2(Accum);
+}
+ 
+void Main(
+	in float4 SVPos : SV_POSITION,
+	out float4 OutColor : SV_Target0)
+{
+	uint2 TextureCoordinates = floor(SVPos.xy);
+	
+	const FWaterInfo WaterInfo = MakeWaterInfo(WaterInfoTexture.Load(int3(TextureCoordinates, 0)));
+	const float2 BlurredVelocity = VelocityBlur(BlurRadius, TextureCoordinates);
+
+	OutColor = float4(BlurredVelocity.rg, WaterInfo.NormalizedWaterZ, WaterInfo.NormalizedGroundZ);
+}
diff --git a/Engine/Plugins/Experimental/Water/Shaders/Private/WaterInfoMerge.usf b/Engine/Plugins/Experimental/Water/Shaders/Private/WaterInfoMerge.usf
new file mode 100644
--- /dev/null
+++ b/Engine/Plugins/Experimental/Water/Shaders/Private/WaterInfoMerge.usf
@@ -0,0 +1,69 @@
+// Copyright Epic Games, Inc. All Rights Reserved.
+
+#include "/Engine/Private/Common.ush"
+#include "/Engine/Private/DeferredShadingCommon.ush"
+
+Texture2D<float4> DepthTexture;
+Texture2D<float4> ColorTexture;
+
+float2 WaterHeightExtents;
+float GroundZMin;
+float CaptureZ;
+
+struct FWaterColorPassResult
+{
+	float2 Velocity;
+	int WaterBodyIndex;
+	float SceneDepth;
+};
+
+FWaterColorPassResult MakeWaterColorPassResult(float4 InSampleData)
+{
+	FWaterColorPassResult Result;
+	Result.Velocity = InSampleData.rg;
+	Result.WaterBodyIndex = InSampleData.b;
+	Result.SceneDepth = InSampleData.a;
+	return Result;
+}
+
+float DecodeFloatRGB(float3 Enc)
+{
+	float3 Encoding = float3(1.0, 1.0 / 255.0, 1.0 / 65025.0);
+	return 1 - dot(Enc, Encoding);
+}
+
+void Main(
+	in float4 SVPos : SV_POSITION,
+	out float4 OutColor : SV_Target0
+	)
+{
+	uint2 TextureCoordinates = floor(SVPos.xy);
+	
+    const float GroundDeviceDepth = DecodeFloatRGB(DepthTexture.Load(int3(TextureCoordinates, 0)).rgb);
+	const float GroundSceneDepth = ConvertFromDeviceZ(GroundDeviceDepth);
+    const FWaterColorPassResult WaterInfo = MakeWaterColorPassResult(ColorTexture.Load(int3(TextureCoordinates, 0)));
+	const float WaterSceneDepth = WaterInfo.SceneDepth;
+	float WaterZ = CaptureZ - WaterSceneDepth;
+
+	float GroundZ = CaptureZ - GroundSceneDepth;
+
+	// If this is a dilated region and it is above the ground we must push it down to prevent artifacts
+	if (WaterInfo.WaterBodyIndex == -1 && WaterZ > GroundZ)
+	{
+		WaterZ = GroundZ - 32.f;
+	}
+
+	const float WaterZMin = WaterHeightExtents.x;
+	const float WaterZMax = WaterHeightExtents.y;
+
+	// If the GroundDepth is further than the far plane of the projection, set the GroundZ to some really far away value
+	if (GroundDeviceDepth == 1.0)
+	{
+		GroundZ = GroundZMin;
+	}
+
+	float NormalizedWaterZ = saturate((WaterZ - WaterZMin) / (WaterZMax - WaterZMin));
+	float NormalizedGroundZ = saturate((GroundZ - GroundZMin) / (WaterZMax - GroundZMin));
+
+	OutColor = float4(WaterInfo.Velocity, NormalizedWaterZ, NormalizedGroundZ);
+}
diff --git a/Engine/Plugins/Experimental/Water/Source/Editor/Private/WaterBrushManager.cpp b/Engine/Plugins/Experimental/Water/Source/Editor/Private/WaterBrushManager.cpp
--- a/Engine/Plugins/Experimental/Water/Source/Editor/Private/WaterBrushManager.cpp
+++ b/Engine/Plugins/Experimental/Water/Source/Editor/Private/WaterBrushManager.cpp
@@ -114,7 +114,6 @@ void AWaterBrushManager::PostLoad()
 					FVector RTWorldLocation, RTWorldSizeVector;
 					if (DeprecateWaterLandscapeInfo(RTWorldLocation, RTWorldSizeVector))
 					{
-						WaterZoneActor->SetLandscapeInfo(RTWorldLocation, RTWorldSizeVector);
 						SetMPCParams();
 					}
 				}
@@ -132,7 +131,6 @@ void AWaterBrushManager::PostLoad()
 					FVector RTWorldLocation, RTWorldSizeVector;
 					if (DeprecateWaterLandscapeInfo(RTWorldLocation, RTWorldSizeVector))
 					{
-						WaterZoneActor->SetLandscapeInfo(RTWorldLocation, RTWorldSizeVector);
 						SetMPCParams();
 					}
 				}
@@ -1203,14 +1201,6 @@ void AWaterBrushManager::SetMPCParams()
 				UE_LOG(LogWaterEditor, Error, TEXT("Failed to set \"WaterClearHeight\" on Landscape MaterialParameterCollection"));
 			}
 		}
-
-		// HACK [jonathan.bard] : Quick and dirty way to move the responsability of updating the Water MPC to the water mesh actor (because AWaterBrushManager doesn't exist on client builds, this 
-		//  information needs to live somewhere, at least until we get rid of the water MPC and the water texture/params are handled outside of the water landscape brush) :
-		TActorIterator<AWaterZone> It(World);
-		if (AWaterZone* WaterZoneActor = It ? *It : nullptr)
-		{
-			WaterZoneActor->SetLandscapeInfo(RTWorldLocation, RTWorldSizeVector);
-		}
 	}
 }
 
diff --git a/Engine/Plugins/Experimental/Water/Source/Editor/Private/WaterEditorSubsystem.cpp b/Engine/Plugins/Experimental/Water/Source/Editor/Private/WaterEditorSubsystem.cpp
--- a/Engine/Plugins/Experimental/Water/Source/Editor/Private/WaterEditorSubsystem.cpp
+++ b/Engine/Plugins/Experimental/Water/Source/Editor/Private/WaterEditorSubsystem.cpp
@@ -130,35 +130,6 @@ void UWaterEditorSubsystem::Deinitialize()
 	Super::Deinitialize();
 }
 
-void UWaterEditorSubsystem::UpdateWaterTextures(
-	UWorld* World, 
-	UTextureRenderTarget2D* SourceVelocityTarget, 
-	UTexture2D*& OutWaterVelocityTexture)
-{
-	TActorIterator<AWaterZone> WaterZoneActorIt(World);
-	AWaterZone* ZoneActor = WaterZoneActorIt ? *WaterZoneActorIt : nullptr;
-	if (ZoneActor)
-	{
-		if (SourceVelocityTarget)
-		{
-			UTexture2D* PreviousTexture = ZoneActor->WaterVelocityTexture;
-			UpdateSingleTexture(ZoneActor->WaterVelocityTexture, SourceVelocityTarget, ZoneActor, TEXT("WaterVelocityTexture"));
-
-			// The water bodies' material instances are referencing the water velocity texture so they need to be in sync : 
-			if (ZoneActor->WaterVelocityTexture != PreviousTexture)
-			{
-				UWaterSubsystem::ForEachWaterBodyComponent(World, [this](UWaterBodyComponent* WaterBodyComponent)
-				{
-					WaterBodyComponent->UpdateMaterialInstances();
-					return true;
-				});
-			}
-
-			OutWaterVelocityTexture = ZoneActor->WaterVelocityTexture;
-		}
-	}
-}
-
 void UWaterEditorSubsystem::RegisterWaterActorSprite(UClass* InClass, UTexture2D* Texture)
 {
 	WaterActorSprites.Add(InClass, Texture);
diff --git a/Engine/Plugins/Experimental/Water/Source/Editor/Private/WaterLandscapeBrush.cpp b/Engine/Plugins/Experimental/Water/Source/Editor/Private/WaterLandscapeBrush.cpp
--- a/Engine/Plugins/Experimental/Water/Source/Editor/Private/WaterLandscapeBrush.cpp
+++ b/Engine/Plugins/Experimental/Water/Source/Editor/Private/WaterLandscapeBrush.cpp
@@ -203,7 +203,7 @@ void AWaterLandscapeBrush::OnActorChanged(AActor* Actor, bool bWeightmapSettings
 	{
 		if (UWaterSubsystem* WaterSubsystem = UWaterSubsystem::GetWaterSubsystem(GetWorld()))
 		{
-			WaterSubsystem->MarkAllWaterMeshesForRebuild();
+			WaterSubsystem->MarkAllWaterZonesForRebuild();
 		}
 	}
 }
@@ -225,7 +225,7 @@ void AWaterLandscapeBrush::OnActorsAffectingLandscapeChanged()
 	RequestLandscapeUpdate();
 	if (UWaterSubsystem* WaterSubsystem = UWaterSubsystem::GetWaterSubsystem(GetWorld()))
 	{
-		WaterSubsystem->MarkAllWaterMeshesForRebuild();
+		WaterSubsystem->MarkAllWaterZonesForRebuild();
 	}
 
 	MarkRenderTargetsDirty();
@@ -489,26 +489,7 @@ void AWaterLandscapeBrush::SetTargetLandscape(ALandscape* InTargetLandscape)
 
 void AWaterLandscapeBrush::OnFullHeightmapRenderDone(UTextureRenderTarget2D* InHeightmapRenderTarget)
 {
-	if (bRenderTargetsDirty)
-	{
-		FScopedDurationTimeLogger DurationTimer(TEXT("Water Texture Update Time"));
-
-		UTextureRenderTarget2D* VelocityRenderTarget = nullptr;
-		BlueprintGetRenderTargets(InHeightmapRenderTarget, VelocityRenderTarget);
-
-		UTexture2D* WaterVelocityTexture = nullptr;
-		GEditor->GetEditorSubsystem<UWaterEditorSubsystem>()->UpdateWaterTextures(
-			GetWorld(),
-			VelocityRenderTarget,
-			WaterVelocityTexture);
-
-		if (WaterVelocityTexture)
-		{
-			BlueprintOnRenderTargetTexturesUpdated(WaterVelocityTexture);
-		}
-
-		bRenderTargetsDirty = false;
-	}
+	bRenderTargetsDirty = false;
 }
 
 void AWaterLandscapeBrush::SetOwningLandscape(ALandscape* InOwningLandscape)
diff --git a/Engine/Plugins/Experimental/Water/Source/Editor/Public/WaterEditorSubsystem.h b/Engine/Plugins/Experimental/Water/Source/Editor/Public/WaterEditorSubsystem.h
--- a/Engine/Plugins/Experimental/Water/Source/Editor/Public/WaterEditorSubsystem.h
+++ b/Engine/Plugins/Experimental/Water/Source/Editor/Public/WaterEditorSubsystem.h
@@ -27,11 +27,6 @@ public:
 	virtual void Initialize(FSubsystemCollectionBase& Collection) override;
 	virtual void Deinitialize() override;
 
-	virtual void UpdateWaterTextures(
-		UWorld* World,
-		UTextureRenderTarget2D* SourceVelocityTarget,
-		UTexture2D*& OutWaterVelocityTexture);
-
 	UMaterialParameterCollection* GetLandscapeMaterialParameterCollection() const {	return LandscapeMaterialParameterCollection; }
 
 	// ~ IWaterEditorServices interface
diff --git a/Engine/Plugins/Experimental/Water/Source/Runtime/Private/WaterBodyActor.cpp b/Engine/Plugins/Experimental/Water/Source/Runtime/Private/WaterBodyActor.cpp
--- a/Engine/Plugins/Experimental/Water/Source/Runtime/Private/WaterBodyActor.cpp
+++ b/Engine/Plugins/Experimental/Water/Source/Runtime/Private/WaterBodyActor.cpp
@@ -605,6 +605,10 @@ void AWaterBody::PostRegisterAllComponents()
 		{
 			WaterBodyComponent->WaterBodyIndex = Manager->AddWaterBodyComponent(WaterBodyComponent);
 			WaterBodyIndex = WaterBodyComponent->WaterBodyIndex;
+
+			// By this point, the water body will have already generated its scene proxy. Since we have just updated the water body index, we will
+			// need to regenerate the scene proxy so that its vertices contain the updated index.
+			WaterBodyComponent->MarkRenderStateDirty();
 		}
 
 		// At this point, the water mesh actor should be ready and we can setup the MID accordingly : 
diff --git a/Engine/Plugins/Experimental/Water/Source/Runtime/Private/WaterBodyComponent.cpp b/Engine/Plugins/Experimental/Water/Source/Runtime/Private/WaterBodyComponent.cpp
--- a/Engine/Plugins/Experimental/Water/Source/Runtime/Private/WaterBodyComponent.cpp
+++ b/Engine/Plugins/Experimental/Water/Source/Runtime/Private/WaterBodyComponent.cpp
@@ -33,6 +33,7 @@
 #include "Misc/UObjectToken.h"
 #include "Logging/MessageLog.h"
 #include "Logging/TokenizedMessage.h"
+#include "WaterBodySceneProxy.h"
 
 #if WITH_EDITOR
 #include "WaterIconHelper.h"
@@ -49,6 +50,7 @@ DECLARE_CYCLE_STAT(TEXT("WaterBody_ComputeWaterDepth"), STAT_WaterBody_ComputeLo
 DECLARE_CYCLE_STAT(TEXT("WaterBody_ComputeWaterDepth"), STAT_WaterBody_ComputeNormal, STATGROUP_Water);
 DECLARE_CYCLE_STAT(TEXT("WaterBody_ComputeLandscapeDepth"), STAT_WaterBody_ComputeLandscapeDepth, STATGROUP_Water);
 DECLARE_CYCLE_STAT(TEXT("WaterBody_ComputeWaveHeight"), STAT_WaterBody_ComputeWaveHeight, STATGROUP_Water);
+
 // ----------------------------------------------------------------------------------
 
 TAutoConsoleVariable<float> CVarWaterOceanFallbackDepth(
@@ -58,14 +60,17 @@ TAutoConsoleVariable<float> CVarWaterOceanFallbackDepth(
 	ECVF_Default);
 
 const FName UWaterBodyComponent::WaterBodyIndexParamName(TEXT("WaterBodyIndex"));
+const FName UWaterBodyComponent::WaterBodyZOffsetParamName(TEXT("WaterBodyZOffset"));
 const FName UWaterBodyComponent::WaterVelocityAndHeightName(TEXT("WaterVelocityAndHeight"));
 const FName UWaterBodyComponent::GlobalOceanHeightName(TEXT("GlobalOceanHeight"));
 const FName UWaterBodyComponent::FixedZHeightName(TEXT("FixedZHeight"));
-const FName UWaterBodyComponent::WaterAreaParamName(TEXT("WaterArea"));
 const FName UWaterBodyComponent::FixedVelocityName(TEXT("FixedVelocity"));
 const FName UWaterBodyComponent::FixedWaterDepthName(TEXT("FixedWaterDepth"));
-
-// ----------------------------------------------------------------------------------
+const FName UWaterBodyComponent::WaterAreaParamName(TEXT("WaterArea"));
+const FName UWaterBodyComponent::MaxFlowVelocityParamName(TEXT("MaxFlowVelocity"));
+const FName UWaterBodyComponent::WaterZMinParamName(TEXT("WaterZMin"));
+const FName UWaterBodyComponent::WaterZMaxParamName(TEXT("WaterZMax"));
+const FName UWaterBodyComponent::GroundZMinParamName(TEXT("GroundZMin"));
 
 UWaterBodyComponent::UWaterBodyComponent(const FObjectInitializer& ObjectInitializer)
 	: Super(ObjectInitializer)
@@ -75,11 +80,14 @@ UWaterBodyComponent::UWaterBodyComponent(const FObjectInitializer& ObjectInitial
 	CollisionProfileName = GetDefault<UWaterRuntimeSettings>()->GetDefaultWaterCollisionProfileName();
 
 	WaterMID = nullptr;
+	WaterInfoMID = nullptr;
 
 	TargetWaveMaskDepth = 2048.f;
 
 	bCanAffectNavigation = false;
 	bFillCollisionUnderWaterBodiesForNavmesh = false;
+
+	WaterInfoMaterial = GetDefault<UWaterRuntimeSettings>()->GetDefaultWaterInfoMaterial();
 }
 
 void UWaterBodyComponent::OnVisibilityChanged()
@@ -96,6 +104,19 @@ void UWaterBodyComponent::OnHiddenInGameChanged()
 	UpdateComponentVisibility(/* bAllowWaterMeshRebuild = */true);
 }
 
+FPrimitiveSceneProxy* UWaterBodyComponent::CreateSceneProxy()
+{
+	return new FWaterBodySceneProxy(this);
+}
+
+void UWaterBodyComponent::GetUsedMaterials(TArray<UMaterialInterface*>& OutMaterialInterfaces, bool bGetDebugMaterials) const
+{
+	if (WaterInfoMID)
+	{
+		OutMaterialInterfaces.Add(WaterInfoMID);
+	}
+}
+
 bool UWaterBodyComponent::IsFlatSurface() const
 {
 	// Lakes and oceans have surfaces aligned with the XY plane
@@ -129,6 +150,11 @@ FBox UWaterBodyComponent::GetCollisionComponentBounds() const
 	return Box;
 }
 
+FBoxSphereBounds UWaterBodyComponent::CalcBounds(const FTransform& LocalToWorld) const
+{
+	return Super::CalcBounds(LocalToWorld);
+}
+
 AWaterBody* UWaterBodyComponent::GetWaterBodyActor() const
 {
 	// If we have an Owner, it must be an AWaterBody
@@ -205,6 +231,12 @@ UMaterialInstanceDynamic* UWaterBodyComponent::GetUnderwaterPostProcessMaterialI
 	return UnderwaterPostProcessMID;
 }
 
+UMaterialInstanceDynamic* UWaterBodyComponent::GetWaterInfoMaterialInstance()
+{
+	CreateOrUpdateWaterInfoMID();
+	return WaterInfoMID;
+}
+
 void UWaterBodyComponent::SetUnderwaterPostProcessMaterial(UMaterialInterface* InMaterial)
 {
 	UnderwaterPostProcessMaterial = InMaterial;
@@ -605,6 +637,8 @@ void UWaterBodyComponent::OnRegister()
 
 	check(WaterSplineMetadata);
 
+	CreateOrUpdateWaterInfoMID();
+
 #if WITH_EDITOR
 	RegisterOnChangeWaterSplineMetadata(WaterSplineMetadata, /*bRegister = */true);
 	GetWaterSpline()->OnSplineDataChanged().AddUObject(this, &UWaterBodyComponent::OnSplineDataChanged);
@@ -653,6 +687,7 @@ EObjectFlags UWaterBodyComponent::GetTransientMIDFlags() const
 void UWaterBodyComponent::UpdateMaterialInstances()
 {
 	CreateOrUpdateWaterMID();
+	CreateOrUpdateWaterInfoMID();
 	CreateOrUpdateUnderwaterPostProcessMID();
 }
 
@@ -708,6 +743,16 @@ void UWaterBodyComponent::CreateOrUpdateUnderwaterPostProcessMID()
 	}
 }
 
+void UWaterBodyComponent::CreateOrUpdateWaterInfoMID()
+{
+	if (GetWorld())
+	{
+		WaterInfoMID = FWaterUtils::GetOrCreateTransientMID(WaterInfoMID, TEXT("WaterInfoMID"), WaterInfoMaterial, GetTransientMIDFlags());
+
+		SetDynamicParametersOnWaterInfoMID(WaterInfoMID);
+	}
+}
+
 void UWaterBodyComponent::PrepareCurrentPostProcessSettings()
 {
 	// Prepare the transient settings that are actually used by the post-process system : 
@@ -775,7 +820,7 @@ void UWaterBodyComponent::UpdateComponentVisibility(bool bAllowWaterMeshRebuild)
 			{
 				if (AWaterZone* WaterZone = WaterSubsystem->GetWaterZoneActor())
 				{
-					WaterZone->MarkWaterMeshComponentForRebuild();
+					WaterZone->MarkForRebuild(EWaterZoneRebuildFlags::All);
 				}
 	 		}
 	 	}
@@ -820,7 +865,8 @@ void UWaterBodyComponent::OnPostEditChangeProperty(FPropertyChangedEvent& Proper
 		bWeightmapSettingsChanged = true;
 	}
 	else if ((PropertyName == GET_MEMBER_NAME_CHECKED(UWaterBodyComponent, WaterMaterial)) ||
-			(PropertyName == GET_MEMBER_NAME_CHECKED(UWaterBodyComponent, UnderwaterPostProcessMaterial)))
+			(PropertyName == GET_MEMBER_NAME_CHECKED(UWaterBodyComponent, UnderwaterPostProcessMaterial)) ||
+			(PropertyName == GET_MEMBER_NAME_CHECKED(UWaterBodyComponent, WaterInfoMaterial)))
 	{
 		UpdateMaterialInstances();
 	}
@@ -832,7 +878,7 @@ void UWaterBodyComponent::OnPostEditChangeProperty(FPropertyChangedEvent& Proper
 	{
 		bShapeOrPositionChanged = true;
 	}
-	if (PropertyChangedEvent.MemberProperty && PropertyChangedEvent.MemberProperty->GetFName() == FName(TEXT("RelativeScale3D")) &&
+	else if (PropertyChangedEvent.MemberProperty && PropertyChangedEvent.MemberProperty->GetFName() == FName(TEXT("RelativeScale3D")) &&
 		PropertyName == FName(TEXT("Z")))
 	{
 		// Prevent scaling on the Z axis
@@ -1106,6 +1152,8 @@ void UWaterBodyComponent::UpdateAll(bool bShapeOrPositionChanged)
 		if (bShapeOrPositionChanged)
 		{
 			FNavigationSystem::UpdateActorAndComponentData(*WaterBodyOwner);
+
+			MarkRenderStateDirty();
 		}
 
 		UpdateComponentVisibility(/* bAllowWaterMeshRebuild = */true);
@@ -1211,21 +1259,33 @@ bool UWaterBodyComponent::SetDynamicParametersOnMID(UMaterialInstanceDynamic* In
 	InMID->SetScalarParameterValue(GlobalOceanHeightName, GlobalOceanHeight);
 	InMID->SetScalarParameterValue(FixedZHeightName, GetConstantSurfaceZ());
 	InMID->SetScalarParameterValue(FixedWaterDepthName, GetConstantDepth());
+
 	InMID->SetVectorParameterValue(FixedVelocityName, GetConstantVelocity());
 
 	// Use WaterZone actor of the same level
-	if (AWaterZone* WaterZone = WaterSubsystem->GetWaterZoneActor(GetTypedOuter<ULevel>()))
+	if (const AWaterZone* WaterZone = WaterSubsystem->GetWaterZoneActor(GetTypedOuter<ULevel>()))
 	{
-		InMID->SetTextureParameterValue(WaterVelocityAndHeightName, WaterZone->WaterVelocityTexture);
+		InMID->SetTextureParameterValue(WaterVelocityAndHeightName, WaterZone->WaterInfoTexture);
 
-		UWaterMeshComponent* WaterMeshComponent = WaterZone->GetWaterMeshComponent();
+		const UWaterMeshComponent* WaterMeshComponent = WaterZone->GetWaterMeshComponent();
 		check(WaterMeshComponent);
 
-		// LWC_TODO: precision loss
-		FLinearColor WaterArea = FLinearColor(WaterMeshComponent->RTWorldLocation);
-		WaterArea.B = WaterMeshComponent->RTWorldSizeVector.X;
-		WaterArea.A = WaterMeshComponent->RTWorldSizeVector.Y;
+		// Location should be the bottom left of the zone
+		const FVector2D ZoneExtent = WaterZone->GetZoneExtent();
+		const FVector2D WaterAreaLocation = FVector2D(WaterZone->GetActorLocation()) - (ZoneExtent / 2.f);
+
+		FLinearColor WaterArea;
+		WaterArea.R = WaterAreaLocation.X;
+		WaterArea.G = WaterAreaLocation.Y;
+		WaterArea.B = ZoneExtent.X;
+		WaterArea.A = ZoneExtent.Y;
 		InMID->SetVectorParameterValue(WaterAreaParamName, WaterArea);
+
+		const FVector2f WaterHeightExtents = WaterZone->GetWaterHeightExtents();
+		const float GroundZMin = WaterZone->GetGroundZMin();
+		InMID->SetScalarParameterValue(WaterZMinParamName, WaterHeightExtents.X);
+		InMID->SetScalarParameterValue(WaterZMaxParamName, WaterHeightExtents.Y);
+		InMID->SetScalarParameterValue(GroundZMinParamName, GroundZMin);
 	}
 
 	return true;
@@ -1247,6 +1307,20 @@ bool UWaterBodyComponent::SetDynamicParametersOnUnderwaterPostProcessMID(UMateri
 	return true;
 }
 
+bool UWaterBodyComponent::SetDynamicParametersOnWaterInfoMID(UMaterialInstanceDynamic* InMID)
+{
+	UWaterSubsystem* WaterSubsystem = UWaterSubsystem::GetWaterSubsystem(GetWorld());
+	if ((InMID == nullptr) || (WaterSubsystem == nullptr))
+	{
+		return false;
+	}
+
+	InMID->SetScalarParameterValue(WaterBodyZOffsetParamName, WaterHeightmapSettings.FalloffSettings.ZOffset);
+	InMID->SetScalarParameterValue(MaxFlowVelocityParamName, FWaterUtils::GetWaterMaxFlowVelocity(false));
+
+	return true;
+}
+
 float UWaterBodyComponent::GetWaveReferenceTime() const
 {
 	if (HasWaves())
@@ -1339,3 +1413,14 @@ UWaterWavesBase* UWaterBodyComponent::GetWaterWaves() const
 }
 
 #undef LOCTEXT_NAMESPACE
+
+AWaterZone* UWaterBodyComponent::GetWaterZone() const
+{
+	// #todo_water [roey]: Currently returns the global water zone. Eventually this function will return the specific water zone which encapsulates this water body. 
+
+	if (UWaterSubsystem* WaterSubsystem = UWaterSubsystem::GetWaterSubsystem(GetWorld()))
+	{
+		return WaterSubsystem->GetWaterZoneActor();
+	}
+	return nullptr;
+}
diff --git a/Engine/Plugins/Experimental/Water/Source/Runtime/Private/WaterBodyCustomComponent.cpp b/Engine/Plugins/Experimental/Water/Source/Runtime/Private/WaterBodyCustomComponent.cpp
--- a/Engine/Plugins/Experimental/Water/Source/Runtime/Private/WaterBodyCustomComponent.cpp
+++ b/Engine/Plugins/Experimental/Water/Source/Runtime/Private/WaterBodyCustomComponent.cpp
@@ -102,6 +102,12 @@ void UWaterBodyCustomComponent::OnUpdateBody(bool bWithExclusionVolumes)
 	MeshComp->MarkRenderStateDirty();
 }
 
+FPrimitiveSceneProxy* UWaterBodyCustomComponent::CreateSceneProxy()
+{
+	// Don't create a scene proxy for custom water body components since they don't render into the water info texture (yet)
+	return nullptr;
+}
+
 void UWaterBodyCustomComponent::BeginUpdateWaterBody()
 {
 	Super::BeginUpdateWaterBody();
diff --git a/Engine/Plugins/Experimental/Water/Source/Runtime/Private/WaterBodyLakeComponent.cpp b/Engine/Plugins/Experimental/Water/Source/Runtime/Private/WaterBodyLakeComponent.cpp
--- a/Engine/Plugins/Experimental/Water/Source/Runtime/Private/WaterBodyLakeComponent.cpp
+++ b/Engine/Plugins/Experimental/Water/Source/Runtime/Private/WaterBodyLakeComponent.cpp
@@ -44,6 +44,14 @@ TArray<UPrimitiveComponent*> UWaterBodyLakeComponent::GetStandardRenderableCompo
 	return Result;
 }
 
+FBoxSphereBounds UWaterBodyLakeComponent::CalcBounds(const FTransform& LocalToWorld) const
+{
+	FBox BoxExtent = GetWaterSpline()->GetLocalBounds().GetBox();
+	BoxExtent.Max.Z += MaxWaveHeightOffset;
+	BoxExtent.Min.Z -= GetChannelDepth();
+	return FBoxSphereBounds(BoxExtent).TransformBy(LocalToWorld);
+}
+
 void UWaterBodyLakeComponent::Reset()
 {
 	AActor* Owner = GetOwner();
diff --git a/Engine/Plugins/Experimental/Water/Source/Runtime/Private/WaterBodyOceanComponent.cpp b/Engine/Plugins/Experimental/Water/Source/Runtime/Private/WaterBodyOceanComponent.cpp
--- a/Engine/Plugins/Experimental/Water/Source/Runtime/Private/WaterBodyOceanComponent.cpp
+++ b/Engine/Plugins/Experimental/Water/Source/Runtime/Private/WaterBodyOceanComponent.cpp
@@ -8,6 +8,9 @@
 #include "WaterBodyOceanActor.h"
 #include "WaterBooleanUtils.h"
 #include "WaterSubsystem.h"
+#include "WaterZoneActor.h"
+#include "UObject/UObjectIterator.h"
+#include "UObject/FortniteMainBranchObjectVersion.h"
 
 #if WITH_EDITOR
 #include "WaterIconHelper.h"
@@ -19,6 +22,7 @@ UWaterBodyOceanComponent::UWaterBodyOceanComponent(const FObjectInitializer& Obj
 	: Super(ObjectInitializer)
 {
 	CollisionExtents = FVector(50000.f, 50000.f, 10000.f);
+	VisualExtents = FVector2D(150000.f, 150000.f);
 
 	// @todo_water : Remove these checks (Once AWaterBody is no more Blueprintable, these methods should become PURE_VIRTUAL and this class should overload them)
 	check(IsFlatSurface());
@@ -56,6 +60,16 @@ void UWaterBodyOceanComponent::SetHeightOffset(float InHeightOffset)
 	}
 }
 
+void UWaterBodyOceanComponent::SetVisualExtents(FVector2D NewExtents)
+{
+	if (VisualExtents != NewExtents)
+	{
+		VisualExtents = NewExtents;
+		MarkRenderStateDirty();
+		Modify();
+	}
+}
+
 void UWaterBodyOceanComponent::BeginUpdateWaterBody()
 {
 	Super::BeginUpdateWaterBody();
@@ -77,6 +91,10 @@ void UWaterBodyOceanComponent::OnPostEditChangeProperty(FPropertyChangedEvent& P
 		// Affects the physics shape
 		bShapeOrPositionChanged = true;
 	}
+	else if (PropertyChangedEvent.GetPropertyName() == GET_MEMBER_NAME_CHECKED(UWaterBodyOceanComponent, VisualExtents))
+	{
+		bShapeOrPositionChanged = true;
+	}
 }
 #endif
 
@@ -100,6 +118,35 @@ void UWaterBodyOceanComponent::Reset()
 	CollisionHullSets.Reset();
 }
 
+void UWaterBodyOceanComponent::PostLoad()
+{
+	Super::PostLoad();
+
+	if (GetLinkerCustomVersion(FFortniteMainBranchObjectVersion::GUID) < FFortniteMainBranchObjectVersion::WaterZonesRefactor)
+	{
+		if (UWorld* World = GetWorld())
+		{
+			// #todo_water: this assumes only one water zone actor right now.  In the future we may need to associate a water mesh actor with a water body more directly
+			for (TObjectIterator<AWaterZone> It; It; ++It)
+			{
+				AWaterZone* WaterZone = *It;
+				if (WaterZone && (WaterZone->GetWorld() == World))
+				{
+					SetVisualExtents(WaterZone->GetZoneExtent());
+					return;
+				}
+			}
+		}
+	}
+}
+
+FBoxSphereBounds UWaterBodyOceanComponent::CalcBounds(const FTransform& LocalToWorld) const
+{
+	FVector Min(FVector(-VisualExtents / 2.f, -1 * GetChannelDepth()));
+	FVector Max(FVector(VisualExtents / 2.f, 0.f));
+	return FBoxSphereBounds(FBox(Min, Max)).TransformBy(LocalToWorld);
+}
+
 void UWaterBodyOceanComponent::OnUpdateBody(bool bWithExclusionVolumes)
 {
 	AActor* OwnerActor = GetOwner();
diff --git a/Engine/Plugins/Experimental/Water/Source/Runtime/Private/WaterBodyRiverComponent.cpp b/Engine/Plugins/Experimental/Water/Source/Runtime/Private/WaterBodyRiverComponent.cpp
--- a/Engine/Plugins/Experimental/Water/Source/Runtime/Private/WaterBodyRiverComponent.cpp
+++ b/Engine/Plugins/Experimental/Water/Source/Runtime/Private/WaterBodyRiverComponent.cpp
@@ -82,6 +82,23 @@ void UWaterBodyRiverComponent::OnUpdateBody(bool bWithExclusionVolumes)
 	}
 }
 
+FBoxSphereBounds UWaterBodyRiverComponent::CalcBounds(const FTransform& LocalToWorld) const
+{
+	FBox BoxExtent(ForceInit);
+	for (const USplineMeshComponent* SplineMeshComponent : SplineMeshComponents)
+	{
+		if (SplineMeshComponent != nullptr)
+		{
+			BoxExtent += SplineMeshComponent->CalcLocalBounds().BoxExtent;
+		}
+	}
+	// Spline mesh components aren't storing our vertical bounds so account for that with the ChannelDepth parameter.
+	BoxExtent.Max.Z += MaxWaveHeightOffset;
+	BoxExtent.Min.Z -= GetChannelDepth();
+	FBoxSphereBounds Result(BoxExtent);
+	return Result.TransformBy(LocalToWorld);
+}
+
 void UWaterBodyRiverComponent::UpdateMaterialInstances()
 {
 	Super::UpdateMaterialInstances();
diff --git a/Engine/Plugins/Experimental/Water/Source/Runtime/Private/WaterBodySceneProxy.cpp b/Engine/Plugins/Experimental/Water/Source/Runtime/Private/WaterBodySceneProxy.cpp
new file mode 100644
--- /dev/null
+++ b/Engine/Plugins/Experimental/Water/Source/Runtime/Private/WaterBodySceneProxy.cpp
@@ -0,0 +1,621 @@
+// Copyright Epic Games, Inc. All Rights Reserved.
+
+#include "WaterBodySceneProxy.h"
+#include "WaterSplineMetadata.h"
+#include "WaterModule.h"
+#include "WaterUtils.h"
+#include "WaterBodyComponent.h"
+#include "WaterBodyOceanComponent.h"
+#include "WaterBodyRiverComponent.h"
+#include "WaterBodyLakeComponent.h"
+#include "WaterSplineComponent.h"
+#include "Algo/Transform.h"
+#include "MeshBoundaryLoops.h"
+#include "Curve/GeneralPolygon2.h"
+#include "Operations/MinimalHoleFiller.h"
+#include "DynamicMesh/DynamicMeshChangeTracker.h"
+#include "Operations/InsetMeshRegion.h"
+#include "ConstrainedDelaunay2.h"
+#include "Materials/MaterialInstanceDynamic.h"
+#include "SceneInterface.h"
+
+
+static TAutoConsoleVariable<int32> CVarWaterShowProxies(
+	TEXT("r.Water.ShowWaterSceneProxies"),
+	0,
+	TEXT("Allows editor visualization of water scene proxies. If the mode is set to 1 we will show only the selected water body in wireframe, if it is set to 2 we will show all in wireframe, and if it is 3 we will show all as opaque meshes"),
+	ECVF_Default);
+
+// ----------------------------------------------------------------------------------
+
+class FWaterBodyMeshSection
+{
+public:
+	FStaticMeshVertexBuffers VertexBuffers;
+	FDynamicMeshIndexBuffer32 IndexBuffer;
+	FLocalVertexFactory VertexFactory;
+
+	FWaterBodyMeshSection(ERHIFeatureLevel::Type InFeatureLevel)
+		: VertexFactory(InFeatureLevel, "FWaterBodyMeshSection")
+	{}
+};
+
+// ----------------------------------------------------------------------------------
+
+static FColor PackWaterFlow(float VelocityMagnitude, float DirectionAngle)
+{
+	check((DirectionAngle >= 0.f) && (DirectionAngle <= TWO_PI));
+
+	const float MaxVelocity = FWaterUtils::GetWaterMaxFlowVelocity(false);
+
+	float NormalizedMagnitude = FMath::Clamp(VelocityMagnitude, 0.f, MaxVelocity) / MaxVelocity;
+	float NormalizedAngle = DirectionAngle / TWO_PI;
+	float MappedMagnitude = NormalizedMagnitude * TNumericLimits<uint16>::Max();
+	float MappedAngle = NormalizedAngle * TNumericLimits<uint16>::Max();
+
+	uint16 ResultMag = (uint16)MappedMagnitude;
+	uint16 ResultAngle = (uint16)MappedAngle;
+	
+	FColor Result;
+	Result.R = (ResultMag >> 8) & 0xFF;
+	Result.G = (ResultMag >> 0) & 0xFF;
+	Result.B = (ResultAngle >> 8) & 0xFF;
+	Result.A = (ResultAngle >> 0) & 0xFF;
+
+	return Result;
+}
+
+// ----------------------------------------------------------------------------------
+
+using namespace UE::Geometry;
+
+void FWaterBodySceneProxy::AddSectionFromDynamicMesh(const FDynamicMesh3& DynamicMesh, TFunctionRef<FDynamicMeshVertex(const FVector3f&)> VertTransformFunc)
+{
+	TArray<FDynamicMeshVertex> Vertices;
+	Vertices.Reserve(DynamicMesh.VertexCount());
+
+	Algo::Transform(DynamicMesh.GetVerticesBuffer(), Vertices, VertTransformFunc);
+
+	FWaterBodyMeshSection* NewSection = new FWaterBodyMeshSection(GetScene().GetFeatureLevel());
+	NewSection->IndexBuffer.Indices.Reserve(DynamicMesh.TriangleCount() * 3);
+	for (const FIndex3i& Index : DynamicMesh.GetTrianglesBuffer())
+	{
+		NewSection->IndexBuffer.Indices.Add(Index.A);
+		NewSection->IndexBuffer.Indices.Add(Index.B);
+		NewSection->IndexBuffer.Indices.Add(Index.C);
+	}
+	
+	NewSection->VertexBuffers.InitFromDynamicVertex(&NewSection->VertexFactory, Vertices);
+	InitResources(NewSection);
+
+	Sections.Add(NewSection);
+}
+
+void FWaterBodySceneProxy::GenerateLakeMesh(UWaterBodyLakeComponent* Component)
+{
+	TRACE_CPUPROFILER_EVENT_SCOPE(GenerateLakeMesh);
+
+	const UWaterSplineComponent* SplineComp = Component->GetWaterSpline();
+	if (SplineComp->GetNumberOfSplineSegments() < 3)
+	{
+		return;
+	}
+	FPolygon2d LakePoly;
+	{
+		TArray<FVector> PolyLineVertices;
+		SplineComp->ConvertSplineToPolyLine(ESplineCoordinateSpace::Local, FMath::Square(10.f), PolyLineVertices);
+		
+		for (int32 i = 0; i < PolyLineVertices.Num() - 1; ++i) // skip the last vertex since it's the same as the first vertex
+		{
+			LakePoly.AppendVertex(FVector2D(PolyLineVertices[i]));
+		}
+	}
+	
+	FConstrainedDelaunay2d Triangulation;
+	Triangulation.FillRule = FConstrainedDelaunay2d::EFillRule::Positive;
+	Triangulation.Add(LakePoly);
+	Triangulation.Triangulate();
+
+	if (Triangulation.Triangles.Num() == 0)
+	{
+		return;
+	}
+
+	FDynamicMesh3 LakeMesh(EMeshComponents::None);
+	for (const FVector2d& Vertex : Triangulation.Vertices)
+	{
+		LakeMesh.AppendVertex(FVector3d(Vertex.X, Vertex.Y, 0));
+	}
+	for (const FIndex3i& Triangle : Triangulation.Triangles)
+	{
+		LakeMesh.AppendTriangle(Triangle);
+	}
+
+	AddSectionFromDynamicMesh(LakeMesh, [WaterBodyIndex = Component->GetWaterBodyIndex()](const FVector3f& In) {
+			FDynamicMeshVertex Result(FVector(In.X, In.Y, 0.f));
+			Result.Color = FColor::Black;
+			Result.TextureCoordinate[0].X = WaterBodyIndex;
+			return Result;
+		});
+	
+	if (Component->ShapeDilation > 0.f)
+	{
+		// Inset the mesh by -ShapeDilation to effectively expand the mesh
+		FInsetMeshRegion Inset(&LakeMesh);
+		Inset.InsetDistance = -1 * Component->ShapeDilation / 2.f;
+
+		for (const FIndex3i& Triangle : LakeMesh.GetTrianglesBuffer())
+		{
+			Inset.Triangles.Add(Triangle.A);
+			Inset.Triangles.Add(Triangle.B);
+			Inset.Triangles.Add(Triangle.C);
+		}
+		
+		if (!Inset.Apply())
+		{
+			UE_LOG(LogWater, Warning, TEXT("Failed to apply mesh inset for shape dilation (%s"), *Component->GetOwner()->GetActorNameOrLabel());
+		}
+	}
+	
+	// Push down the dilated region vertices by -ZOffset to prevent dilated regions from overwriting data from adjacent water bodies
+	AddSectionFromDynamicMesh(LakeMesh, [ZOffset = Component->GetShapeDilationZOffset()](const FVector3f& In) {
+			FDynamicMeshVertex Result(FVector(In.X, In.Y, ZOffset));
+			Result.Color = FColor::Black;
+			Result.TextureCoordinate[0].X = -1;
+			return Result;
+		});
+}
+
+void FWaterBodySceneProxy::GenerateOceanMesh(UWaterBodyOceanComponent* Component)
+{
+	TRACE_CPUPROFILER_EVENT_SCOPE(GenerateOceanMesh);
+
+	const UWaterSplineComponent* SplineComp = Component->GetWaterSpline();
+	
+	if (SplineComp->GetNumberOfSplineSegments() < 3)
+	{
+		return;
+	}
+
+	FPolygon2d Island;
+	{
+		TArray<FVector> PolyLineVertices;
+		SplineComp->ConvertSplineToPolyLine(ESplineCoordinateSpace::Local, FMath::Square(10.f), PolyLineVertices);
+		
+		// Construct a 2D polygon describing the central island
+		for (int32 i = PolyLineVertices.Num() - 2; i >= 0; --i) // skip the last vertex since it's the same as the first vertex
+		{
+			Island.AppendVertex(FVector2D(PolyLineVertices[i]));
+		}
+	}
+
+	FVector OceanLocation = Component->GetComponentLocation();
+	FVector2D OceanExtent = Component->GetVisualExtents();
+	FPolygon2d OceanBoundingPolygon = FPolygon2d::MakeRectangle(FVector2d(OceanLocation.X, OceanLocation.Y), OceanExtent.X, OceanExtent.Y);
+	FGeneralPolygon2d FinalPoly(OceanBoundingPolygon);
+	FinalPoly.AddHole(Island);
+
+	FConstrainedDelaunay2d Triangulation;
+	Triangulation.FillRule = FConstrainedDelaunay2d::EFillRule::Positive;
+	Triangulation.Add(FinalPoly);
+	Triangulation.Triangulate();
+
+	if (Triangulation.Triangles.Num() == 0)
+	{
+		return;
+	}
+
+	FDynamicMesh3 OceanMesh(EMeshComponents::None);
+	for (FVector2d Vertex : Triangulation.Vertices)
+	{
+		OceanMesh.AppendVertex(FVector3d(Vertex.X, Vertex.Y, 0));
+	}
+	for (FIndex3i Triangle : Triangulation.Triangles)
+	{
+		OceanMesh.AppendTriangle(Triangle);
+	}
+	AddSectionFromDynamicMesh(OceanMesh, [WaterBodyIndex = Component->GetWaterBodyIndex()](const FVector3f& In) {
+			FDynamicMeshVertex Result(FVector(In.X, In.Y, 0.f));
+			Result.Color = FColor::Black;
+			Result.TextureCoordinate[0].X = WaterBodyIndex;
+			return Result;
+		});
+
+	if (Component->ShapeDilation > 0.f)
+	{
+		// Inset the mesh by -ShapeDilation to effectively expand the mesh
+		FInsetMeshRegion Inset(&OceanMesh);
+		Inset.InsetDistance = -1 * Component->ShapeDilation / 2.f;
+
+		for (FIndex3i Triangle : OceanMesh.GetTrianglesBuffer())
+		{
+			Inset.Triangles.Add(Triangle.A);
+			Inset.Triangles.Add(Triangle.B);
+			Inset.Triangles.Add(Triangle.C);
+		}
+		
+		if (!Inset.Apply())
+		{
+			UE_LOG(LogWater, Warning, TEXT("Failed to apply mesh inset for shape dilation (%s"), *Component->GetOwner()->GetActorNameOrLabel());
+		}
+	}
+	
+	// Push down the dilated region vertices by -ZOffset to prevent dilated regions from overwriting data from adjacent water bodies
+	AddSectionFromDynamicMesh(OceanMesh, [ZOffset = Component->GetShapeDilationZOffset()](const FVector3f& In) {
+			FDynamicMeshVertex Result(FVector(In.X, In.Y, ZOffset));
+			Result.Color = FColor::Black;
+			Result.TextureCoordinate[0].X = -1;
+			return Result;
+		});
+}
+
+static void AddVerticesForRiverSplineStep(float DistanceAlongSpline, const UWaterBodyRiverComponent* Component, TArray<FDynamicMeshVertex>& Vertices, TArray<uint32>& Indices)
+{
+	const UWaterSplineComponent* SplineComp = Component->GetWaterSpline();
+	
+	const FVector Tangent = SplineComp->GetTangentAtDistanceAlongSpline(DistanceAlongSpline, ESplineCoordinateSpace::Local).GetSafeNormal();
+	const FVector Up = SplineComp->GetUpVectorAtDistanceAlongSpline(DistanceAlongSpline, ESplineCoordinateSpace::Local).GetSafeNormal();
+
+	const FVector Normal = FVector::CrossProduct(Tangent, Up).GetSafeNormal();
+	const FVector Pos = SplineComp->GetLocationAtDistanceAlongSpline(DistanceAlongSpline, ESplineCoordinateSpace::Local);
+
+	const float Key = SplineComp->SplineCurves.ReparamTable.Eval(DistanceAlongSpline, 0.f);
+	const float HalfWidth = Cast<UWaterSplineMetadata>(SplineComp->GetSplinePointsMetadata())->RiverWidth.Eval(Key) / 2;
+	float Velocity = Cast<UWaterSplineMetadata>(SplineComp->GetSplinePointsMetadata())->WaterVelocityScalar.Eval(Key);
+
+	// Distance from the center of the spline to place our first vertices
+	FVector OutwardDistance = Normal * HalfWidth;
+	// Prevent there being a relative height difference between the two vertices even when the spline has a slight roll to it
+	OutwardDistance.Z = 0.f;
+
+	const float DilationAmount = Component->ShapeDilation;
+	const FVector DilationOffset = Normal * DilationAmount;
+
+	FDynamicMeshVertex Left(Pos - OutwardDistance);
+	FDynamicMeshVertex Right(Pos + OutwardDistance);
+
+	FDynamicMeshVertex DilatedFarLeft(Pos - OutwardDistance - DilationOffset);
+	DilatedFarLeft.Position.Z += Component->GetShapeDilationZOffsetFar();
+	FDynamicMeshVertex DilatedLeft(Pos - OutwardDistance);
+	DilatedLeft.Position.Z += Component->GetShapeDilationZOffset();
+	FDynamicMeshVertex DilatedRight(Pos + OutwardDistance);
+	DilatedRight.Position.Z += Component->GetShapeDilationZOffset();
+	FDynamicMeshVertex DilatedFarRight(Pos + OutwardDistance + DilationOffset);
+	DilatedFarRight.Position.Z += Component->GetShapeDilationZOffsetFar();
+
+	float FlowDirection = Tangent.HeadingAngle() + FMath::DegreesToRadians(Component->GetRelativeRotation().Yaw);
+	// Convert negative angles into positive angles
+	if (FlowDirection < 0.f)
+	{
+		FlowDirection = TWO_PI + FlowDirection;
+	}
+
+	// If negative velocity, inverse the direction and change the velocity back to positive.
+	if (Velocity < 0.f)
+	{
+		Velocity *= -1.f;
+		FlowDirection = FMath::Fmod(PI + FlowDirection, TWO_PI);
+	}
+
+	const FColor EmptyFlowData = FColor(0.f);
+
+	const FColor PackedFlowData = PackWaterFlow(Velocity, FlowDirection);
+	Left.Color = PackedFlowData;
+	Right.Color = PackedFlowData;
+
+	DilatedFarLeft.Color = EmptyFlowData;
+	DilatedLeft.Color = EmptyFlowData;
+	DilatedRight.Color = EmptyFlowData;
+	DilatedFarRight.Color = EmptyFlowData;
+
+	// Embed the water body index in the vertex data so that we can distinguish between dilated and undilated regions of the texture
+	const int32 WaterBodyIndex = Component->GetWaterBodyIndex();
+	Left.TextureCoordinate[0].X = WaterBodyIndex;
+	Right.TextureCoordinate[0].X = WaterBodyIndex;
+	
+	DilatedFarLeft.TextureCoordinate[0].X = -1;
+	DilatedLeft.TextureCoordinate[0].X = -1;
+	DilatedRight.TextureCoordinate[0].X = -1;
+	DilatedFarRight.TextureCoordinate[0].X = -1;
+
+	/* River segment geometry:
+		6---7,8---9,10--11
+		| /  |  /  |  / |
+		0---1,2---3,4---5
+	*/
+	const uint32 BaseIndex = Vertices.Num();
+	Vertices.Append({ DilatedFarLeft, DilatedLeft, Left, Right, DilatedRight, DilatedFarRight });
+	// Append left dilation quad
+	Indices.Append({ BaseIndex, BaseIndex + 7, BaseIndex + 1, BaseIndex, BaseIndex + 6, BaseIndex + 7 });
+	// Append main quad
+	Indices.Append({ BaseIndex + 2, BaseIndex + 9, BaseIndex + 3, BaseIndex + 2, BaseIndex + 8, BaseIndex + 9});
+	// Append right dilation quad
+	Indices.Append({ BaseIndex + 4, BaseIndex + 11, BaseIndex + 5, BaseIndex + 4, BaseIndex + 10, BaseIndex + 11 });
+}
+
+enum class ERiverBoundaryEdge {
+	Start,
+	End,
+};
+
+static void AddTerminalVerticesForRiverSpline(ERiverBoundaryEdge Edge, const UWaterBodyRiverComponent* Component, TArray<FDynamicMeshVertex>& Vertices, TArray<uint32>& Indices)
+{
+	const UWaterSplineComponent* SplineComp = Component->GetWaterSpline();
+
+	const float DistanceAlongSpline = (Edge == ERiverBoundaryEdge::Start) ? 0.f : SplineComp->GetSplineLength();
+	
+	const FVector Tangent = SplineComp->GetTangentAtDistanceAlongSpline(DistanceAlongSpline, ESplineCoordinateSpace::Local).GetSafeNormal();
+	const FVector Up = SplineComp->GetUpVectorAtDistanceAlongSpline(DistanceAlongSpline, ESplineCoordinateSpace::Local).GetSafeNormal();
+	
+	const FVector Normal = FVector::CrossProduct(Tangent, Up).GetSafeNormal();
+	const FVector Pos = SplineComp->GetLocationAtDistanceAlongSpline(DistanceAlongSpline, ESplineCoordinateSpace::Local);
+	
+	const float Key = SplineComp->SplineCurves.ReparamTable.Eval(DistanceAlongSpline, 0.f);
+	const float HalfWidth = Cast<UWaterSplineMetadata>(SplineComp->GetSplinePointsMetadata())->RiverWidth.Eval(Key) / 2;
+
+	const float DilationAmount = Component->ShapeDilation;
+	const FVector DilationOffset = Normal * DilationAmount;
+	FVector OutwardDistance = Normal * HalfWidth;
+	OutwardDistance.Z = 0.f;
+
+	FVector TangentialOffset = Tangent * DilationAmount;
+	TangentialOffset.Z = 0.f;
+
+	// For the starting edge the tangential offset is negative to push it backwards
+	if (Edge == ERiverBoundaryEdge::Start)
+	{
+		TangentialOffset *= -1;
+	}
+
+	FDynamicMeshVertex BackLeft(Pos - OutwardDistance + TangentialOffset - DilationOffset);
+	BackLeft.Position.Z += Component->GetShapeDilationZOffsetFar();
+	FDynamicMeshVertex Left(Pos - OutwardDistance + TangentialOffset);
+	Left.Position.Z += Component->GetShapeDilationZOffsetFar();
+	FDynamicMeshVertex Right(Pos + OutwardDistance + TangentialOffset);
+	Right.Position.Z += Component->GetShapeDilationZOffsetFar();
+	FDynamicMeshVertex BackRight(Pos + OutwardDistance + TangentialOffset + DilationOffset);
+	BackRight.Position.Z += Component->GetShapeDilationZOffsetFar();
+
+	// Initialize the vertex data to correct represent what we expect the dilated region to look like
+	// (no color and -1 UVs[0].x
+	BackLeft.Color = FColor(0.f);
+	BackLeft.TextureCoordinate[0].X = -1;
+	Left.Color = FColor(0.f);
+	Left.TextureCoordinate[0].X = -1;
+	Right.Color = FColor(0.f);
+	Right.TextureCoordinate[0].X = -1;
+	BackRight.Color = FColor(0.f);
+	BackRight.TextureCoordinate[0].X = -1;
+	
+
+	Vertices.Append({ BackLeft, Left, Right, BackRight });
+	if (Edge == ERiverBoundaryEdge::Start)
+	{
+		/* Dilated front segment geometry: (ignore vertices 6-7 since they are the non-dilated vertices)
+			4---5,6---7,8---9
+			|    |     |    |
+			0----1-----2----3
+		*/
+		Indices.Append({ 0, 5, 1, 0, 4, 5});
+		Indices.Append({ 1, 8, 2, 1, 5, 8});
+		Indices.Append({ 2, 9, 3, 2, 8, 9});
+	}
+	else
+	{
+		/* Dilated back segment geometry: (ignore vertices 6-7 since they are the non-dilated vertices)
+			6----7-----8----9
+			|    |     |    |
+			0---1,2---3,4---5
+		*/
+		check(Vertices.Num() >= 10);
+		const uint32 BaseIndex = Vertices.Num() - 10;
+		// Since we don't know if we're at the end point or not in the AddVerticesForRiverSplineStep function,
+		// we end up adding a bunch of indices which link to the next set of vertices but they don't exist so we must trim them here.
+		Indices.SetNum(Indices.Num() - 18);
+
+		Indices.Append({ BaseIndex + 0, BaseIndex + 6, BaseIndex + 7, BaseIndex + 0, BaseIndex + 7, BaseIndex + 1 });
+		Indices.Append({ BaseIndex + 1, BaseIndex + 7, BaseIndex + 8, BaseIndex + 1, BaseIndex + 8, BaseIndex + 4 });
+		Indices.Append({ BaseIndex + 4, BaseIndex + 9, BaseIndex + 5, BaseIndex + 4, BaseIndex + 8, BaseIndex + 9 });
+	}
+}
+
+void FWaterBodySceneProxy::GenerateRiverMesh(UWaterBodyRiverComponent* Component)
+{
+	TRACE_CPUPROFILER_EVENT_SCOPE(GenerateRiverMesh);
+
+	const UWaterSplineComponent* SplineComp = Component->GetWaterSpline();
+
+	TArray<FDynamicMeshVertex> Vertices;
+	TArray<uint32> Indices;
+
+	// Add an extra point at the start to dilate starting edge
+	AddTerminalVerticesForRiverSpline(ERiverBoundaryEdge::Start, Component, Vertices, Indices);
+
+	TArray<double> Distances;
+	TArray<FVector> Points;
+	SplineComp->DivideSplineIntoPolylineRecursiveWithDistances(0.f, SplineComp->GetSplineLength(), ESplineCoordinateSpace::Local, FMath::Square(10.f), Points, Distances);
+	
+	for (double DistanceAlongSpline : Distances)
+	{
+		AddVerticesForRiverSplineStep(DistanceAlongSpline, Component, Vertices, Indices);
+	}
+	
+	// Add an extra point at the end to dilate ending edge
+	AddTerminalVerticesForRiverSpline(ERiverBoundaryEdge::End, Component, Vertices, Indices);
+
+
+	FWaterBodyMeshSection* NewSection = new FWaterBodyMeshSection(GetScene().GetFeatureLevel());
+	NewSection->IndexBuffer.Indices = MoveTemp(Indices);
+	
+	NewSection->VertexBuffers.InitFromDynamicVertex(&NewSection->VertexFactory, Vertices);
+
+	InitResources(NewSection);
+
+	Sections.Add(NewSection);
+}
+
+void FWaterBodySceneProxy::InitResources(FWaterBodyMeshSection* Section)
+{
+	check(Section != nullptr);
+	BeginInitResource(&Section->VertexBuffers.PositionVertexBuffer);
+	BeginInitResource(&Section->VertexBuffers.StaticMeshVertexBuffer);
+	BeginInitResource(&Section->VertexBuffers.ColorVertexBuffer);
+	BeginInitResource(&Section->IndexBuffer);
+	BeginInitResource(&Section->VertexFactory);
+}
+
+FWaterBodySceneProxy::FWaterBodySceneProxy(UWaterBodyComponent* Component)
+	: FPrimitiveSceneProxy(Component)
+{
+	if (Component == nullptr || Component->GetWaterSpline() == nullptr)
+	{
+		return;
+	}
+
+	switch(Component->GetWaterBodyType())
+	{
+		case EWaterBodyType::Lake:
+			GenerateLakeMesh(CastChecked<UWaterBodyLakeComponent>(Component));
+			break;
+		case EWaterBodyType::Ocean:
+			GenerateOceanMesh(CastChecked<UWaterBodyOceanComponent>(Component));
+			break;
+		case EWaterBodyType::River:
+			GenerateRiverMesh(CastChecked<UWaterBodyRiverComponent>(Component));
+			break;
+		default:
+			break;
+	}
+
+	if (UMaterialInstance* WaterInfoMaterial = Component->GetWaterInfoMaterialInstance())
+	{
+		Material = WaterInfoMaterial->GetRenderProxy();
+	}
+}
+
+FWaterBodySceneProxy::~FWaterBodySceneProxy()
+{
+	for (FWaterBodyMeshSection* Section : Sections)
+	{
+		Section->VertexBuffers.PositionVertexBuffer.ReleaseResource();
+		Section->VertexBuffers.StaticMeshVertexBuffer.ReleaseResource();
+		Section->VertexBuffers.ColorVertexBuffer.ReleaseResource();
+		Section->IndexBuffer.ReleaseResource();
+		Section->VertexFactory.ReleaseResource();
+
+		delete Section;
+	}
+}
+
+void FWaterBodySceneProxy::GetDynamicMeshElements(const TArray<const FSceneView*>& Views, const FSceneViewFamily& ViewFamily, uint32 VisibilityMap, FMeshElementCollector& Collector) const
+{
+	FMaterialRenderProxy* MaterialToUse = Material;
+	if (!MaterialToUse)
+	{
+		FColoredMaterialRenderProxy* FallbackMaterial = new FColoredMaterialRenderProxy(
+			GEngine->DebugMeshMaterial ? GEngine->DebugMeshMaterial->GetRenderProxy() : nullptr,
+			FLinearColor(1.f, 1.f, 1.f));
+
+		Collector.RegisterOneFrameMaterialProxy(FallbackMaterial);
+
+		MaterialToUse = FallbackMaterial;
+	}
+
+	// If we are not in the waterinfo pass and the cvar is not set to show opaque bodies, we should be in wireframe
+	const bool bWireframe = AllowDebugViewmodes() && !(bWithinWaterInfoPass) && (CVarWaterShowProxies.GetValueOnRenderThread() != 3);
+
+	if (bWireframe)
+	{
+		FColoredMaterialRenderProxy* WireframeMaterialInstance= new FColoredMaterialRenderProxy(
+			GEngine->WireframeMaterial ? GEngine->WireframeMaterial->GetRenderProxy() : NULL,
+			FLinearColor(0, 0.5f, 1.f)
+			);
+
+		Collector.RegisterOneFrameMaterialProxy(WireframeMaterialInstance);
+		MaterialToUse = WireframeMaterialInstance;
+	}
+	
+	for (int32 ViewIndex = 0; ViewIndex < Views.Num(); ++ViewIndex)
+	{
+		for (const FWaterBodyMeshSection* Section : Sections)
+		{
+			FMeshBatch& Mesh = Collector.AllocateMesh();
+			FMeshBatchElement& BatchElement = Mesh.Elements[0];
+			BatchElement.IndexBuffer = &Section->IndexBuffer;
+			Mesh.bWireframe = bWireframe;
+			Mesh.VertexFactory = &Section->VertexFactory;
+			Mesh.MaterialRenderProxy = MaterialToUse;
+
+			bool bHasPrecomputedVolumetricLightmap;
+			FMatrix PreviousLocalToWorld;
+			int32 SingleCaptureIndex;
+			bool bOutputVelocity;
+			GetScene().GetPrimitiveUniformShaderParameters_RenderThread(GetPrimitiveSceneInfo(), bHasPrecomputedVolumetricLightmap, PreviousLocalToWorld, SingleCaptureIndex, bOutputVelocity);
+
+			FDynamicPrimitiveUniformBuffer& DynamicPrimitiveUniformBuffer = Collector.AllocateOneFrameResource<FDynamicPrimitiveUniformBuffer>();
+			DynamicPrimitiveUniformBuffer.Set(GetLocalToWorld(), PreviousLocalToWorld, GetBounds(), GetLocalBounds(), true, bHasPrecomputedVolumetricLightmap, DrawsVelocity(), bOutputVelocity); BatchElement.PrimitiveUniformBufferResource = &DynamicPrimitiveUniformBuffer.UniformBuffer;
+
+			BatchElement.FirstIndex = 0;
+			BatchElement.NumPrimitives = Section->IndexBuffer.Indices.Num() / 3;
+			check(BatchElement.NumPrimitives != 0);
+			BatchElement.MinVertexIndex = 0;
+			BatchElement.MaxVertexIndex = Section->VertexBuffers.PositionVertexBuffer.GetNumVertices() - 1;
+			Mesh.ReverseCulling = IsLocalToWorldDeterminantNegative();
+			Mesh.Type = PT_TriangleList;
+			Mesh.DepthPriorityGroup = bWireframe ? SDPG_Foreground : SDPG_World;
+			Mesh.bCanApplyViewModeOverrides = false;
+			Collector.AddMesh(ViewIndex, Mesh);
+		}
+	}
+}
+
+FPrimitiveViewRelevance FWaterBodySceneProxy::GetViewRelevance(const FSceneView* View) const
+{
+	FPrimitiveViewRelevance Result;
+	Result.bDrawRelevance = IsShown(View);
+	Result.bShadowRelevance = IsShadowCast(View);
+	Result.bDynamicRelevance = true;
+	Result.bRenderInMainPass = ShouldRenderInMainPass();
+	Result.bUsesLightingChannels = GetLightingChannelMask() != GetDefaultLightingChannelMask();
+	Result.bRenderCustomDepth = ShouldRenderCustomDepth();
+	Result.bTranslucentSelfShadow = bCastVolumetricTranslucentShadow;
+	Result.bVelocityRelevance = IsMovable();
+	return Result;
+}
+
+SIZE_T FWaterBodySceneProxy::GetTypeHash() const
+{
+	static size_t UniquePtr;
+	return reinterpret_cast<size_t>(&UniquePtr);
+}
+
+uint32 FWaterBodySceneProxy::GetMemoryFootprint() const 
+{
+	return (sizeof(*this) + GetAllocatedSize());
+}
+
+uint32 FWaterBodySceneProxy::GetAllocatedSize() const
+{
+	return FPrimitiveSceneProxy::GetAllocatedSize();
+}
+
+bool FWaterBodySceneProxy::IsShown(const FSceneView* View) const
+{
+	if (!bWithinWaterInfoPass)
+	{
+		const int32 ShowProxiesCVar = CVarWaterShowProxies.GetValueOnRenderThread();
+		if ((ShowProxiesCVar == 1 && IsSelected()) || ShowProxiesCVar >= 2)
+		{
+			return true;
+		}
+
+		return false;
+	}
+
+	return FPrimitiveSceneProxy::IsShown(View);
+}
+
+void FWaterBodySceneProxy::SetWithinWaterInfoPass(bool bInWithinWaterInfoPass)
+{
+	bWithinWaterInfoPass = bInWithinWaterInfoPass;
+}
diff --git a/Engine/Plugins/Experimental/Water/Source/Runtime/Private/WaterInfoRendering.cpp b/Engine/Plugins/Experimental/Water/Source/Runtime/Private/WaterInfoRendering.cpp
new file mode 100644
--- /dev/null
+++ b/Engine/Plugins/Experimental/Water/Source/Runtime/Private/WaterInfoRendering.cpp
@@ -0,0 +1,600 @@
+// Copyright Epic Games, Inc. All Rights Reserved.
+
+#include "WaterInfoRendering.h"
+
+#include "EngineModule.h"
+#include "LegacyScreenPercentageDriver.h"
+#include "RenderCaptureInterface.h"
+#include "WaterBodyActor.h"
+#include "WaterBodySceneProxy.h"
+#include "Engine/TextureRenderTarget2D.h"
+#include "Renderer/Private/ScenePrivate.h"
+#include "Renderer/Private/RendererModule.h"
+#include "Renderer/Private/SceneCaptureRendering.h"
+#include "Renderer/Private/PostProcess/SceneFilterRendering.h"
+#include "Math/OrthoMatrix.h"
+#include "GameFramework/WorldSettings.h"
+#include "WaterZoneActor.h"
+#include "Rendering/NaniteStreamingManager.h"
+#include "SceneViewExtension.h"
+
+namespace WaterInfo
+{
+
+struct FUpdateWaterInfoParams
+{
+	FSceneRenderer* DepthRenderer;
+	FSceneRenderer* ColorRenderer;
+	FRenderTarget* RenderTarget;
+	FTexture* OutputTexture;
+
+	FVector2f WaterHeightExtents;
+	float GroundZMin;
+	float CaptureZ;
+	int32 VelocityBlurRadius;
+};
+
+
+// ---------------------------------------------------------------------------------------------------------------------
+
+/** A pixel shader for capturing a component of the rendered scene for a scene capture.*/
+class FWaterInfoMergePS : public FGlobalShader
+{
+public:
+	DECLARE_GLOBAL_SHADER(FWaterInfoMergePS);
+	SHADER_USE_PARAMETER_STRUCT(FWaterInfoMergePS, FGlobalShader);
+
+	BEGIN_SHADER_PARAMETER_STRUCT(FParameters, )
+		SHADER_PARAMETER_STRUCT_REF(FViewUniformShaderParameters, View)
+		SHADER_PARAMETER_STRUCT_INCLUDE(FSceneTextureShaderParameters, SceneTextures)
+		SHADER_PARAMETER_RDG_TEXTURE(Texture2D, DepthTexture)
+		SHADER_PARAMETER_RDG_TEXTURE(Texture2D<float4>, ColorTexture)
+		SHADER_PARAMETER(FVector2f, WaterHeightExtents)
+		SHADER_PARAMETER(float, GroundZMin)
+		SHADER_PARAMETER(float, CaptureZ)
+		RENDER_TARGET_BINDING_SLOTS()
+	END_SHADER_PARAMETER_STRUCT()
+};
+
+IMPLEMENT_GLOBAL_SHADER(FWaterInfoMergePS, "/Plugin/Water/Private/WaterInfoMerge.usf", "Main", SF_Pixel);
+
+static void MergeWaterInfoAndDepth(
+	FRDGBuilder& GraphBuilder,
+	const FMinimalSceneTextures& SceneTextures,
+	const FSceneViewFamily& ViewFamily,
+	FViewInfo& View,
+	FRDGTextureRef OutputTexture,
+	FRDGTextureRef DepthTexture,
+	FRDGTextureRef ColorTexture,
+	FUpdateWaterInfoParams Params)
+{
+	RDG_EVENT_SCOPE(GraphBuilder, "WaterInfoDepthMerge");
+
+	FGraphicsPipelineStateInitializer GraphicsPSOInit;
+	GraphicsPSOInit.RasterizerState = TStaticRasterizerState<FM_Solid, CM_None>::GetRHI();
+	GraphicsPSOInit.DepthStencilState = TStaticDepthStencilState<false, CF_Always>::GetRHI();
+	GraphicsPSOInit.BlendState = TStaticBlendState<>::GetRHI();
+
+	{
+		FWaterInfoMergePS::FParameters* PassParameters = GraphBuilder.AllocParameters<FWaterInfoMergePS::FParameters>();
+		PassParameters->View = View.ViewUniformBuffer;
+		PassParameters->RenderTargets[0] = FRenderTargetBinding(OutputTexture, ERenderTargetLoadAction::ENoAction);
+		PassParameters->SceneTextures = SceneTextures.GetSceneTextureShaderParameters(ViewFamily.GetFeatureLevel());
+		PassParameters->DepthTexture = DepthTexture;
+		PassParameters->ColorTexture = ColorTexture;
+		PassParameters->CaptureZ = Params.CaptureZ;
+		PassParameters->WaterHeightExtents = Params.WaterHeightExtents;
+		PassParameters->GroundZMin = Params.GroundZMin;
+
+		TShaderMapRef<FScreenVS> VertexShader(View.ShaderMap);
+		TShaderMapRef<FWaterInfoMergePS> PixelShader(View.ShaderMap);
+
+		GraphicsPSOInit.BoundShaderState.VertexDeclarationRHI = GFilterVertexDeclaration.VertexDeclarationRHI;
+		GraphicsPSOInit.BoundShaderState.VertexShaderRHI = VertexShader.GetVertexShader();
+		GraphicsPSOInit.BoundShaderState.PixelShaderRHI = PixelShader.GetPixelShader();
+		GraphicsPSOInit.PrimitiveType = PT_TriangleList;
+
+		GraphBuilder.AddPass(
+			RDG_EVENT_NAME("WaterInfoDepthMerge"),
+			PassParameters,
+			ERDGPassFlags::Raster,
+			[PassParameters, GraphicsPSOInit, VertexShader, PixelShader, &View](FRHICommandListImmediate& RHICmdList)
+			{
+				FGraphicsPipelineStateInitializer LocalGraphicsPSOInit = GraphicsPSOInit;
+				RHICmdList.ApplyCachedRenderTargets(LocalGraphicsPSOInit);
+				SetGraphicsPipelineState(RHICmdList, LocalGraphicsPSOInit, 0);
+				SetShaderParameters(RHICmdList, PixelShader, PixelShader.GetPixelShader(), *PassParameters);
+
+				DrawRectangle(
+					RHICmdList,
+					View.ViewRect.Min.X, View.ViewRect.Min.Y,
+					View.ViewRect.Width(), View.ViewRect.Height(),
+					View.ViewRect.Min.X, View.ViewRect.Min.Y,
+					View.ViewRect.Width(), View.ViewRect.Height(),
+					View.UnconstrainedViewRect.Size(),
+					View.UnconstrainedViewRect.Size(),
+					VertexShader,
+					EDRF_UseTriangleOptimization);
+			});
+	}
+}
+
+// ---------------------------------------------------------------------------------------------------------------------
+
+
+/** A pixel shader for capturing a component of the rendered scene for a scene capture.*/
+class FWaterInfoFinalizePS : public FGlobalShader
+{
+public:
+	DECLARE_GLOBAL_SHADER(FWaterInfoFinalizePS);
+	SHADER_USE_PARAMETER_STRUCT(FWaterInfoFinalizePS, FGlobalShader);
+
+	BEGIN_SHADER_PARAMETER_STRUCT(FParameters, )
+		SHADER_PARAMETER_STRUCT_REF(FViewUniformShaderParameters, View)
+		SHADER_PARAMETER_STRUCT_INCLUDE(FSceneTextureShaderParameters, SceneTextures)
+		SHADER_PARAMETER_RDG_TEXTURE(Texture2D, WaterInfoTexture)
+		SHADER_PARAMETER(FVector2f, WaterHeightExtents)
+		SHADER_PARAMETER(float, GroundZMin)
+		SHADER_PARAMETER(float, CaptureZ)
+		SHADER_PARAMETER(int, BlurRadius)
+		RENDER_TARGET_BINDING_SLOTS()
+	END_SHADER_PARAMETER_STRUCT()
+};
+
+IMPLEMENT_GLOBAL_SHADER(FWaterInfoFinalizePS, "/Plugin/Water/Private/WaterInfoFinalize.usf", "Main", SF_Pixel);
+
+static void FinalizeWaterInfo(
+	FRDGBuilder& GraphBuilder,
+	const FMinimalSceneTextures& SceneTextures,
+	const FSceneViewFamily& ViewFamily,
+	FViewInfo& View,
+	FRDGTextureRef WaterInfoTexture,
+	FRDGTextureRef OutputTexture,
+	FUpdateWaterInfoParams Params)
+{
+	RDG_EVENT_SCOPE(GraphBuilder, "WaterInfoFinalize");
+
+	FGraphicsPipelineStateInitializer GraphicsPSOInit;
+	GraphicsPSOInit.RasterizerState = TStaticRasterizerState<FM_Solid, CM_None>::GetRHI();
+	GraphicsPSOInit.DepthStencilState = TStaticDepthStencilState<false, CF_Always>::GetRHI();
+	GraphicsPSOInit.BlendState = TStaticBlendState<>::GetRHI();
+
+	{
+		FWaterInfoFinalizePS::FParameters* PassParameters = GraphBuilder.AllocParameters<FWaterInfoFinalizePS::FParameters>();
+		PassParameters->View = View.ViewUniformBuffer;
+		PassParameters->RenderTargets[0] = FRenderTargetBinding(OutputTexture, ERenderTargetLoadAction::ENoAction);
+		PassParameters->SceneTextures = SceneTextures.GetSceneTextureShaderParameters(ViewFamily.GetFeatureLevel());
+		PassParameters->WaterInfoTexture = WaterInfoTexture;
+		PassParameters->BlurRadius = Params.VelocityBlurRadius;
+		PassParameters->CaptureZ = Params.CaptureZ;
+		PassParameters->WaterHeightExtents = Params.WaterHeightExtents;
+		PassParameters->GroundZMin = Params.GroundZMin;
+
+		TShaderMapRef<FScreenVS> VertexShader(View.ShaderMap);
+		TShaderMapRef<FWaterInfoFinalizePS> PixelShader(View.ShaderMap);
+
+		GraphicsPSOInit.BoundShaderState.VertexDeclarationRHI = GFilterVertexDeclaration.VertexDeclarationRHI;
+		GraphicsPSOInit.BoundShaderState.VertexShaderRHI = VertexShader.GetVertexShader();
+		GraphicsPSOInit.BoundShaderState.PixelShaderRHI = PixelShader.GetPixelShader();
+		GraphicsPSOInit.PrimitiveType = PT_TriangleList;
+
+		GraphBuilder.AddPass(
+			RDG_EVENT_NAME("WaterInfoFinalize"),
+			PassParameters,
+			ERDGPassFlags::Raster,
+			[PassParameters, GraphicsPSOInit, VertexShader, PixelShader, &View](FRHICommandListImmediate& RHICmdList)
+			{
+				FGraphicsPipelineStateInitializer LocalGraphicsPSOInit = GraphicsPSOInit;
+				RHICmdList.ApplyCachedRenderTargets(LocalGraphicsPSOInit);
+				SetGraphicsPipelineState(RHICmdList, LocalGraphicsPSOInit, 0);
+				SetShaderParameters(RHICmdList, PixelShader, PixelShader.GetPixelShader(), *PassParameters);
+				
+				DrawRectangle(
+					RHICmdList,
+					View.ViewRect.Min.X, View.ViewRect.Min.Y,
+					View.ViewRect.Width(), View.ViewRect.Height(),
+					View.ViewRect.Min.X, View.ViewRect.Min.Y,
+					View.ViewRect.Width(), View.ViewRect.Height(),
+					View.UnconstrainedViewRect.Size(),
+					View.UnconstrainedViewRect.Size(),
+					VertexShader,
+					EDRF_UseTriangleOptimization);
+			});
+	}
+}
+// ---------------------------------------------------------------------------------------------------------------------
+
+static FMatrix BuildOrthoMatrix(float InOrthoWidth, float InOrthoHeight)
+{
+	check((int32)ERHIZBuffer::IsInverted);
+
+	const float OrthoWidth = InOrthoWidth / 2.0f;
+	const float OrthoHeight = InOrthoHeight / 2.0f;
+
+	const float NearPlane = 0.f;
+	const float FarPlane = WORLD_MAX / 8.0f;
+
+	const float ZScale = 1.0f / (FarPlane - NearPlane);
+	const float ZOffset = 0;
+
+	return FReversedZOrthoMatrix(
+		OrthoWidth,
+		OrthoHeight,
+		ZScale,
+		ZOffset
+		);
+}
+
+// ---------------------------------------------------------------------------------------------------------------------
+
+static void SetWaterBodiesWithinWaterInfoPass(FSceneRenderer* SceneRenderer, bool bWithinWaterInfoPass)
+{
+	if (SceneRenderer->Views[0].ShowOnlyPrimitives.IsSet())
+	{
+		for (FPrimitiveComponentId PrimId : SceneRenderer->Views[0].ShowOnlyPrimitives.GetValue())
+		{
+			for (FPrimitiveSceneProxy* PrimProxy : SceneRenderer->Scene->PrimitiveSceneProxies)
+			{
+				if (PrimProxy && PrimProxy->GetPrimitiveComponentId() == PrimId)
+				{
+					FWaterBodySceneProxy* WaterProxy = (FWaterBodySceneProxy*)PrimProxy;
+					WaterProxy->SetWithinWaterInfoPass(bWithinWaterInfoPass);
+				}
+			}
+		}
+	}
+}
+
+static void UpdateWaterInfoRendering_RenderThread(
+	FRHICommandListImmediate& RHICmdList,
+	const FUpdateWaterInfoParams& Params)
+{
+	FMaterialRenderProxy::UpdateDeferredCachedUniformExpressions();
+
+	FRenderTarget* RenderTarget = Params.RenderTarget;
+	FTexture* OutputTexture = Params.OutputTexture;
+
+	// Depth-only pass for actors which are considered the ground for water rendering
+	{
+		FSceneRenderer* DepthRenderer = Params.DepthRenderer;
+
+		// We need to execute the pre-render view extensions before we do any view dependent work.
+		FSceneRenderer::ViewExtensionPreRender_RenderThread(RHICmdList, DepthRenderer);
+
+		DepthRenderer->RenderThreadBegin(RHICmdList);
+		
+		FDeferredUpdateResource::UpdateResources(RHICmdList);
+		
+		SCOPED_DRAW_EVENT(RHICmdList, DepthRendering_RT);
+		
+		FRDGBuilder GraphBuilder(RHICmdList, RDG_EVENT_NAME("WaterInfoColorRendering"), ERDGBuilderFlags::AllowParallelExecute);
+		
+		FRDGTextureRef TargetTexture = RegisterExternalTexture(GraphBuilder, RenderTarget->GetRenderTargetTexture(), TEXT("WaterDepthTarget"));
+		
+
+		FViewInfo& View = DepthRenderer->Views[0];
+
+		AddClearRenderTargetPass(GraphBuilder, TargetTexture, FLinearColor::Black, View.UnscaledViewRect);
+
+		View.bDisableQuerySubmissions = true;
+		View.bIgnoreExistingQueries = true;
+
+		{
+			RDG_RHI_EVENT_SCOPE(GraphBuilder, RenderWaterInfoDepth);
+			DepthRenderer->Render(GraphBuilder);
+		}
+		
+		FRDGTextureRef ResolveTexture = RegisterExternalTexture(GraphBuilder, OutputTexture->TextureRHI, TEXT("WaterDepthResolve"));
+		FResolveParams ResolveParams;
+		AddCopyToResolveTargetPass(GraphBuilder, TargetTexture, ResolveTexture, ResolveParams);
+
+		if (DepthRenderer->Scene->GetShadingPath() == EShadingPath::Mobile)
+		{
+			const FMinimalSceneTextures& SceneTextures = FSceneTextures::Get(GraphBuilder);
+			const bool bNeedsFlippedRenderTarget = false;
+			RDG_EVENT_SCOPE(GraphBuilder, "CaptureSceneColor");
+			CopySceneCaptureComponentToTarget(
+				GraphBuilder,
+				SceneTextures,
+				ResolveTexture,
+				DepthRenderer->ViewFamily,
+				DepthRenderer->Views,
+				bNeedsFlippedRenderTarget);
+		}
+		
+		GraphBuilder.Execute();
+
+		DepthRenderer->RenderThreadEnd(RHICmdList);
+	}
+
+	// Render the water bodies' data including flow, zoffset, depth
+	{
+		FSceneRenderer* ColorRenderer = Params.ColorRenderer;
+
+		// We need to execute the pre-render view extensions before we do any view dependent work.
+		FSceneRenderer::ViewExtensionPreRender_RenderThread(RHICmdList, ColorRenderer);
+
+		ColorRenderer->RenderThreadBegin(RHICmdList);
+
+		SetWaterBodiesWithinWaterInfoPass(ColorRenderer, true);
+		
+		SCOPED_DRAW_EVENT(RHICmdList, ColorRendering_RT);
+		
+		FRDGBuilder GraphBuilder(RHICmdList, RDG_EVENT_NAME("WaterInfoColorRendering"), ERDGBuilderFlags::AllowParallelExecute);
+		FRDGTextureRef TargetTexture = RegisterExternalTexture(GraphBuilder, RenderTarget->GetRenderTargetTexture(), TEXT("WaterInfoTarget"));
+
+		FRDGTextureRef DepthTexture = GraphBuilder.CreateTexture(TargetTexture->Desc, TEXT("WaterInfoDepth"));
+		AddCopyTexturePass(GraphBuilder, TargetTexture, DepthTexture);
+		
+		FViewInfo& View = ColorRenderer->Views[0];
+
+		AddClearRenderTargetPass(GraphBuilder, TargetTexture, FLinearColor::Black, View.UnscaledViewRect);
+
+		View.bDisableQuerySubmissions = true;
+		View.bIgnoreExistingQueries = true;
+
+		{
+			RDG_RHI_EVENT_SCOPE(GraphBuilder, RenderWaterInfoColor);
+			ColorRenderer->Render(GraphBuilder);
+		}
+
+		const FMinimalSceneTextures& SceneTextures = FSceneTextures::Get(GraphBuilder);
+		FRDGTextureDesc ColorTextureDesc(TargetTexture->Desc);
+		ColorTextureDesc.Format = PF_A32B32G32R32F;
+		FRDGTextureRef ColorTexture = GraphBuilder.CreateTexture(ColorTextureDesc, TEXT("WaterInfoColor"));
+		{
+			const bool bNeedsFlippedRenderTarget = false;
+			RDG_EVENT_SCOPE(GraphBuilder, "CaptureSceneColor");
+			CopySceneCaptureComponentToTarget(
+				GraphBuilder,
+				SceneTextures,
+				ColorTexture,
+				ColorRenderer->ViewFamily,
+				ColorRenderer->Views,
+				bNeedsFlippedRenderTarget);
+		}
+
+		FRDGTextureRef MergeTargetTexture = GraphBuilder.CreateTexture(ColorTextureDesc, TEXT("WaterInfoMerged"));
+		MergeWaterInfoAndDepth(GraphBuilder, SceneTextures, ColorRenderer->ViewFamily, ColorRenderer->Views[0], MergeTargetTexture, DepthTexture, ColorTexture, Params);
+
+		FRDGTextureRef FinalizedTexture = GraphBuilder.CreateTexture(TargetTexture->Desc, TEXT("WaterInfoFinalized"));
+		FinalizeWaterInfo(GraphBuilder, SceneTextures, ColorRenderer->ViewFamily, ColorRenderer->Views[0], MergeTargetTexture, FinalizedTexture, Params);
+		
+		FRDGTextureRef ResolveOutputTexture = RegisterExternalTexture(GraphBuilder, OutputTexture->TextureRHI, TEXT("WaterInfoResolve"));
+		FResolveParams ResolveParams;
+		AddCopyToResolveTargetPass(GraphBuilder, FinalizedTexture, ResolveOutputTexture, ResolveParams);
+		GraphBuilder.Execute();
+		
+		SetWaterBodiesWithinWaterInfoPass(ColorRenderer, false);
+
+		ColorRenderer->RenderThreadEnd(RHICmdList);
+	}
+
+	RHICmdList.Transition(FRHITransitionInfo(Params.OutputTexture->TextureRHI, ERHIAccess::RTV, ERHIAccess::SRVMask));
+}
+
+static FSceneRenderer* CreateWaterInfoDepthRenderer(
+	FSceneInterface* Scene,
+	FRenderTarget* RenderTarget,
+	const WaterInfo::FRenderingContext& Context,
+	FIntPoint RenderTargetSize,
+	const FMatrix& ViewRotationMatrix,
+	const FVector& ViewLocation,
+	const FMatrix& ProjectionMatrix)
+{
+	FEngineShowFlags ShowFlags(ESFIM_Game);
+	ShowFlags.NaniteMeshes = 0;
+	ShowFlags.Atmosphere = 0;
+	ShowFlags.Lighting = 0;
+	ShowFlags.Bloom = 0;
+	ShowFlags.ScreenPercentage = 0;
+	ShowFlags.Translucency = 0;
+	ShowFlags.SeparateTranslucency = 0;
+	ShowFlags.AntiAliasing = 0;
+	FSceneViewFamilyContext DepthViewFamily(FSceneViewFamily::ConstructionValues(
+		RenderTarget,
+		Scene,
+		ShowFlags)
+		.SetRealtimeUpdate(false)
+		.SetResolveScene(false));
+	DepthViewFamily.SceneCaptureSource = ESceneCaptureSource::SCS_DeviceDepth;
+
+	// Setup the view family
+	FSceneViewInitOptions DepthViewInitOptions;
+	DepthViewInitOptions.SetViewRectangle(FIntRect(0, 0, RenderTargetSize.X, RenderTargetSize.Y));
+	DepthViewInitOptions.ViewFamily = &DepthViewFamily;
+	DepthViewInitOptions.ViewActor = Context.ZoneToRender;
+	DepthViewInitOptions.ViewRotationMatrix = ViewRotationMatrix;
+	DepthViewInitOptions.ViewOrigin = ViewLocation;
+	DepthViewInitOptions.BackgroundColor = FLinearColor::Black;
+	DepthViewInitOptions.OverrideFarClippingPlaneDistance = -1.f;
+	DepthViewInitOptions.SceneViewStateInterface = nullptr;
+	DepthViewInitOptions.ProjectionMatrix = ProjectionMatrix;
+	DepthViewInitOptions.LODDistanceFactor = 0.001f;
+	DepthViewInitOptions.OverlayColor = FLinearColor::Black;
+
+	if (DepthViewFamily.Scene->GetWorld() != nullptr && DepthViewFamily.Scene->GetWorld()->GetWorldSettings() != nullptr)
+	{
+		DepthViewInitOptions.WorldToMetersScale = DepthViewFamily.Scene->GetWorld()->GetWorldSettings()->WorldToMeters;
+	}
+
+	FSceneView* DepthView = new FSceneView(DepthViewInitOptions);
+	DepthView->AntiAliasingMethod = EAntiAliasingMethod::AAM_None;
+	DepthView->SetupAntiAliasingMethod();
+
+	if (Context.GroundActors.Num() > 0)
+	{
+		DepthView->ShowOnlyPrimitives.Emplace();
+		DepthView->ShowOnlyPrimitives->Reserve(Context.GroundActors.Num());
+		for (AActor* GroundActor : Context.GroundActors)
+		{
+			TInlineComponentArray<UPrimitiveComponent*> PrimComps(GroundActor);
+			for (UPrimitiveComponent* PrimComp : PrimComps)
+			{
+				if (PrimComp)
+				{
+					DepthView->ShowOnlyPrimitives->Add(PrimComp->ComponentId);
+				}
+			}
+		}
+	}
+
+	DepthViewFamily.Views.Add(DepthView);
+	
+	DepthView->StartFinalPostprocessSettings(ViewLocation);
+	DepthView->EndFinalPostprocessSettings(DepthViewInitOptions);
+
+	DepthViewFamily.SetScreenPercentageInterface(new FLegacyScreenPercentageDriver(DepthViewFamily, 1.f));
+
+	DepthViewFamily.ViewExtensions = GEngine->ViewExtensions->GatherActiveExtensions(FSceneViewExtensionContext(Scene));
+	for (const FSceneViewExtensionRef& Extension : DepthViewFamily.ViewExtensions)
+	{
+		Extension->SetupViewFamily(DepthViewFamily);
+		Extension->SetupView(DepthViewFamily, *DepthView);
+	}
+
+	return FSceneRenderer::CreateSceneRenderer(&DepthViewFamily, nullptr);
+}
+
+static FSceneRenderer* CreateWaterInfoColorRenderer(
+	FSceneInterface* Scene,
+	FRenderTarget* RenderTarget,
+	const WaterInfo::FRenderingContext& Context,
+	FIntPoint RenderTargetSize,
+	const FMatrix& ViewRotationMatrix,
+	const FVector& ViewLocation,
+	const FMatrix& ProjectionMatrix)
+{
+	FEngineShowFlags ShowFlags(ESFIM_Game);
+	ShowFlags.NaniteMeshes = 0;
+	ShowFlags.Atmosphere = 0;
+	ShowFlags.Lighting = 0;
+	ShowFlags.Bloom = 0;
+	ShowFlags.ScreenPercentage = 0;
+	ShowFlags.Translucency = 0;
+	ShowFlags.SeparateTranslucency = 0;
+	ShowFlags.AntiAliasing = 0;
+	
+	FSceneViewFamilyContext ColorViewFamily(FSceneViewFamily::ConstructionValues(
+		RenderTarget,
+		Scene,
+		ShowFlags)
+		.SetRealtimeUpdate(false)
+		.SetResolveScene(false));
+	ColorViewFamily.SceneCaptureSource = ESceneCaptureSource::SCS_SceneColorSceneDepth;
+
+	FSceneViewInitOptions ColorViewInitOptions;
+	ColorViewInitOptions.SetViewRectangle(FIntRect(0, 0, RenderTargetSize.X, RenderTargetSize.Y));
+	ColorViewInitOptions.ViewFamily = &ColorViewFamily;
+	ColorViewInitOptions.ViewActor = Context.ZoneToRender;
+	ColorViewInitOptions.ViewRotationMatrix = ViewRotationMatrix;
+	ColorViewInitOptions.ViewOrigin = ViewLocation;
+	ColorViewInitOptions.BackgroundColor = FLinearColor::Black;
+	ColorViewInitOptions.OverrideFarClippingPlaneDistance = -1.f;
+	ColorViewInitOptions.SceneViewStateInterface = nullptr;
+	ColorViewInitOptions.ProjectionMatrix = ProjectionMatrix;
+	ColorViewInitOptions.LODDistanceFactor = 0.001f;
+	ColorViewInitOptions.OverlayColor = FLinearColor::Black;
+
+	if (ColorViewFamily.Scene->GetWorld() != nullptr && ColorViewFamily.Scene->GetWorld()->GetWorldSettings() != nullptr)
+	{
+		ColorViewInitOptions.WorldToMetersScale = ColorViewFamily.Scene->GetWorld()->GetWorldSettings()->WorldToMeters;
+	}
+
+	FSceneView* ColorView = new FSceneView(ColorViewInitOptions);
+	ColorView->bIsSceneCapture = true;
+	ColorView->AntiAliasingMethod = EAntiAliasingMethod::AAM_None;
+	ColorView->SetupAntiAliasingMethod();
+
+	if (Context.WaterBodies.Num() > 0)
+	{
+		ColorView->ShowOnlyPrimitives.Emplace();
+		ColorView->ShowOnlyPrimitives->Reserve(Context.WaterBodies.Num());
+		for (const UWaterBodyComponent* WaterBodyToRender : Context.WaterBodies)
+		{
+			ColorView->ShowOnlyPrimitives->Add(WaterBodyToRender->ComponentId);
+		}
+	}
+
+	ColorViewFamily.Views.Add(ColorView);
+
+	ColorView->StartFinalPostprocessSettings(ViewLocation);
+	ColorView->EndFinalPostprocessSettings(ColorViewInitOptions);
+
+	ColorViewFamily.SetScreenPercentageInterface(new FLegacyScreenPercentageDriver(ColorViewFamily, 1.f));
+
+	ColorViewFamily.ViewExtensions = GEngine->ViewExtensions->GatherActiveExtensions(FSceneViewExtensionContext(Scene));
+	for (const FSceneViewExtensionRef& Extension : ColorViewFamily.ViewExtensions)
+	{
+		Extension->SetupViewFamily(ColorViewFamily);
+		Extension->SetupView(ColorViewFamily, *ColorView);
+	}
+
+	return FSceneRenderer::CreateSceneRenderer(&ColorViewFamily, nullptr);
+}
+
+void UpdateWaterInfoRendering(
+	FSceneInterface* Scene,
+	UTextureRenderTarget2D* TextureRenderTarget,
+	const WaterInfo::FRenderingContext& Context)
+{
+	if (TextureRenderTarget == nullptr || Scene == nullptr)
+	{
+		return;
+	}
+
+	const FVector2D ZoneExtent = Context.ZoneToRender->GetZoneExtent();
+	FVector ViewLocation = Context.ZoneToRender->GetActorLocation();
+	ViewLocation.Z = Context.CaptureZ;
+
+	// Zone rendering always happens facing towards negative z.
+	const FVector LookAt = ViewLocation - FVector(0.f, 0.f, 1.f);
+
+	FMatrix ViewRotationMat = FLookAtMatrix(ViewLocation, LookAt, FVector(0.f, -1.f, 0.f));
+	ViewRotationMat = ViewRotationMat.RemoveTranslation();
+	ViewRotationMat.RemoveScaling();
+	
+	const FIntPoint CaptureExtent(TextureRenderTarget->GetSurfaceWidth(), TextureRenderTarget->GetSurfaceHeight());
+
+	const FMatrix OrthoProj = BuildOrthoMatrix(ZoneExtent.X, ZoneExtent.Y);
+
+	FSceneRenderer* DepthRenderer = CreateWaterInfoDepthRenderer(
+		Scene,
+		TextureRenderTarget->GameThread_GetRenderTargetResource(),
+		Context,
+		CaptureExtent,
+		ViewRotationMat,
+		ViewLocation,
+		OrthoProj);
+	
+	FSceneRenderer* ColorRenderer = CreateWaterInfoColorRenderer(
+		Scene,
+		TextureRenderTarget->GameThread_GetRenderTargetResource(),
+		Context,
+		CaptureExtent,
+		ViewRotationMat,
+		ViewLocation,
+		OrthoProj);
+
+	FTextureRenderTargetResource* TextureRenderTargetResource = TextureRenderTarget->GameThread_GetRenderTargetResource();
+
+	FUpdateWaterInfoParams Params;
+	Params.DepthRenderer = DepthRenderer;
+	Params.ColorRenderer = ColorRenderer;
+	Params.RenderTarget = TextureRenderTargetResource;
+	Params.OutputTexture = TextureRenderTargetResource;
+	Params.CaptureZ = ViewLocation.Z;
+	Params.WaterHeightExtents = Context.ZoneToRender->GetWaterHeightExtents();
+	Params.GroundZMin = Context.ZoneToRender->GetGroundZMin();
+	Params.VelocityBlurRadius = Context.ZoneToRender->GetVelocityBlurRadius();
+
+	ENQUEUE_RENDER_COMMAND(WaterInfoCommand)(
+	[Params, ZoneName = Context.ZoneToRender->GetActorNameOrLabel()](FRHICommandListImmediate& RHICmdList)
+		{
+			SCOPED_DRAW_EVENTF(RHICmdList, WaterZoneInfoRendering_RT, TEXT("RenderWaterInfo_%s"), *ZoneName);
+
+			UpdateWaterInfoRendering_RenderThread(RHICmdList, Params);
+		});
+}
+
+} // namespace WaterInfo
\ No newline at end of file
diff --git a/Engine/Plugins/Experimental/Water/Source/Runtime/Private/WaterInfoRendering.h b/Engine/Plugins/Experimental/Water/Source/Runtime/Private/WaterInfoRendering.h
new file mode 100644
--- /dev/null
+++ b/Engine/Plugins/Experimental/Water/Source/Runtime/Private/WaterInfoRendering.h
@@ -0,0 +1,22 @@
+// Copyright Epic Games, Inc. All Rights Reserved.
+
+#pragma once
+
+class AWaterZone;
+class UWaterBodyComponent;
+
+namespace WaterInfo
+{
+	struct FRenderingContext
+	{
+		AWaterZone* ZoneToRender = nullptr;
+		TArray<UWaterBodyComponent*> WaterBodies;
+		TArray<AActor*> GroundActors;
+		float CaptureZ;
+	};
+	
+void UpdateWaterInfoRendering(
+	FSceneInterface* Scene,
+	UTextureRenderTarget2D* TextureRenderTarget,
+	const FRenderingContext& Context);
+}
diff --git a/Engine/Plugins/Experimental/Water/Source/Runtime/Private/WaterMeshComponent.cpp b/Engine/Plugins/Experimental/Water/Source/Runtime/Private/WaterMeshComponent.cpp
--- a/Engine/Plugins/Experimental/Water/Source/Runtime/Private/WaterMeshComponent.cpp
+++ b/Engine/Plugins/Experimental/Water/Source/Runtime/Private/WaterMeshComponent.cpp
@@ -361,7 +361,7 @@ void UWaterMeshComponent::RebuildWaterMesh(float InTileSize, const FIntPoint& In
 			const int32 NumOriginalSplinePoints = SplineComp->GetNumberOfSplinePoints();
 			float ConstantZ = SplineComp->GetLocationAtDistanceAlongSpline(0.0f, ESplineCoordinateSpace::World).Z;
 
-			FBox LakeBounds = Actor->GetComponentsBoundingBox();
+			FBox LakeBounds = Actor->GetComponentsBoundingBox(/* bNonColliding = */true);
 			LakeBounds.Max.Z += WaterBodyComponent->GetMaxWaveHeight();
 
 			// Skip lakes with less than 3 spline points
@@ -543,12 +543,6 @@ void UWaterMeshComponent::Update()
 	}
 }
 
-void UWaterMeshComponent::SetLandscapeInfo(const FVector& InRTWorldLocation, const FVector& InRTWorldSizeVector)
-{
-	RTWorldLocation = InRTWorldLocation;
-	RTWorldSizeVector = InRTWorldSizeVector;
-}
-
 #if WITH_EDITOR
 void UWaterMeshComponent::PostEditChangeProperty(FPropertyChangedEvent& PropertyChangedEvent)
 {
diff --git a/Engine/Plugins/Experimental/Water/Source/Runtime/Private/WaterRuntimeSettings.cpp b/Engine/Plugins/Experimental/Water/Source/Runtime/Private/WaterRuntimeSettings.cpp
--- a/Engine/Plugins/Experimental/Water/Source/Runtime/Private/WaterRuntimeSettings.cpp
+++ b/Engine/Plugins/Experimental/Water/Source/Runtime/Private/WaterRuntimeSettings.cpp
@@ -16,6 +16,7 @@ UWaterRuntimeSettings::FOnUpdateSettings UWaterRuntimeSettings::OnSettingsChange
 UWaterRuntimeSettings::UWaterRuntimeSettings()
 	: MaterialParameterCollection(FSoftObjectPath(TEXT("/Water/Materials/MPC/MPC_Water.MPC_Water")))
 	, DefaultWaterCollisionProfileName(TEXT("WaterBodyCollision"))
+	, DefaultWaterInfoMaterial(FSoftObjectPath(TEXT("/Water/Materials/WaterInfo/DrawWaterInfo.DrawWaterInfo")))
 	, WaterBodyRiverComponentClass(UWaterBodyRiverComponent::StaticClass())
 	, WaterBodyLakeComponentClass(UWaterBodyLakeComponent::StaticClass())
 	, WaterBodyOceanComponentClass(UWaterBodyOceanComponent::StaticClass())
@@ -29,6 +30,11 @@ FName UWaterRuntimeSettings::GetCategoryName() const
 	return FName(TEXT("Plugins"));
 }
 
+UMaterialInterface* UWaterRuntimeSettings::GetDefaultWaterInfoMaterial() const
+{
+	return DefaultWaterInfoMaterial.LoadSynchronous();
+}
+
 void UWaterRuntimeSettings::PostInitProperties()
 {
 	Super::PostInitProperties();
diff --git a/Engine/Plugins/Experimental/Water/Source/Runtime/Private/WaterSubsystem.cpp b/Engine/Plugins/Experimental/Water/Source/Runtime/Private/WaterSubsystem.cpp
--- a/Engine/Plugins/Experimental/Water/Source/Runtime/Private/WaterSubsystem.cpp
+++ b/Engine/Plugins/Experimental/Water/Source/Runtime/Private/WaterSubsystem.cpp
@@ -457,21 +457,21 @@ void UWaterSubsystem::SetOceanFloodHeight(float InFloodHeight)
 {
 	if (UWorld* World = GetWorld())
 	{
-	const float ClampedFloodHeight = FMath::Max(0.0f, InFloodHeight);
+		const float ClampedFloodHeight = FMath::Max(0.0f, InFloodHeight);
 
-	if (FloodHeight != ClampedFloodHeight)
-	{
-		FloodHeight = ClampedFloodHeight;
-		MarkAllWaterMeshesForRebuild();
+		if (FloodHeight != ClampedFloodHeight)
+		{
+			FloodHeight = ClampedFloodHeight;
+			MarkAllWaterZonesForRebuild();
 
-		// the ocean body is dynamic and needs to be readjusted when the flood height changes : 
+			// the ocean body is dynamic and needs to be readjusted when the flood height changes : 
 			if (OceanBodyComponent.IsValid())
-		{
+			{
 				OceanBodyComponent->SetHeightOffset(InFloodHeight);
-		}
+			}
 
 			WaterBodyManager.ForEachWaterBodyComponent([this](UWaterBodyComponent* WaterBodyComponent)
-		{
+			{
 				WaterBodyComponent->UpdateMaterialInstances();
 				return true;
 			});
@@ -512,16 +512,16 @@ float UWaterSubsystem::GetOceanBaseHeight() const
 	return TNumericLimits<float>::Lowest();
 }
 
-void UWaterSubsystem::MarkAllWaterMeshesForRebuild()
+void UWaterSubsystem::MarkAllWaterZonesForRebuild()
 {
 	if (UWorld* World = GetWorld())
 	{
-		for (AWaterZone* WaterMesh : TActorRange<AWaterZone>(World))
-	{
-		WaterMesh->MarkWaterMeshComponentForRebuild();
+		for (AWaterZone* WaterZone : TActorRange<AWaterZone>(World))
+		{
+			WaterZone->MarkForRebuild(EWaterZoneRebuildFlags::All);
+		}
 	}
 }
-}
 
 void UWaterSubsystem::NotifyWaterScalabilityChangedInternal(IConsoleVariable* CVar)
 {
@@ -813,4 +813,4 @@ UWaterSubsystem::FScopedAllowWaterSubsystemOnPreviewWorld::~FScopedAllowWaterSub
 
 // ----------------------------------------------------------------------------------
 
-#undef LOCTEXT_NAMESPACE
+#undef LOCTEXT_NAMESPACE
\ No newline at end of file
diff --git a/Engine/Plugins/Experimental/Water/Source/Runtime/Private/WaterUtils.cpp b/Engine/Plugins/Experimental/Water/Source/Runtime/Private/WaterUtils.cpp
--- a/Engine/Plugins/Experimental/Water/Source/Runtime/Private/WaterUtils.cpp
+++ b/Engine/Plugins/Experimental/Water/Source/Runtime/Private/WaterUtils.cpp
@@ -8,6 +8,12 @@ extern TAutoConsoleVariable<int32> CVarWaterEnabled;
 extern TAutoConsoleVariable<int32> CVarWaterMeshEnabled;
 extern TAutoConsoleVariable<int32> CVarWaterMeshEnableRendering;
 
+static TAutoConsoleVariable<float> CVarWaterMaxFlowVelocity(
+	TEXT("r.Water.MaxFlowVelocity"),
+	1024.0f,
+	TEXT("The maximum magnitude for the velocity of a river to encode in the WaterInfo texture"),
+	ECVF_Default);
+
 UMaterialInstanceDynamic* FWaterUtils::GetOrCreateTransientMID(UMaterialInstanceDynamic* InMID, FName InMIDName, UMaterialInterface* InMaterialInterface, EObjectFlags InAdditionalObjectFlags)
 {
 	if (!IsValid(InMaterialInterface))
@@ -111,3 +117,8 @@ bool FWaterUtils::IsWaterMeshRenderingEnabled(bool bIsRenderThread)
 {
 	return IsWaterMeshEnabled(bIsRenderThread) && !!(bIsRenderThread ? CVarWaterMeshEnableRendering.GetValueOnRenderThread() : CVarWaterMeshEnableRendering.GetValueOnGameThread());
 }
+
+float FWaterUtils::GetWaterMaxFlowVelocity(bool bIsRenderThread)
+{
+	return (bIsRenderThread ? CVarWaterMaxFlowVelocity.GetValueOnRenderThread() : CVarWaterMaxFlowVelocity.GetValueOnGameThread());
+}
diff --git a/Engine/Plugins/Experimental/Water/Source/Runtime/Private/WaterZoneActor.cpp b/Engine/Plugins/Experimental/Water/Source/Runtime/Private/WaterZoneActor.cpp
--- a/Engine/Plugins/Experimental/Water/Source/Runtime/Private/WaterZoneActor.cpp
+++ b/Engine/Plugins/Experimental/Water/Source/Runtime/Private/WaterZoneActor.cpp
@@ -1,32 +1,60 @@
 // Copyright Epic Games, Inc. All Rights Reserved.
 
 #include "WaterZoneActor.h"
+#include "WaterModule.h"
 #include "WaterSubsystem.h"
 #include "WaterMeshComponent.h"
 #include "WaterBodyActor.h"
-#include "Components/BoxComponent.h"
+#include "WaterInfoRendering.h"
+#include "EngineUtils.h"
+#include "LandscapeProxy.h"
+#include "WaterUtils.h"
+#include "Materials/MaterialInstanceDynamic.h"
+#include "UObject/FortniteMainBranchObjectVersion.h"
+#include "WaterBodyOceanComponent.h"
+#include "RenderCaptureInterface.h"
 
 #if	WITH_EDITOR
 #include "Algo/Transform.h"
 #include "LevelEditor.h"
 #include "Modules/ModuleManager.h"
 #include "WaterIconHelper.h"
+#include "Components/BoxComponent.h"
 #endif // WITH_EDITOR
 
+static int32 RenderCaptureNextWaterInfoDraws = 0;
+static FAutoConsoleVariableRef CVarRenderCaptureNextWaterInfoDraws(
+	TEXT("r.Water.WaterInfo.RenderCaptureNextWaterInfoDraws"),
+	RenderCaptureNextWaterInfoDraws,
+	TEXT("Enable capturing of the water info texture for the next N draws"));
+
+static int32 ForceUpdateWaterInfoNextFrames = 0;
+static FAutoConsoleVariableRef CVarForceUpdateWaterInfoNextFrames(
+	TEXT("r.Water.WaterInfo.ForceUpdateWaterInfoNextFrames"),
+	ForceUpdateWaterInfoNextFrames,
+	TEXT("Force the water info texture to regenerate on the next N frames. A negative value will force update every frame."));
+
+
 AWaterZone::AWaterZone(const FObjectInitializer& Initializer)
+	: Super(Initializer)
+	, RenderTargetResolution(512, 512)
 {
 	WaterMesh = CreateDefaultSubobject<UWaterMeshComponent>(TEXT("WaterMesh"));
 	SetRootComponent(WaterMesh);
+	ZoneExtent = FVector2D(51200.f, 51200.f);
+	
+#if	WITH_EDITOR
+	// Setup bounds component
+	{
+		BoundsComponent = CreateDefaultSubobject<UBoxComponent>(TEXT("BoundsComponent"));
+		BoundsComponent->SetCollisionObjectType(ECC_WorldStatic);
+		BoundsComponent->SetCollisionResponseToAllChannels(ECR_Ignore);
+		BoundsComponent->SetGenerateOverlapEvents(false);
+		BoundsComponent->SetupAttachment(WaterMesh);
+		// Bounds component extent is half-extent, ZoneExtent is full extent.
+		BoundsComponent->SetBoxExtent(FVector(ZoneExtent / 2.f, 8192.f));
+	}
 
-	BoundsComponent = CreateDefaultSubobject<UBoxComponent>(TEXT("BoundsComponent"));
-	BoundsComponent->SetCollisionObjectType(ECC_WorldStatic);
-	BoundsComponent->SetCollisionResponseToAllChannels(ECR_Ignore);
-	BoundsComponent->SetCollisionEnabled(ECollisionEnabled::NoCollision);
-	BoundsComponent->SetGenerateOverlapEvents(false);
-	BoundsComponent->SetupAttachment(WaterMesh);
-	BoundsComponent->SetBoxExtent(FVector(WaterMesh->GetExtentInTiles() * WaterMesh->GetTileSize(), 51200.f));
-
-#if WITH_EDITOR
 	if (GIsEditor && !IsTemplate())
 	{
 		FLevelEditorModule& LevelEditorModule = FModuleManager::LoadModuleChecked<FLevelEditorModule>("LevelEditor");
@@ -34,7 +62,26 @@ AWaterZone::AWaterZone(const FObjectInitializer& Initializer)
 	}
 
 	ActorIcon = FWaterIconHelper::EnsureSpriteComponentCreated(this, TEXT("/Water/Icons/WaterZoneActorSprite"));
-#endif
+#endif // WITH_EDITOR
+}
+
+void AWaterZone::SetZoneExtent(FVector2D NewExtent)
+{
+	ZoneExtent = NewExtent;
+	OnExtentChanged();
+}
+
+void AWaterZone::SetRenderTargetResolution(FIntPoint NewResolution)
+{
+	RenderTargetResolution = NewResolution;
+	MarkForRebuild(EWaterZoneRebuildFlags::UpdateWaterInfoTexture);
+}
+
+void AWaterZone::BeginPlay()
+{
+	Super::BeginPlay();
+
+	MarkForRebuild(EWaterZoneRebuildFlags::All);
 }
 
 void AWaterZone::PostLoadSubobjects(FObjectInstancingGraph* OuterInstanceGraph)
@@ -45,43 +92,86 @@ void AWaterZone::PostLoadSubobjects(FObjectInstancingGraph* OuterInstanceGraph)
 	Super::PostLoadSubobjects(OuterInstanceGraph);
 }
 
-void AWaterZone::MarkWaterMeshComponentForRebuild()
+void AWaterZone::PostLoad()
 {
-	WaterMesh->MarkWaterMeshGridDirty();
+	Super::PostLoad();
+
+#if WITH_EDITORONLY_DATA
+	if (GetLinkerCustomVersion(FFortniteMainBranchObjectVersion::GUID) < FFortniteMainBranchObjectVersion::WaterZonesRefactor)
+	{
+		// Divide by two since the ZoneExtent represents the radius of the box while the result of the Extent * Size gives the total size
+		const FVector2D ExtentInTiles = WaterMesh->GetExtentInTiles();
+		ZoneExtent = FVector2D(ExtentInTiles * WaterMesh->GetTileSize());
+		OnExtentChanged();
+	}
+#endif // WITH_EDITORONLY_DATA
 }
 
-void AWaterZone::Update()
+void AWaterZone::MarkForRebuild(EWaterZoneRebuildFlags Flags)
 {
-	if (WaterMesh)
+	if (EnumHasAnyFlags(Flags, EWaterZoneRebuildFlags::UpdateWaterMesh))
 	{
-		WaterMesh->Update();
+		WaterMesh->MarkWaterMeshGridDirty();
+	}
+	if (EnumHasAnyFlags(Flags, EWaterZoneRebuildFlags::UpdateWaterInfoTexture))
+	{
+		bNeedsWaterInfoRebuild = true;
 	}
 }
 
-void AWaterZone::SetLandscapeInfo(const FVector& InRTWorldLocation, const FVector& InRTWorldSizeVector)
+void AWaterZone::Update()
 {
+	if (bNeedsWaterInfoRebuild || (ForceUpdateWaterInfoNextFrames != 0))
+	{
+		ForceUpdateWaterInfoNextFrames = (ForceUpdateWaterInfoNextFrames < 0) ? ForceUpdateWaterInfoNextFrames : FMath::Max(0, ForceUpdateWaterInfoNextFrames - 1);
+		UpdateWaterInfoTexture();
+		bNeedsWaterInfoRebuild = false;
+	}
+	
 	if (WaterMesh)
 	{
-		WaterMesh->SetLandscapeInfo(InRTWorldLocation, InRTWorldSizeVector);
+		WaterMesh->Update();
 	}
 }
 
 #if	WITH_EDITOR
+void AWaterZone::ForceUpdateWaterInfoTexture()
+{
+	UpdateWaterInfoTexture();
+}
+
 void AWaterZone::PostEditMove(bool bFinished)
 {
 	Super::PostEditMove(bFinished);
 
 	// Ensure that the water mesh is rebuilt if it moves
-	MarkWaterMeshComponentForRebuild();
+	MarkForRebuild(EWaterZoneRebuildFlags::All);
 }
 
 void AWaterZone::PostEditChangeProperty(FPropertyChangedEvent& PropertyChangedEvent)
 {
 	Super::PostEditChangeProperty(PropertyChangedEvent);
 
-	if (PropertyChangedEvent.GetPropertyName() == GET_MEMBER_NAME_CHECKED(AWaterZone, BoundsComponent))
+	if (PropertyChangedEvent.MemberProperty != nullptr &&
+		PropertyChangedEvent.MemberProperty->GetFName() == GET_MEMBER_NAME_CHECKED(AWaterZone, ZoneExtent))
 	{
-		OnBoundsChanged();
+		OnExtentChanged();
+	}
+	else if (PropertyChangedEvent.GetPropertyName() == GET_MEMBER_NAME_CHECKED(AWaterZone, BoundsComponent))
+	{
+		OnBoundsComponentModified();
+	}
+	else if (PropertyChangedEvent.GetPropertyName() == GET_MEMBER_NAME_CHECKED(AWaterZone, RenderTargetResolution))
+	{
+		MarkForRebuild(EWaterZoneRebuildFlags::UpdateWaterInfoTexture);
+	}
+	else if (PropertyChangedEvent.GetPropertyName() == GET_MEMBER_NAME_CHECKED(AWaterZone, bHalfPrecisionTexture))
+	{
+		MarkForRebuild(EWaterZoneRebuildFlags::UpdateWaterInfoTexture);
+	}
+	else if (PropertyChangedEvent.GetPropertyName() == GET_MEMBER_NAME_CHECKED(AWaterZone, VelocityBlurRadius))
+	{
+		MarkForRebuild(EWaterZoneRebuildFlags::UpdateWaterInfoTexture);
 	}
 }
 
@@ -98,21 +188,102 @@ void AWaterZone::OnActorSelectionChanged(const TArray<UObject*>& NewSelection, b
 	if (SelectedWaterBodies != NewWeakWaterBodiesSelection)
 	{
 		SelectedWaterBodies = NewWeakWaterBodiesSelection;
-		MarkWaterMeshComponentForRebuild();
+		MarkForRebuild(EWaterZoneRebuildFlags::All);
 	}
 }
 #endif // WITH_EDITOR
 
-void AWaterZone::OnBoundsChanged()
+void AWaterZone::OnExtentChanged()
 {
+	// Compute the new tile extent based on the new bounds
 	const float MeshTileSize = WaterMesh->GetTileSize();
 
-	// Compute the new tile extent based on the new bounds
-	const FVector NewBounds = BoundsComponent->GetUnscaledBoxExtent();
+	int32 NewExtentInTilesX = FMath::FloorToInt(ZoneExtent.X / MeshTileSize);
+	int32 NewExtentInTilesY = FMath::FloorToInt(ZoneExtent.Y / MeshTileSize);
+	
+	// We must ensure that the zone is always at least 1x1
+	NewExtentInTilesX = FMath::Max(1, NewExtentInTilesX);
+	NewExtentInTilesY = FMath::Max(1, NewExtentInTilesY);
+
+	WaterMesh->SetExtentInTiles(FIntPoint(NewExtentInTilesX, NewExtentInTilesY));
+
+#if WITH_EDITOR
+	// Bounds component extent is half-extent, ZoneExtent is full extent.
+	BoundsComponent->SetBoxExtent(FVector(ZoneExtent / 2.f, 8192.f));
+#endif // WITH_EDITOR
+
+	MarkForRebuild(EWaterZoneRebuildFlags::All);
+}
+
+
+#if WITH_EDITOR
+void AWaterZone::OnBoundsComponentModified()
+{
+	const FVector2D NewBounds = FVector2D(BoundsComponent->GetUnscaledBoxExtent());
+
+	SetZoneExtent(NewBounds);
+}
+#endif // WITH_EDITOR
+
+void AWaterZone::UpdateWaterInfoTexture()
+{
+	if (UWorld* World = GetWorld(); World && FApp::CanEverRender())
+	{
+		RenderCaptureInterface::FScopedCapture RenderCapture((RenderCaptureNextWaterInfoDraws != 0), TEXT("RenderWaterInfo"));
+		RenderCaptureNextWaterInfoDraws = FMath::Max(0, RenderCaptureNextWaterInfoDraws - 1);
+
+		float WaterZMin(TNumericLimits<double>::Max());
+		float WaterZMax(TNumericLimits<double>::Lowest());
+	
+		// #todo_water [roey]: we need to store which actors this zone is responsible for rendering once we implement support for multi zones.
+		// For now whenever we update the water info texture we will just collect all water bodies and pass those to the renderer.
+		TArray<UWaterBodyComponent*> WaterBodies;
+		UWaterSubsystem::ForEachWaterBodyComponent(World, [&WaterBodies, &WaterZMax, &WaterZMin](UWaterBodyComponent* Component)
+		{
+			WaterBodies.Add(Component);
+			const FBox WaterBodyBounds = Component->CalcBounds(Component->GetComponentToWorld()).GetBox();
+			WaterZMax = FMath::Max(WaterZMax, WaterBodyBounds.Max.Z);
+			WaterZMin = FMath::Min(WaterZMin, WaterBodyBounds.Min.Z);
+			return true;
+		});
 
-	const int32 NewTileExtentX = FMath::FloorToInt(NewBounds.X / MeshTileSize);
-	const int32 NewTileExtentY = FMath::FloorToInt(NewBounds.Y / MeshTileSize);
+		// If we don't have any water bodies we don't need to do anything.
+		if (WaterBodies.Num() == 0)
+		{
+			return;
+		}
 
-	WaterMesh->SetExtentInTiles(FIntPoint(NewTileExtentX, NewTileExtentY));
-	MarkWaterMeshComponentForRebuild();
-}
\ No newline at end of file
+		WaterZMax += CaptureZOffset;
+
+		checkf(WaterZMin != WaterZMax, TEXT("Water has a height extent of 0 which can cause divide by zero errors in the shaders! Ensure water bodies have proper bounding boxes."));
+		WaterHeightExtents = FVector2f(WaterZMin, WaterZMax);
+
+		// Only compute the ground min since we can use the water max z as the ground max z for more precision.
+		GroundZMin = TNumericLimits<double>::Max();
+
+		TArray<AActor*> GroundActors;
+		for (ALandscapeProxy* LandscapeProxy : TActorRange<ALandscapeProxy>(World))
+		{
+			const FBox LandscapeBox = LandscapeProxy->GetComponentsBoundingBox();
+			GroundZMin = FMath::Min(GroundZMin, LandscapeBox.Min.Z);
+			GroundActors.Add(LandscapeProxy);
+		}
+
+		WaterInfo::FRenderingContext Context;
+		Context.ZoneToRender = this;
+		Context.WaterBodies = WaterBodies;
+		Context.GroundActors = MoveTemp(GroundActors);
+		Context.CaptureZ = WaterZMax;
+
+		const ETextureRenderTargetFormat Format = bHalfPrecisionTexture ? ETextureRenderTargetFormat::RTF_RGBA16f : RTF_RGBA32f;
+		WaterInfoTexture = FWaterUtils::GetOrCreateTransientRenderTarget2D(WaterInfoTexture, TEXT("WaterInfoTexture"), RenderTargetResolution, Format);
+		WaterInfo::UpdateWaterInfoRendering(World->Scene, WaterInfoTexture, Context);
+
+		for (UWaterBodyComponent* Component : WaterBodies)
+		{
+			Component->UpdateMaterialInstances();
+		}
+
+		UE_LOG(LogWater, Log, TEXT("Queued Water Info texture update"));
+	}
+}
diff --git a/Engine/Plugins/Experimental/Water/Source/Runtime/Public/WaterBodyComponent.h b/Engine/Plugins/Experimental/Water/Source/Runtime/Public/WaterBodyComponent.h
--- a/Engine/Plugins/Experimental/Water/Source/Runtime/Public/WaterBodyComponent.h
+++ b/Engine/Plugins/Experimental/Water/Source/Runtime/Public/WaterBodyComponent.h
@@ -21,6 +21,7 @@
 class UWaterSplineComponent;
 class AWaterBodyIsland;
 class AWaterBodyExclusionVolume;
+class AWaterZone;
 class ALandscapeProxy;
 class UMaterialInstanceDynamic;
 class FTokenizedMessage;
@@ -120,9 +121,11 @@ public:
 	/** Returns the type of body */
 	virtual EWaterBodyType GetWaterBodyType() const PURE_VIRTUAL(UWaterBodyComponent::GetWaterBodyType, return EWaterBodyType::Transition; )
 
-	/** Returns collision extents (For internal use. Please use AWaterBodyOcean instead.) */
+	/** Returns collision half-extents */
 	virtual FVector GetCollisionExtents() const { return FVector::ZeroVector; }
 
+	virtual FBoxSphereBounds CalcBounds(const FTransform& LocalToWorld) const override;
+
 	/** Sets an additional water height (For internal use. Please use AWaterBodyOcean instead.) */
 	virtual void SetHeightOffset(float InHeightOffset) { check(false); }
 
@@ -174,6 +177,10 @@ public:
 	/** Returns under water post process MID */
 	UFUNCTION(BlueprintCallable, Category = Rendering)
 	UMaterialInstanceDynamic* GetUnderwaterPostProcessMaterialInstance();
+	
+	/** Returns water info MID */
+	UFUNCTION(BlueprintCallable, Category = Rendering)
+	UMaterialInstanceDynamic* GetWaterInfoMaterialInstance();
 
 	/** Sets under water post process material */
 	void SetUnderwaterPostProcessMaterial(UMaterialInterface* InMaterial);
@@ -271,6 +278,9 @@ public:
 	/** Sets the dynamic parameters needed by the underwater post process material instance for rendering. Returns true if the operation was successfull */
 	virtual bool SetDynamicParametersOnUnderwaterPostProcessMID(UMaterialInstanceDynamic* InMID);
 
+	/** Sets the dynamic parameters needed by the material instance for rendering the water info texture. Returns true if the operation was successfull */
+	virtual bool SetDynamicParametersOnWaterInfoMID(UMaterialInstanceDynamic* InMID);
+
 	/** Returns true if the location is within one of this water body's exclusion volumes */
 	bool IsWorldLocationInExclusionVolume(const FVector& InWorldLocation) const;
 
@@ -288,6 +298,8 @@ public:
 	/** Returns the minimum and maximum Z of the water surface, including waves */
 	virtual void GetSurfaceMinMaxZ(float& OutMinZ, float& OutMaxZ) const;
 
+	float GetShapeDilationZOffset() const { return ShapeDilationZOffset; }
+
 	virtual ALandscapeProxy* FindLandscape() const;
 
 	/** Returns what can be considered the single water depth of the water surface.
@@ -304,6 +316,11 @@ public:
 
 	virtual void Reset() {}
 
+	/** Gets the water zone to which this component belongs */
+	AWaterZone* GetWaterZone() const;
+
+	virtual FPrimitiveSceneProxy* CreateSceneProxy() override;
+	virtual void GetUsedMaterials(TArray<UMaterialInterface*>& OutMaterials, bool bGetDebugMaterials) const override;
 protected:
 	//~ Begin USceneComponent Interface.
 	virtual void OnVisibilityChanged();
@@ -359,6 +376,7 @@ protected:
 	bool UpdateWaterHeight();
 	void CreateOrUpdateWaterMID();
 	void CreateOrUpdateUnderwaterPostProcessMID();
+	void CreateOrUpdateWaterInfoMID();
 	void PrepareCurrentPostProcessSettings();
 	void ApplyNavigationSettings();
 	void RequestGPUWaveDataUpdate();
@@ -391,12 +409,17 @@ public:
 
 	/** Public static constants : */
 	static const FName WaterBodyIndexParamName;
+	static const FName WaterBodyZOffsetParamName;
 	static const FName WaterVelocityAndHeightName;
 	static const FName GlobalOceanHeightName;
 	static const FName FixedZHeightName;
 	static const FName FixedVelocityName;
 	static const FName FixedWaterDepthName;
 	static const FName WaterAreaParamName;
+	static const FName MaxFlowVelocityParamName;
+	static const FName WaterZMinParamName;
+	static const FName WaterZMaxParamName;
+	static const FName GroundZMinParamName;
 
 	UPROPERTY(EditDefaultsOnly, Category = Collision, meta = (EditCondition = "bGenerateCollisions"))
 	UPhysicalMaterial* PhysicalMaterial;
@@ -428,10 +451,13 @@ public:
 	UPROPERTY(Category = Rendering, EditAnywhere, BlueprintReadOnly, meta = (EditCondition = "bGenerateCollisions", DisplayAfter = "WaterMaterial"))
 	UMaterialInterface* UnderwaterPostProcessMaterial;
 
-#if WITH_EDITORONLY_DATA
+	UPROPERTY(Category = Rendering, EditAnywhere, BlueprintReadOnly, meta = (DisplayAfter = "WaterMaterial"))
+	UMaterialInterface* WaterInfoMaterial;
+	
 	UPROPERTY(Category = Terrain, EditAnywhere, BlueprintReadWrite)
 	FWaterBodyHeightmapSettings WaterHeightmapSettings;
 
+#if WITH_EDITORONLY_DATA
 	UPROPERTY(Category = Terrain, EditAnywhere, BlueprintReadWrite)
 	TMap<FName, FWaterBodyWeightmapSettings> LayerWeightmapSettings;
 #endif
@@ -475,6 +501,9 @@ protected:
 
 	UPROPERTY(Category = Debug, VisibleInstanceOnly, Transient, NonPIEDuplicateTransient, TextExportTransient, meta = (DisplayAfter = "UnderwaterPostProcessMaterial"))
 	UMaterialInstanceDynamic* UnderwaterPostProcessMID;
+	
+	UPROPERTY(Category = Debug, VisibleInstanceOnly, Transient, NonPIEDuplicateTransient, TextExportTransient, meta = (DisplayAfter = "WaterInfoMaterial"))
+	UMaterialInstanceDynamic* WaterInfoMID;
 
 	/** Islands in this water body*/
 	UPROPERTY(Category = Water, EditInstanceOnly, AdvancedDisplay)
@@ -496,6 +525,10 @@ protected:
 	UPROPERTY(EditAnywhere, BlueprintReadWrite, Category = Navigation, meta = (EditCondition = "bCanAffectNavigation && bGenerateCollisions"))
 	TSubclassOf<UNavAreaBase> WaterNavAreaClass;
 
+	/** Push down the dilated region to prevent overwriting adjacent water body data */
+	UPROPERTY(EditAnywhere, BlueprintReadWrite, AdvancedDisplay, Category = Rendering)
+	float ShapeDilationZOffset = -64.f;
+
 #if WITH_EDITORONLY_DATA
 	UPROPERTY()
 	bool bOverrideWaterMesh_DEPRECATED;
diff --git a/Engine/Plugins/Experimental/Water/Source/Runtime/Public/WaterBodyCustomComponent.h b/Engine/Plugins/Experimental/Water/Source/Runtime/Public/WaterBodyCustomComponent.h
--- a/Engine/Plugins/Experimental/Water/Source/Runtime/Public/WaterBodyCustomComponent.h
+++ b/Engine/Plugins/Experimental/Water/Source/Runtime/Public/WaterBodyCustomComponent.h
@@ -27,6 +27,7 @@ protected:
 	virtual void BeginUpdateWaterBody() override;
 	virtual void OnUpdateBody(bool bWithExclusionVolumes) override;
 
+	virtual FPrimitiveSceneProxy* CreateSceneProxy() override;
 #if WITH_EDITOR
 	virtual TArray<TSharedRef<FTokenizedMessage>> CheckWaterBodyStatus() const override;
 #endif // WITH_EDITOR
diff --git a/Engine/Plugins/Experimental/Water/Source/Runtime/Public/WaterBodyLakeComponent.h b/Engine/Plugins/Experimental/Water/Source/Runtime/Public/WaterBodyLakeComponent.h
--- a/Engine/Plugins/Experimental/Water/Source/Runtime/Public/WaterBodyLakeComponent.h
+++ b/Engine/Plugins/Experimental/Water/Source/Runtime/Public/WaterBodyLakeComponent.h
@@ -24,6 +24,8 @@ public:
 	virtual TArray<UPrimitiveComponent*> GetCollisionComponents() const override;
 	virtual TArray<UPrimitiveComponent*> GetStandardRenderableComponents() const override;
 
+	virtual FBoxSphereBounds CalcBounds(const FTransform& LocalToWorld) const override;
+
 protected:
 	/** UWaterBodyComponent Interface */
 	virtual void Reset() override;
diff --git a/Engine/Plugins/Experimental/Water/Source/Runtime/Public/WaterBodyOceanComponent.h b/Engine/Plugins/Experimental/Water/Source/Runtime/Public/WaterBodyOceanComponent.h
--- a/Engine/Plugins/Experimental/Water/Source/Runtime/Public/WaterBodyOceanComponent.h
+++ b/Engine/Plugins/Experimental/Water/Source/Runtime/Public/WaterBodyOceanComponent.h
@@ -24,12 +24,19 @@ public:
 	virtual FVector GetCollisionExtents() const override { return CollisionExtents; }
 	virtual void SetHeightOffset(float InHeightOffset) override;
 	virtual float GetHeightOffset() const override { return HeightOffset; }
+
+	void SetVisualExtents(FVector2D NewExtents);
+	FVector2D GetVisualExtents() const { return VisualExtents; }
 protected:
 	/** UWaterBodyComponent Interface */
 	virtual bool IsBodyDynamic() const override { return true; }
 	virtual void BeginUpdateWaterBody() override;
 	virtual void OnUpdateBody(bool bWithExclusionVolumes) override;
 	virtual void Reset() override;
+
+	virtual void PostLoad() override;
+	virtual FBoxSphereBounds CalcBounds(const FTransform& LocalToWorld) const;
+
 #if WITH_EDITOR
 	virtual void OnPostEditChangeProperty(FPropertyChangedEvent& PropertyChangedEvent, bool& bShapeOrPositionChanged, bool& bWeightmapSettingsChanged) override;
 #endif
@@ -40,6 +47,10 @@ protected:
 	UPROPERTY(NonPIEDuplicateTransient)
 	TArray<UOceanCollisionComponent*> CollisionHullSets;
 
+	/** The area over which the ocean should be displayed, centered on the actor */
+	UPROPERTY(Category = Water, EditAnywhere, BlueprintReadOnly)
+	FVector2D VisualExtents;
+
 	UPROPERTY(Category = Collision, EditAnywhere, BlueprintReadOnly)
 	FVector CollisionExtents;
 
diff --git a/Engine/Plugins/Experimental/Water/Source/Runtime/Public/WaterBodyRiverComponent.h b/Engine/Plugins/Experimental/Water/Source/Runtime/Public/WaterBodyRiverComponent.h
--- a/Engine/Plugins/Experimental/Water/Source/Runtime/Public/WaterBodyRiverComponent.h
+++ b/Engine/Plugins/Experimental/Water/Source/Runtime/Public/WaterBodyRiverComponent.h
@@ -31,11 +31,15 @@ public:
 	void SetLakeTransitionMaterial(UMaterialInterface* InMat);
 	void SetOceanTransitionMaterial(UMaterialInterface* InMat);
 
+	float GetShapeDilationZOffsetFar() const { return ShapeDilationZOffsetFar; }
+
 protected:
 	/** UWaterBodyComponent Interface */
 	virtual void Reset() override;
 	virtual void UpdateMaterialInstances() override;
 	virtual void OnUpdateBody(bool bWithExclusionVolumes) override;
+
+	virtual FBoxSphereBounds CalcBounds(const FTransform& LocalToWorld) const override;
 #if WITH_EDITOR
 	virtual void OnPostEditChangeProperty(FPropertyChangedEvent& PropertyChangedEvent, bool& bShapeOrPositionChanged, bool& bWeightmapSettingsChanged) override;
 #endif
@@ -63,4 +67,8 @@ protected:
 
 	UPROPERTY(Category = Debug, VisibleInstanceOnly, Transient, NonPIEDuplicateTransient, TextExportTransient, meta = (DisplayAfter = "OceanTransitionMaterial"))
 	UMaterialInstanceDynamic* OceanTransitionMID;
+
+	/** How far to push down the furthest vertices of the dilated portion. */
+	UPROPERTY(EditAnywhere, BlueprintReadWrite, AdvancedDisplay, Category = Rendering)
+	float ShapeDilationZOffsetFar = -128.f;
 };
\ No newline at end of file
diff --git a/Engine/Plugins/Experimental/Water/Source/Runtime/Public/WaterBodySceneProxy.h b/Engine/Plugins/Experimental/Water/Source/Runtime/Public/WaterBodySceneProxy.h
new file mode 100644
--- /dev/null
+++ b/Engine/Plugins/Experimental/Water/Source/Runtime/Public/WaterBodySceneProxy.h
@@ -0,0 +1,48 @@
+// Copyright Epic Games, Inc. All Rights Reserved.
+
+#pragma once
+
+#include "CoreMinimal.h"
+#include "Engine/Scene.h"
+#include "DynamicMesh/DynamicMesh3.h"
+#include "StaticMeshResources.h"
+#include "DynamicMeshBuilder.h"
+
+class UWaterBodyComponent;
+class UWaterBodyLakeComponent;
+class UWaterBodyOceanComponent;
+class UWaterBodyRiverComponent;
+class FWaterBodyMeshSection;
+class FMaterialRenderProxy;
+
+class FWaterBodySceneProxy final : public FPrimitiveSceneProxy
+{
+public:
+	FWaterBodySceneProxy(UWaterBodyComponent* Component);
+	virtual ~FWaterBodySceneProxy();
+	virtual void GetDynamicMeshElements(const TArray<const FSceneView*>& Views, const FSceneViewFamily& ViewFamily, uint32 VisibilityMap, FMeshElementCollector& Collector) const override;
+	virtual FPrimitiveViewRelevance GetViewRelevance(const FSceneView* View) const override;
+	virtual SIZE_T GetTypeHash() const override;
+	virtual uint32 GetMemoryFootprint() const override;
+	uint32 GetAllocatedSize() const;
+
+	bool IsShown(const FSceneView* View) const;
+
+	bool IsWithinWaterInfoPass() const { return bWithinWaterInfoPass; }
+	void SetWithinWaterInfoPass(bool bInWithinWaterInfoPass);
+
+private:
+	void GenerateLakeMesh(UWaterBodyLakeComponent* Component);
+	void GenerateOceanMesh(UWaterBodyOceanComponent* Component);
+	void GenerateRiverMesh(UWaterBodyRiverComponent* Component);
+
+	void AddSectionFromDynamicMesh(const UE::Geometry::FDynamicMesh3& DynamicMesh, TFunctionRef<FDynamicMeshVertex(const FVector3f&)> VertTransformFunc);
+
+	void InitResources(FWaterBodyMeshSection* Section);
+
+	TArray<FWaterBodyMeshSection*> Sections;
+	FMaterialRenderProxy* Material;
+
+	bool bWithinWaterInfoPass = false;
+};
+
diff --git a/Engine/Plugins/Experimental/Water/Source/Runtime/Public/WaterMeshComponent.h b/Engine/Plugins/Experimental/Water/Source/Runtime/Public/WaterMeshComponent.h
--- a/Engine/Plugins/Experimental/Water/Source/Runtime/Public/WaterMeshComponent.h
+++ b/Engine/Plugins/Experimental/Water/Source/Runtime/Public/WaterMeshComponent.h
@@ -77,17 +77,6 @@ public:
 
 	UFUNCTION(BlueprintPure, Category = Mesh)
 	bool IsEnabled() const { return bIsEnabled; }
-
-	// HACK [jonathan.bard] (start) : This is to make sure that the RTWorldLocation / RTWorldSizeVector MPC params can be serialized and set at runtime on the Water MPC.
-	void SetLandscapeInfo(const FVector& InRTWorldLocation, const FVector& InRTWorldSizeVector);
-
-	UPROPERTY(VisibleAnywhere, BlueprintReadOnly, Category = Texture)
-	FVector RTWorldLocation;
-
-	UPROPERTY(VisibleAnywhere, BlueprintReadOnly, Category = Texture)
-	FVector RTWorldSizeVector;
-	// HACK [jonathan.bard] (end)
-
 private:
 	//~ Begin USceneComponent Interface
 	virtual FBoxSphereBounds CalcBounds(const FTransform& LocalToWorld) const override;
diff --git a/Engine/Plugins/Experimental/Water/Source/Runtime/Public/WaterRuntimeSettings.h b/Engine/Plugins/Experimental/Water/Source/Runtime/Public/WaterRuntimeSettings.h
--- a/Engine/Plugins/Experimental/Water/Source/Runtime/Public/WaterRuntimeSettings.h
+++ b/Engine/Plugins/Experimental/Water/Source/Runtime/Public/WaterRuntimeSettings.h
@@ -31,6 +31,8 @@ public:
 
 	FName GetDefaultWaterCollisionProfileName() const { return DefaultWaterCollisionProfileName; }
 
+	UMaterialInterface* GetDefaultWaterInfoMaterial() const;
+
 	TSubclassOf<UWaterBodyRiverComponent> GetWaterBodyRiverComponentClass() const;
 
 	TSubclassOf<UWaterBodyLakeComponent> GetWaterBodyLakeComponentClass() const;
@@ -71,6 +73,9 @@ private:
 	UPROPERTY(VisibleAnywhere, config, Category = Collision)
 	FName DefaultWaterCollisionProfileName;
 
+	UPROPERTY(EditAnywhere, config, Category = Water)
+	TSoftObjectPtr<UMaterialInterface> DefaultWaterInfoMaterial;
+	
 	UPROPERTY(EditAnywhere, Config, Category = Water,  meta = (MetaClass = "WaterBodyRiverComponent"))
 	TSubclassOf<UWaterBodyRiverComponent> WaterBodyRiverComponentClass;
 
diff --git a/Engine/Plugins/Experimental/Water/Source/Runtime/Public/WaterSubsystem.h b/Engine/Plugins/Experimental/Water/Source/Runtime/Public/WaterSubsystem.h
--- a/Engine/Plugins/Experimental/Water/Source/Runtime/Public/WaterSubsystem.h
+++ b/Engine/Plugins/Experimental/Water/Source/Runtime/Public/WaterSubsystem.h
@@ -165,7 +165,7 @@ public:
 
 	UMaterialParameterCollection* GetMaterialParameterCollection() const {	return MaterialParameterCollection; }
 	
-	void MarkAllWaterMeshesForRebuild();
+	void MarkAllWaterZonesForRebuild();
 
 #if WITH_EDITOR
 	/** Little scope object to temporarily change the value of bAllowWaterSubsystemOnPreviewWorld */
diff --git a/Engine/Plugins/Experimental/Water/Source/Runtime/Public/WaterUtils.h b/Engine/Plugins/Experimental/Water/Source/Runtime/Public/WaterUtils.h
--- a/Engine/Plugins/Experimental/Water/Source/Runtime/Public/WaterUtils.h
+++ b/Engine/Plugins/Experimental/Water/Source/Runtime/Public/WaterUtils.h
@@ -42,4 +42,5 @@ struct FWaterUtils
 	static WATER_API bool IsWaterEnabled(bool bIsRenderThread);
 	static WATER_API bool IsWaterMeshEnabled(bool bIsRenderThread);
 	static WATER_API bool IsWaterMeshRenderingEnabled(bool bIsRenderThread);
+	static WATER_API float GetWaterMaxFlowVelocity(bool bIsRenderThread);
 };
diff --git a/Engine/Plugins/Experimental/Water/Source/Runtime/Public/WaterZoneActor.h b/Engine/Plugins/Experimental/Water/Source/Runtime/Public/WaterZoneActor.h
--- a/Engine/Plugins/Experimental/Water/Source/Runtime/Public/WaterZoneActor.h
+++ b/Engine/Plugins/Experimental/Water/Source/Runtime/Public/WaterZoneActor.h
@@ -3,61 +3,112 @@
 #pragma once
 
 #include "CoreMinimal.h"
+#include "Components/BoxComponent.h"
 #include "GameFramework/Actor.h"
+#include "Engine/TextureRenderTarget2D.h"
 #include "WaterZoneActor.generated.h"
 
 class UWaterMeshComponent;
-class USceneCaptureComponent2D;
 class UBoxComponent;
+class AWaterBody;
+
+enum class EWaterZoneRebuildFlags
+{
+	None = 0,
+	UpdateWaterInfoTexture = (1 << 1),
+	UpdateWaterMesh = (1 << 2),
+	All = (~0),
+};
+ENUM_CLASS_FLAGS(EWaterZoneRebuildFlags);
 
 UCLASS(Blueprintable)
 class WATER_API AWaterZone : public AActor
 {
 	GENERATED_UCLASS_BODY()
-
 public:
-
 	UWaterMeshComponent* GetWaterMeshComponent() { return WaterMesh; }
 	const UWaterMeshComponent* GetWaterMeshComponent() const { return WaterMesh; }
 
-	void MarkWaterMeshComponentForRebuild();
+	void MarkForRebuild(EWaterZoneRebuildFlags Flags);
 	void Update();
 
-	UPROPERTY(VisibleAnywhere, BlueprintReadOnly, Category = Texture)
-	UTexture2D* WaterVelocityTexture;
+	FVector2D GetZoneExtent() const { return ZoneExtent; }
+	void SetZoneExtent(FVector2D NewExtents);
 
-	// HACK [jonathan.bard] : See UWaterMeshComponent, r : emove ASAP
-	void SetLandscapeInfo(const FVector& InRTWorldLocation, const FVector& InRTWorldSizeVector);
+	void SetRenderTargetResolution(FIntPoint NewResolution);
+	FIntPoint GetRenderTargetResolution() const { return RenderTargetResolution; }
 
-#if WITH_EDITOR
-	virtual void PostEditMove(bool bFinished) override;
-
-	virtual void PostEditChangeProperty(FPropertyChangedEvent& PropertyChangedEvent) override;
-#endif
+	uint32 GetVelocityBlurRadius() const { return VelocityBlurRadius; }
 
+	virtual void BeginPlay() override;
 	virtual void PostLoadSubobjects(FObjectInstancingGraph* OuterInstanceGraph) override;
+	virtual void PostLoad() override;
+
+	FVector2f GetWaterHeightExtents() const { return WaterHeightExtents; }
+	float GetGroundZMin() const { return GroundZMin; }
+
+	UPROPERTY(Transient, DuplicateTransient, VisibleAnywhere, BlueprintReadOnly, Category = Texture)
+	UTextureRenderTarget2D* WaterInfoTexture;
 private:
 
+	void UpdateWaterInfoTexture();
+
+	void OnExtentChanged();
+	
 #if WITH_EDITOR
 	void OnActorSelectionChanged(const TArray<UObject*>& NewSelection, bool bForceRefresh);
-#endif
+	
+	UFUNCTION(CallInEditor, Category = Debug)
+	void ForceUpdateWaterInfoTexture();
+
+	virtual void PostEditMove(bool bFinished) override;
 
-	void OnBoundsChanged();
+	virtual void PostEditChangeProperty(FPropertyChangedEvent& PropertyChangedEvent) override;
 
-	/** A manipulatable box for visualizing/editing the water zone bounds: */
-	UPROPERTY(Category = Bounds, VisibleAnywhere, BlueprintReadOnly, meta = (AllowPrivateAccess = "true"))
-	UBoxComponent* BoundsComponent;
+	/** Called when the Bounds component is modified. Updates the value of ZoneExtent to match the new bounds */
+	void OnBoundsComponentModified();
+#endif // WITH_EDITOR
+	
+	UPROPERTY(EditAnywhere, BlueprintReadOnly, Category = Texture, meta = (AllowPrivateAccess = "true"))
+	FIntPoint RenderTargetResolution;
+
+	/** The water mesh component */
+	UPROPERTY(VisibleAnywhere, Category = Water, BlueprintReadOnly, meta = (AllowPrivateAccess = "true"))
+	UWaterMeshComponent* WaterMesh;
+
+	/** Radius of the zone bounding box */
+	UPROPERTY(Category = Shape, EditAnywhere, BlueprintReadWrite, meta = (AllowPrivateAccess = "true"))
+	FVector2D ZoneExtent;
+
+	/** Offsets the height above the water zone at which the WaterInfoTexture is rendered. This is applied after computing the maximum Z of all the water bodies within the zone. */
+	UPROPERTY(Category = Shape, EditAnywhere, meta = (AllowPrivateAccess = "true"))
+	float CaptureZOffset = 64.f;
+
+	/** Determines if the WaterInfoTexture should be 16 or 32 bits per channel */
+	UPROPERTY(Category = Texture, EditAnywhere, AdvancedDisplay, meta = (AllowPrivateAccess = "true"))
+	bool bHalfPrecisionTexture = true;
+
+	/** Radius of the velocity blur in the finalize water info pass */
+	UPROPERTY(Category = Texture, EditAnywhere, AdvancedDisplay)
+	int32 VelocityBlurRadius = 1;
+
+	bool bNeedsWaterInfoRebuild = true;
+
+	FVector2f WaterHeightExtents;
+	float GroundZMin;
 
 #if WITH_EDITORONLY_DATA
+	/** A manipulatable box for visualizing/editing the water zone bounds */
+	UPROPERTY(Transient)
+	UBoxComponent* BoundsComponent;
 
 	UPROPERTY(Transient)
-	TArray<TWeakObjectPtr<class AWaterBody>> SelectedWaterBodies;
+	TArray<TWeakObjectPtr<AWaterBody>> SelectedWaterBodies;
 
 	UPROPERTY(Transient)
 	UBillboardComponent* ActorIcon;
-#endif
 
-	/** The water mesh component */
-	UPROPERTY(VisibleAnywhere, Category = Water, BlueprintReadOnly, meta = (AllowPrivateAccess = "true"))
-	UWaterMeshComponent* WaterMesh;
-};
+	UPROPERTY(Transient)
+	UTexture2D* WaterVelocityTexture_DEPRECATED;
+#endif // WITH_EDITORONLY_DATA
+};
\ No newline at end of file
diff --git a/Engine/Plugins/Experimental/Water/Source/Runtime/Water.Build.cs b/Engine/Plugins/Experimental/Water/Source/Runtime/Water.Build.cs
--- a/Engine/Plugins/Experimental/Water/Source/Runtime/Water.Build.cs
+++ b/Engine/Plugins/Experimental/Water/Source/Runtime/Water.Build.cs
@@ -25,12 +25,14 @@ public class Water : ModuleRules
 		PrivateDependencyModuleNames.AddRange(
 			new string[] {
 				"RenderCore",
+				"Renderer",
 				"NiagaraCore",
 				"Niagara",
 				"NiagaraShader",
 				"Projects",
 				"Landscape",
 				"GeometryCore",
+				"GeometryAlgorithms",
 				"DynamicMesh",
 				"ChaosCore",
 				"Chaos",
diff --git a/Engine/Source/Runtime/Core/Public/UObject/FortniteMainBranchObjectVersion.h b/Engine/Source/Runtime/Core/Public/UObject/FortniteMainBranchObjectVersion.h
--- a/Engine/Source/Runtime/Core/Public/UObject/FortniteMainBranchObjectVersion.h
+++ b/Engine/Source/Runtime/Core/Public/UObject/FortniteMainBranchObjectVersion.h
@@ -193,6 +193,9 @@ struct CORE_API FFortniteMainBranchObjectVersion
 		// Anim Dynamics Node Physics parameters for each body in a chain are now stored in an array and can be edited.
 		AnimDynamicsEditableChainParameters,
 
+		// Decoupled the generation of the water texture from the Water Brush and the landscape
+		WaterZonesRefactor,
+
 		// -----<new versions can be added above this line>-------------------------------------------------
 		VersionPlusOne,
 		LatestVersion = VersionPlusOne - 1
diff --git a/Engine/Source/Runtime/Engine/Classes/Components/SplineComponent.h b/Engine/Source/Runtime/Engine/Classes/Components/SplineComponent.h
--- a/Engine/Source/Runtime/Engine/Classes/Components/SplineComponent.h
+++ b/Engine/Source/Runtime/Engine/Classes/Components/SplineComponent.h
@@ -756,6 +756,11 @@ public:
 	UFUNCTION(BlueprintCallable, Category = Spline)
 	bool DivideSplineIntoPolylineRecursive(float StartDistanceAlongSpline, float EndDistanceAlongSpline, ESplineCoordinateSpace::Type CoordinateSpace, const float MaxSquareDistanceFromSpline, TArray<FVector>& OutPoints) const;
 
+	/** Given a threshold, recursively sub-divides the spline section until the list of segments (polyline) matches the spline shape. */
+	UFUNCTION(BlueprintCallable, Category = Spline)
+	bool DivideSplineIntoPolylineRecursiveWithDistances(float StartDistanceAlongSpline, float EndDistanceAlongSpline, ESplineCoordinateSpace::Type CoordinateSpace, const float MaxSquareDistanceFromSpline, TArray<FVector>& OutPoints, TArray<double>& OutDistancesAlongSpline) const;
+
+
 	/** Given a threshold, returns a list of vertices along the spline segment that, treated as a list of segments (polyline), matches the spline shape. */
 	UFUNCTION(BlueprintCallable, Category = Spline)
 	bool ConvertSplineSegmentToPolyLine(int32 SplinePointStartIndex, ESplineCoordinateSpace::Type CoordinateSpace, const float MaxSquareDistanceFromSpline, TArray<FVector>& OutPoints) const;
diff --git a/Engine/Source/Runtime/Engine/Private/Components/SplineComponent.cpp b/Engine/Source/Runtime/Engine/Private/Components/SplineComponent.cpp
--- a/Engine/Source/Runtime/Engine/Private/Components/SplineComponent.cpp
+++ b/Engine/Source/Runtime/Engine/Private/Components/SplineComponent.cpp
@@ -1089,7 +1089,9 @@ float USplineComponent::GetDistanceAlongSplineAtSplinePoint(int32 PointIndex) co
 	const int32 NumPoints = SplineCurves.Position.Points.Num();
 	const int32 NumSegments = bClosedLoop ? NumPoints : NumPoints - 1;
 
-	if ((PointIndex >= 0) && (PointIndex < NumSegments + 1))
+	// Ensure that if the reparam table is not prepared yet we don't attempt to access it. This can happen
+	// early in the construction of the spline component object.
+	if ((PointIndex >= 0) && (PointIndex < NumSegments + 1) && ((PointIndex * ReparamStepsPerSegment) < SplineCurves.ReparamTable.Points.Num()))
 	{
 		return SplineCurves.ReparamTable.Points[PointIndex * ReparamStepsPerSegment].InVal;
 	}
@@ -1515,7 +1517,7 @@ FTransform USplineComponent::FindTransformClosestToWorldLocation(const FVector&
 	return GetTransformAtSplineInputKey(Param, CoordinateSpace, bUseScale);
 }
 
-bool USplineComponent::DivideSplineIntoPolylineRecursive(float StartDistanceAlongSpline, float EndDistanceAlongSpline, ESplineCoordinateSpace::Type CoordinateSpace, const float MaxSquareDistanceFromSpline, TArray<FVector>& OutPoints) const
+bool USplineComponent::DivideSplineIntoPolylineRecursiveWithDistances(float StartDistanceAlongSpline, float EndDistanceAlongSpline, ESplineCoordinateSpace::Type CoordinateSpace, const float MaxSquareDistanceFromSpline, TArray<FVector>& OutPoints, TArray<double>& OutDistancesAlongSpline) const
 {
 	double Dist = EndDistanceAlongSpline - StartDistanceAlongSpline;
 	if (Dist <= 0.0f)
@@ -1523,7 +1525,6 @@ bool USplineComponent::DivideSplineIntoPolylineRecursive(float StartDistanceAlon
 		return false;
 	}
 	double MiddlePointDistancAlongSpline = StartDistanceAlongSpline + Dist / 2.0f;
-
 	FVector Samples[3];
 	Samples[0] = GetLocationAtDistanceAlongSpline(StartDistanceAlongSpline, CoordinateSpace);
 	Samples[1] = GetLocationAtDistanceAlongSpline(MiddlePointDistancAlongSpline, CoordinateSpace);
@@ -1532,26 +1533,40 @@ bool USplineComponent::DivideSplineIntoPolylineRecursive(float StartDistanceAlon
 	if (FMath::PointDistToSegmentSquared(Samples[1], Samples[0], Samples[2]) > MaxSquareDistanceFromSpline)
 	{
 		TArray<FVector> NewPoints[2];
-		DivideSplineIntoPolylineRecursive(StartDistanceAlongSpline, MiddlePointDistancAlongSpline, CoordinateSpace, MaxSquareDistanceFromSpline, NewPoints[0]);
-		DivideSplineIntoPolylineRecursive(MiddlePointDistancAlongSpline, EndDistanceAlongSpline, CoordinateSpace, MaxSquareDistanceFromSpline, NewPoints[1]);
+		TArray<double> NewDistancesAlongSpline[2];
+		DivideSplineIntoPolylineRecursiveWithDistances(StartDistanceAlongSpline, MiddlePointDistancAlongSpline, CoordinateSpace, MaxSquareDistanceFromSpline, NewPoints[0], NewDistancesAlongSpline[0]);
+		DivideSplineIntoPolylineRecursiveWithDistances(MiddlePointDistancAlongSpline, EndDistanceAlongSpline, CoordinateSpace, MaxSquareDistanceFromSpline, NewPoints[1], NewDistancesAlongSpline[1]);
 		if ((NewPoints[0].Num() > 0) && (NewPoints[1].Num() > 0))
 		{
 			check(NewPoints[0].Last() == NewPoints[1][0]);
+			check(NewDistancesAlongSpline[0].Last() == NewDistancesAlongSpline[1][0]);
 			NewPoints[0].RemoveAt(NewPoints[0].Num() - 1);
+			NewDistancesAlongSpline[0].RemoveAt(NewDistancesAlongSpline[0].Num() - 1);
 		}
 		NewPoints[0].Append(NewPoints[1]);
+		NewDistancesAlongSpline[0].Append(NewDistancesAlongSpline[1]);
 		OutPoints.Append(NewPoints[0]);
+		OutDistancesAlongSpline.Append(NewDistancesAlongSpline[0]);
 	}
 	else
 	{
 		// The middle point is close enough to the other 2 points, let's keep those and stop the recursion :
 		OutPoints.Add(Samples[0]);
 		OutPoints.Add(Samples[2]);
+		OutDistancesAlongSpline.Add(StartDistanceAlongSpline);
+		OutDistancesAlongSpline.Add(EndDistanceAlongSpline);
 	}
 
+	check(OutPoints.Num() == OutDistancesAlongSpline.Num())
 	return (OutPoints.Num() > 0);
 }
 
+bool USplineComponent::DivideSplineIntoPolylineRecursive(float StartDistanceAlongSpline, float EndDistanceAlongSpline, ESplineCoordinateSpace::Type CoordinateSpace, const float MaxSquareDistanceFromSpline, TArray<FVector>& OutPoints) const
+{
+	TArray<double> DummyDistancesAlongSpline;
+	return DivideSplineIntoPolylineRecursiveWithDistances(StartDistanceAlongSpline, EndDistanceAlongSpline, CoordinateSpace, MaxSquareDistanceFromSpline, OutPoints, DummyDistancesAlongSpline);
+}
+
 bool USplineComponent::ConvertSplineSegmentToPolyLine(int32 SplinePointStartIndex, ESplineCoordinateSpace::Type CoordinateSpace, const float MaxSquareDistanceFromSpline, TArray<FVector>& OutPoints) const
 {
 	OutPoints.Empty();
@@ -1590,7 +1605,6 @@ bool USplineComponent::ConvertSplineSegmentToPolyLine(int32 SplinePointStartInde
 	return (OutPoints.Num() > 0);
 }
 
-
 bool USplineComponent::ConvertSplineToPolyLine(ESplineCoordinateSpace::Type CoordinateSpace, const float MaxSquareDistanceFromSpline, TArray<FVector>& OutPoints) const
 {
 	int32 NumSegments = GetNumberOfSplineSegments();
diff --git a/Engine/Source/Runtime/Renderer/Private/SceneCaptureRendering.h b/Engine/Source/Runtime/Renderer/Private/SceneCaptureRendering.h
--- a/Engine/Source/Runtime/Renderer/Private/SceneCaptureRendering.h
+++ b/Engine/Source/Runtime/Renderer/Private/SceneCaptureRendering.h
@@ -10,7 +10,7 @@
 #include "RenderGraphDefinitions.h"
 #include "SceneTextures.h"
 
-void CopySceneCaptureComponentToTarget(
+RENDERER_API void CopySceneCaptureComponentToTarget(
 	FRDGBuilder& GraphBuilder,
 	const FMinimalSceneTextures& SceneTextures,
 	FRDGTextureRef ViewFamilyTexture,
diff --git a/Engine/Source/Runtime/Renderer/Private/SceneRendering.h b/Engine/Source/Runtime/Renderer/Private/SceneRendering.h
--- a/Engine/Source/Runtime/Renderer/Private/SceneRendering.h
+++ b/Engine/Source/Runtime/Renderer/Private/SceneRendering.h
@@ -1895,18 +1895,18 @@ public:
 	virtual ~FSceneRenderer();
 
 	// Initializes the scene renderer on the render thread.
-	void RenderThreadBegin(FRHICommandListImmediate& RHICmdList);
-	void RenderThreadEnd(FRHICommandListImmediate& RHICmdList);
+	RENDERER_API void RenderThreadBegin(FRHICommandListImmediate& RHICmdList);
+	RENDERER_API void RenderThreadEnd(FRHICommandListImmediate& RHICmdList);
 
 	// FSceneRenderer interface
 
-	virtual void Render(FRDGBuilder& GraphBuilder) = 0;
+	RENDERER_API virtual void Render(FRDGBuilder& GraphBuilder) = 0;
 	virtual void RenderHitProxies(FRDGBuilder& GraphBuilder) {}
 	virtual bool ShouldRenderVelocities() const { return false; }
 	virtual bool ShouldRenderPrePass() const { return false; }
 
 	/** Creates a scene renderer based on the current feature level. */
-	static FSceneRenderer* CreateSceneRenderer(const FSceneViewFamily* InViewFamily, FHitProxyConsumer* HitProxyConsumer);
+	RENDERER_API static FSceneRenderer* CreateSceneRenderer(const FSceneViewFamily* InViewFamily, FHitProxyConsumer* HitProxyConsumer);
 
 	/** Setups FViewInfo::ViewRect according to ViewFamilly's ScreenPercentageInterface. */
 	void PrepareViewRectsForRendering(FRHICommandListImmediate& RHICmdList);
@@ -1947,7 +1947,7 @@ public:
 	*
 	* @param SceneRenderer	Scene renderer to use for rendering.
 	*/
-	static void ViewExtensionPreRender_RenderThread(FRHICommandListImmediate& RHICmdList, FSceneRenderer* SceneRenderer);
+	static void RENDERER_API ViewExtensionPreRender_RenderThread(FRHICommandListImmediate& RHICmdList, FSceneRenderer* SceneRenderer);
 
 	/** the last thing we do with a scene renderer, lots of cleanup related to the threading **/
 	void WaitForTasksAndClearSnapshots(FParallelMeshDrawCommandPass::EWaitThread WaitThread);
