Fixed non-unity build

diff --git a/Engine/Plugins/Experimental/Water/Source/Runtime/Private/WaterInfoRendering.h b/Engine/Plugins/Experimental/Water/Source/Runtime/Private/WaterInfoRendering.h
--- a/Engine/Plugins/Experimental/Water/Source/Runtime/Private/WaterInfoRendering.h
+++ b/Engine/Plugins/Experimental/Water/Source/Runtime/Private/WaterInfoRendering.h
@@ -2,8 +2,13 @@
 
 #pragma once
 
+#include "CoreMinimal.h"
+
 class AWaterZone;
 class UWaterBodyComponent;
+class FSceneInterface;
+class UTextureRenderTarget2D;
+class AActor;
 
 namespace WaterInfo
 {
diff --git a/Engine/Plugins/Experimental/Water/Source/Runtime/Private/WaterRuntimeSettings.cpp b/Engine/Plugins/Experimental/Water/Source/Runtime/Private/WaterRuntimeSettings.cpp
--- a/Engine/Plugins/Experimental/Water/Source/Runtime/Private/WaterRuntimeSettings.cpp
+++ b/Engine/Plugins/Experimental/Water/Source/Runtime/Private/WaterRuntimeSettings.cpp
@@ -2,6 +2,7 @@
 
 #include "WaterRuntimeSettings.h"
 #include "Materials/MaterialParameterCollection.h"
+#include "Materials/MaterialInterface.h"
 #include "WaterBodyComponent.h"
 #include "WaterBodyRiverComponent.h"
 #include "WaterBodyLakeComponent.h"
diff --git a/Engine/Plugins/Experimental/Water/Source/Runtime/Public/WaterRuntimeSettings.h b/Engine/Plugins/Experimental/Water/Source/Runtime/Public/WaterRuntimeSettings.h
--- a/Engine/Plugins/Experimental/Water/Source/Runtime/Public/WaterRuntimeSettings.h
+++ b/Engine/Plugins/Experimental/Water/Source/Runtime/Public/WaterRuntimeSettings.h
@@ -9,6 +9,7 @@
 #include "Engine/EngineTypes.h"
 #include "WaterRuntimeSettings.generated.h"
 
+class UMaterialInterface;
 class UMaterialParameterCollection;
 class UWaterBodyComponent;
 class UWaterBodyRiverComponent;
