Fixed todo comment appearing as a property tooltip

diff --git a/Engine/Plugins/Experimental/Water/Source/Runtime/Public/WaterBodyActor.h b/Engine/Plugins/Experimental/Water/Source/Runtime/Public/WaterBodyActor.h
--- a/Engine/Plugins/Experimental/Water/Source/Runtime/Public/WaterBodyActor.h
+++ b/Engine/Plugins/Experimental/Water/Source/Runtime/Public/WaterBodyActor.h
@@ -123,7 +123,7 @@ protected:
 	EWaterBodyType WaterBodyType;
 
 	// #todo_water: This should be moved to the component when component subobjects are supported
-	UPROPERTY(EditAnywhere, BlueprintReadOnly, Instanced, Category = Wave, DisplayName = "Waves Source")
+	UPROPERTY(EditAnywhere, BlueprintReadOnly, Instanced, Category = Wave, DisplayName = "Waves Source", meta = (Tooltip = ""))
 	UWaterWavesBase* WaterWaves = nullptr;
 
 #if WITH_EDITORONLY_DATA
