Updated water icons and integrated water body component icons into the builtin USceneComponent billboard component.

diff --git a/Engine/Plugins/Experimental/Water/Source/Editor/Private/WaterUIStyle.cpp b/Engine/Plugins/Experimental/Water/Source/Editor/Private/WaterUIStyle.cpp
--- a/Engine/Plugins/Experimental/Water/Source/Editor/Private/WaterUIStyle.cpp
+++ b/Engine/Plugins/Experimental/Water/Source/Editor/Private/WaterUIStyle.cpp
@@ -10,6 +10,7 @@
 
 
 #define IMAGE_BRUSH(RelativePath, ...) FSlateImageBrush(StyleSet->RootToContentDir(RelativePath, TEXT(".png")), __VA_ARGS__)
+#define IMAGE_BRUSH_SVG(RelativePath, ...) FSlateVectorImageBrush(StyleSet->RootToContentDir(RelativePath, TEXT(".svg")), __VA_ARGS__)
 #define BOX_BRUSH(RelativePath, ...) FSlateBoxBrush(StyleSet->RootToContentDir(RelativePath, TEXT(".png")), __VA_ARGS__)
 #define DEFAULT_FONT(...) FCoreStyle::GetDefaultFontStyle(__VA_ARGS__)
 
@@ -44,42 +45,42 @@ void FWaterUIStyle::Initialize()
 	StyleSet->SetCoreContentRoot(FPaths::EngineContentDir() / TEXT("Slate"));
 
 	// WaterBodyRiver
-	StyleSet->Set("ClassIcon.WaterBodyRiver", new IMAGE_BRUSH("Icons/WaterBodyRiver_16x", Icon16x16));
-	StyleSet->Set("ClassThumbnail.WaterBodyRiver", new IMAGE_BRUSH("Icons/WaterBodyRiver_64x", Icon64x64));
+	StyleSet->Set("ClassIcon.WaterBodyRiver", new IMAGE_BRUSH_SVG("Icons/WaterBodyRiver", Icon16x16));
+	StyleSet->Set("ClassThumbnail.WaterBodyRiver", new IMAGE_BRUSH_SVG("Icons/WaterBodyRiver_64", Icon64x64));
 	// WaterBodyOcean
-	StyleSet->Set("ClassIcon.WaterBodyOcean", new IMAGE_BRUSH("Icons/WaterBodyOcean_16x", Icon16x16));
-	StyleSet->Set("ClassThumbnail.WaterBodyOcean", new IMAGE_BRUSH("Icons/WaterBodyOcean_64x", Icon64x64));
+	StyleSet->Set("ClassIcon.WaterBodyOcean", new IMAGE_BRUSH_SVG("Icons/WaterBodyOcean", Icon16x16));
+	StyleSet->Set("ClassThumbnail.WaterBodyOcean", new IMAGE_BRUSH_SVG("Icons/WaterBodyOcean_64", Icon64x64));
 	// WaterBodyLake
-	StyleSet->Set("ClassIcon.WaterBodyLake", new IMAGE_BRUSH("Icons/WaterBodyLake_16x", Icon16x16));
-	StyleSet->Set("ClassThumbnail.WaterBodyLake", new IMAGE_BRUSH("Icons/WaterBodyLake_64x", Icon64x64));
+	StyleSet->Set("ClassIcon.WaterBodyLake", new IMAGE_BRUSH_SVG("Icons/WaterBodyLake", Icon16x16));
+	StyleSet->Set("ClassThumbnail.WaterBodyLake", new IMAGE_BRUSH_SVG("Icons/WaterBodyLake_64", Icon64x64));
 	// WaterBodyCustom
-	StyleSet->Set("ClassIcon.WaterBodyCustom", new IMAGE_BRUSH("Icons/WaterBodyCustom_16x", Icon16x16));
-	StyleSet->Set("ClassThumbnail.WaterBodyCustom", new IMAGE_BRUSH("Icons/WaterBodyCustom_64x", Icon64x64));
+	StyleSet->Set("ClassIcon.WaterBodyCustom", new IMAGE_BRUSH_SVG("Icons/WaterBodyCustom", Icon16x16));
+	StyleSet->Set("ClassThumbnail.WaterBodyCustom", new IMAGE_BRUSH_SVG("Icons/WaterBodyCustom_64", Icon64x64));
 	// WaterBodyIsland
-	StyleSet->Set("ClassIcon.WaterBodyIsland", new IMAGE_BRUSH("Icons/WaterBodyIsland_16x", Icon16x16));
-	StyleSet->Set("ClassThumbnail.WaterBodyIsland", new IMAGE_BRUSH("Icons/WaterBodyIsland_64x", Icon64x64));
+	StyleSet->Set("ClassIcon.WaterBodyIsland", new IMAGE_BRUSH_SVG("Icons/WaterBodyIsland", Icon16x16));
+	StyleSet->Set("ClassThumbnail.WaterBodyIsland", new IMAGE_BRUSH_SVG("Icons/WaterBodyIsland_64", Icon64x64));
 	// WaterBodyExclusionVolume
-	StyleSet->Set("ClassIcon.WaterBodyExclusionVolume", new IMAGE_BRUSH("Icons/WaterBodyExclusionVolume_16x", Icon16x16));
-	StyleSet->Set("ClassThumbnail.WaterBodyExclusionVolume", new IMAGE_BRUSH("Icons/WaterBodyExclusionVolume_64x", Icon64x64));
+	StyleSet->Set("ClassIcon.WaterBodyExclusionVolume", new IMAGE_BRUSH_SVG("Icons/WaterBodyExclusionVolume", Icon16x16));
+	StyleSet->Set("ClassThumbnail.WaterBodyExclusionVolume", new IMAGE_BRUSH_SVG("Icons/WaterBodyExclusionVolume_64", Icon64x64));
 	// WaterZone
-	StyleSet->Set("ClassIcon.WaterZone", new IMAGE_BRUSH("Icons/WaterZoneActor_16x", Icon16x16));
-	StyleSet->Set("ClassThumbnail.WaterZone", new IMAGE_BRUSH("Icons/WaterZoneActor_64x", Icon64x64));
+	StyleSet->Set("ClassIcon.WaterZone", new IMAGE_BRUSH_SVG("Icons/WaterZone", Icon16x16));
+	StyleSet->Set("ClassThumbnail.WaterZone", new IMAGE_BRUSH_SVG("Icons/WaterZone_64", Icon64x64));
 	// WaterLandscapeBrush
-	StyleSet->Set("ClassIcon.WaterLandscapeBrush", new IMAGE_BRUSH("Icons/WaterLandscapeBrush_16x", Icon16x16));
-	StyleSet->Set("ClassThumbnail.WaterLandscapeBrush", new IMAGE_BRUSH("Icons/WaterLandscapeBrush_64x", Icon64x64));
+	StyleSet->Set("ClassIcon.WaterLandscapeBrush", new IMAGE_BRUSH_SVG("Icons/WaterLandscapeBrush", Icon16x16));
+	StyleSet->Set("ClassThumbnail.WaterLandscapeBrush", new IMAGE_BRUSH_SVG("Icons/WaterLandscapeBrush_64", Icon64x64));
 	// WaterWaves
-	StyleSet->Set("ClassIcon.WaterWavesBase", new IMAGE_BRUSH("Icons/WaterWave_16x", Icon16x16));
-	StyleSet->Set("ClassThumbnail.WaterWavesBase", new IMAGE_BRUSH("Icons/WaterWave_64x", Icon64x64));
+	StyleSet->Set("ClassIcon.WaterWavesBase", new IMAGE_BRUSH_SVG("Icons/WaterWave", Icon16x16));
+	StyleSet->Set("ClassThumbnail.WaterWavesBase", new IMAGE_BRUSH_SVG("Icons/WaterWave_64", Icon64x64));
 	// WaterWavesAsset
-	StyleSet->Set("ClassIcon.WaterWavesAsset", new IMAGE_BRUSH("Icons/WaterWaveAsset_16x", Icon16x16));
-	StyleSet->Set("ClassThumbnail.WaterWavesAsset", new IMAGE_BRUSH("Icons/WaterWaveAsset_64x", Icon64x64));
+	StyleSet->Set("ClassIcon.WaterWavesAsset", new IMAGE_BRUSH_SVG("Icons/WaterWave", Icon16x16));
+	StyleSet->Set("ClassThumbnail.WaterWavesAsset", new IMAGE_BRUSH_SVG("Icons/WaterWave_64", Icon64x64));
 	// WaterWaves Asset Editor
 	StyleSet->Set("WaterWavesEditor.TogglePauseWaveTime", new IMAGE_BRUSH("Icons/PauseWaveTime_40x", Icon40x40));
 
 	FSlateStyleRegistry::RegisterSlateStyle(*StyleSet.Get());
 };
 
-#undef IMAGE_PLUGIN_BRUSH
+#undef IMAGE_BRUSH_SVG
 #undef IMAGE_BRUSH
 #undef BOX_BRUSH
 #undef DEFAULT_FONT
diff --git a/Engine/Plugins/Experimental/Water/Source/Runtime/Private/WaterBodyActor.cpp b/Engine/Plugins/Experimental/Water/Source/Runtime/Private/WaterBodyActor.cpp
--- a/Engine/Plugins/Experimental/Water/Source/Runtime/Private/WaterBodyActor.cpp
+++ b/Engine/Plugins/Experimental/Water/Source/Runtime/Private/WaterBodyActor.cpp
@@ -59,10 +59,6 @@ AWaterBody::AWaterBody(const FObjectInitializer& ObjectInitializer)
 	// Temporarily set the root component to the spline because the WaterBodyComponent has not yet been created
 	RootComponent = SplineComp;
 
-#if WITH_EDITOR
-	ActorIcon = FWaterIconHelper::EnsureSpriteComponentCreated(this, TEXT("/Water/Icons/WaterSprite"));
-#endif
-
 #if WITH_EDITORONLY_DATA
 	bAffectsLandscape_DEPRECATED = true;
 	CollisionProfileName_DEPRECATED = GetDefault<UWaterRuntimeSettings>()->GetDefaultWaterCollisionProfileName();
@@ -155,57 +151,6 @@ void AWaterBody::PostEditChangeProperty(FPropertyChangedEvent& PropertyChangedEv
 	}
 }
 
-void AWaterBody::UpdateActorIcon()
-{
-	if (ActorIcon && !bIsEditorPreviewActor)
-	{
-		// Actor icon gets in the way of meshes
-		ActorIcon->SetVisibility(IsIconVisible());
-
-		UTexture2D* IconTexture = ActorIcon->Sprite;
-		IWaterModuleInterface& WaterModule = FModuleManager::GetModuleChecked<IWaterModuleInterface>("Water");
-		if (const IWaterEditorServices* WaterEditorServices = WaterModule.GetWaterEditorServices())
-		{
-			bool bHasError = false;
-			if (WaterBodyComponent)
-			{
-				TArray<TSharedRef<FTokenizedMessage>> StatusMessages = WaterBodyComponent->CheckWaterBodyStatus();
-				for (const TSharedRef<FTokenizedMessage>& StatusMessage : StatusMessages)
-				{
-					// Message severities are ordered from most severe to least severe.
-					if (StatusMessage->GetSeverity() <= EMessageSeverity::Error)
-					{
-						bHasError = true;
-						break;
-					}
-				}
-			}
-
-			if (bHasError)
-			{
-				IconTexture = WaterEditorServices->GetErrorSprite();
-			}
-			else
-			{
-				IconTexture = WaterEditorServices->GetWaterActorSprite(GetClass());
-			}
-		}
-		FWaterIconHelper::UpdateSpriteComponent(this, IconTexture);
-
-		if (GetWaterBodyType() == EWaterBodyType::Lake && SplineComp)
-		{
-			// Move the actor icon to the center of the lake
-			const FVector ZOffset(0.0f, 0.0f, GetDefault<UWaterRuntimeSettings>()->WaterBodyIconWorldZOffset);
-			ActorIcon->SetWorldLocation(SplineComp->Bounds.Origin + ZOffset);
-		}
-	}
-}
-
-bool AWaterBody::IsIconVisible() const
-{
-	return GetWaterBodyType() != EWaterBodyType::Transition;
-}
-
 void AWaterBody::PostDuplicate(bool bDuplicateForPIE)
 {
 	Super::PostDuplicate(bDuplicateForPIE);
diff --git a/Engine/Plugins/Experimental/Water/Source/Runtime/Private/WaterBodyComponent.cpp b/Engine/Plugins/Experimental/Water/Source/Runtime/Private/WaterBodyComponent.cpp
--- a/Engine/Plugins/Experimental/Water/Source/Runtime/Private/WaterBodyComponent.cpp
+++ b/Engine/Plugins/Experimental/Water/Source/Runtime/Private/WaterBodyComponent.cpp
@@ -630,6 +630,11 @@ float UWaterBodyComponent::GetAudioIntensityAtSplineInputKey(float InKey) const
 
 void UWaterBodyComponent::OnRegister()
 {
+#if WITH_EDITORONLY_DATA
+	// Prevents USceneComponent from creating the SpriteComponent in OnRegister because we want to provide a different texture
+	bVisualizeComponent = false;
+#endif // WITH_EDITORONLY_DATA
+
 	Super::OnRegister();
 
 	AWaterBody* OwningWaterBodyActor = GetWaterBodyActor();
@@ -642,6 +647,8 @@ void UWaterBodyComponent::OnRegister()
 #if WITH_EDITOR
 	RegisterOnChangeWaterSplineMetadata(WaterSplineMetadata, /*bRegister = */true);
 	GetWaterSpline()->OnSplineDataChanged().AddUObject(this, &UWaterBodyComponent::OnSplineDataChanged);
+
+	CreateWaterSpriteComponent();
 #endif // WITH_EDITOR
 }
 
@@ -1159,7 +1166,7 @@ void UWaterBodyComponent::UpdateAll(bool bShapeOrPositionChanged)
 		UpdateComponentVisibility(/* bAllowWaterMeshRebuild = */true);
 
 #if WITH_EDITOR
-		WaterBodyOwner->UpdateActorIcon();
+		UpdateWaterSpriteComponent();
 #endif
 	}
 }
@@ -1412,8 +1419,6 @@ UWaterWavesBase* UWaterBodyComponent::GetWaterWaves() const
 	return nullptr;
 }
 
-#undef LOCTEXT_NAMESPACE
-
 AWaterZone* UWaterBodyComponent::GetWaterZone() const
 {
 	// #todo_water [roey]: Currently returns the global water zone. Eventually this function will return the specific water zone which encapsulates this water body. 
@@ -1424,3 +1429,67 @@ AWaterZone* UWaterBodyComponent::GetWaterZone() const
 	}
 	return nullptr;
 }
+
+#if WITH_EDITOR
+
+void UWaterBodyComponent::CreateWaterSpriteComponent()
+{
+	UTexture2D* Texture = LoadObject<UTexture2D>(nullptr, GetWaterSpriteTextureName());
+
+	IWaterModuleInterface& WaterModule = FModuleManager::GetModuleChecked<IWaterModuleInterface>(TEXT("Water"));
+	if (IWaterEditorServices* WaterEditorServices = WaterModule.GetWaterEditorServices())
+	{
+		WaterEditorServices->RegisterWaterActorSprite(GetClass(), Texture);
+	}
+
+	bVisualizeComponent = true;
+	CreateSpriteComponent(Texture);
+
+	SpriteComponent->SetRelativeScale3D(FVector(1.f, 1.f, 1.f));
+	SpriteComponent->SetRelativeLocation(FVector(0.f, 0.f,  GetDefault<UWaterRuntimeSettings>()->WaterBodyIconWorldZOffset));
+}
+
+void UWaterBodyComponent::UpdateWaterSpriteComponent()
+{
+	if (SpriteComponent)
+	{
+		SpriteComponent->SetVisibility(IsIconVisible());
+
+		UTexture2D* IconTexture = SpriteComponent->Sprite;
+		IWaterModuleInterface& WaterModule = FModuleManager::GetModuleChecked<IWaterModuleInterface>("Water");
+		if (const IWaterEditorServices* WaterEditorServices = WaterModule.GetWaterEditorServices())
+		{
+			bool bHasError = false;
+			TArray<TSharedRef<FTokenizedMessage>> StatusMessages = CheckWaterBodyStatus();
+			for (const TSharedRef<FTokenizedMessage>& StatusMessage : StatusMessages)
+			{
+				// Message severities are ordered from most severe to least severe.
+				if (StatusMessage->GetSeverity() <= EMessageSeverity::Error)
+				{
+					bHasError = true;
+					break;
+				}
+			}
+
+			if (bHasError)
+			{
+				IconTexture = WaterEditorServices->GetErrorSprite();
+			}
+			else
+			{
+				IconTexture = WaterEditorServices->GetWaterActorSprite(GetClass());
+			}
+		}
+
+
+		const FVector ZOffset(0.0f, 0.0f, GetDefault<UWaterRuntimeSettings>()->WaterBodyIconWorldZOffset);
+		SpriteComponent->SetWorldLocation(GetWaterSpriteLocation() + ZOffset);
+		SpriteComponent->Sprite = IconTexture;
+
+		SpriteComponent->MarkRenderStateDirty();
+	}
+}
+#endif // WITH_EDITOR
+
+#undef LOCTEXT_NAMESPACE
+
diff --git a/Engine/Plugins/Experimental/Water/Source/Runtime/Private/WaterBodyCustomActor.cpp b/Engine/Plugins/Experimental/Water/Source/Runtime/Private/WaterBodyCustomActor.cpp
--- a/Engine/Plugins/Experimental/Water/Source/Runtime/Private/WaterBodyCustomActor.cpp
+++ b/Engine/Plugins/Experimental/Water/Source/Runtime/Private/WaterBodyCustomActor.cpp
@@ -16,10 +16,6 @@ AWaterBodyCustom::AWaterBodyCustom(const FObjectInitializer& ObjectInitializer)
 {
 	WaterBodyType = EWaterBodyType::Transition;
 
-#if WITH_EDITOR
-	ActorIcon = FWaterIconHelper::EnsureSpriteComponentCreated(this, TEXT("/Water/Icons/WaterBodyCustomSprite"));
-#endif
-
 #if WITH_EDITORONLY_DATA
 	bAffectsLandscape_DEPRECATED = false;
 #endif
diff --git a/Engine/Plugins/Experimental/Water/Source/Runtime/Private/WaterBodyCustomComponent.cpp b/Engine/Plugins/Experimental/Water/Source/Runtime/Private/WaterBodyCustomComponent.cpp
--- a/Engine/Plugins/Experimental/Water/Source/Runtime/Private/WaterBodyCustomComponent.cpp
+++ b/Engine/Plugins/Experimental/Water/Source/Runtime/Private/WaterBodyCustomComponent.cpp
@@ -140,6 +140,10 @@ TArray<TSharedRef<FTokenizedMessage>> UWaterBodyCustomComponent::CheckWaterBodyS
 
 	return StatusMessages;
 }
+const TCHAR* UWaterBodyCustomComponent::GetWaterSpriteTextureName() const
+{
+	return TEXT("/Water/Icons/WaterBodyCustomComponent");
+}
 #endif // WITH_EDITOR
 
 #undef LOCTEXT_NAMESPACE
diff --git a/Engine/Plugins/Experimental/Water/Source/Runtime/Private/WaterBodyLakeActor.cpp b/Engine/Plugins/Experimental/Water/Source/Runtime/Private/WaterBodyLakeActor.cpp
--- a/Engine/Plugins/Experimental/Water/Source/Runtime/Private/WaterBodyLakeActor.cpp
+++ b/Engine/Plugins/Experimental/Water/Source/Runtime/Private/WaterBodyLakeActor.cpp
@@ -15,10 +15,6 @@ AWaterBodyLake::AWaterBodyLake(const FObjectInitializer& ObjectInitializer)
 	: Super(ObjectInitializer)
 {
 	WaterBodyType = EWaterBodyType::Lake;
-
-#if WITH_EDITOR
-	ActorIcon = FWaterIconHelper::EnsureSpriteComponentCreated(this, TEXT("/Water/Icons/WaterBodyLakeSprite"));
-#endif
 }
 
 void AWaterBodyLake::PostLoad()
diff --git a/Engine/Plugins/Experimental/Water/Source/Runtime/Private/WaterBodyLakeComponent.cpp b/Engine/Plugins/Experimental/Water/Source/Runtime/Private/WaterBodyLakeComponent.cpp
--- a/Engine/Plugins/Experimental/Water/Source/Runtime/Private/WaterBodyLakeComponent.cpp
+++ b/Engine/Plugins/Experimental/Water/Source/Runtime/Private/WaterBodyLakeComponent.cpp
@@ -151,4 +151,18 @@ void UWaterBodyLakeComponent::OnUpdateBody(bool bWithExclusionVolumes)
 			LakeCollision->UpdateCollision(LakeCollisionExtent, true);
 		}
 	}
-}
\ No newline at end of file
+}
+
+#if WITH_EDITOR
+
+const TCHAR* UWaterBodyLakeComponent::GetWaterSpriteTextureName() const
+{
+	return TEXT("/Water/Icons/WaterBodyLakeSprite");
+}
+
+FVector UWaterBodyLakeComponent::GetWaterSpriteLocation() const
+{
+	return GetWaterSpline()->Bounds.Origin;
+}
+
+#endif // WITH_EDITOR
diff --git a/Engine/Plugins/Experimental/Water/Source/Runtime/Private/WaterBodyOceanActor.cpp b/Engine/Plugins/Experimental/Water/Source/Runtime/Private/WaterBodyOceanActor.cpp
--- a/Engine/Plugins/Experimental/Water/Source/Runtime/Private/WaterBodyOceanActor.cpp
+++ b/Engine/Plugins/Experimental/Water/Source/Runtime/Private/WaterBodyOceanActor.cpp
@@ -14,10 +14,6 @@ AWaterBodyOcean::AWaterBodyOcean(const FObjectInitializer& ObjectInitializer)
 	: Super(ObjectInitializer)
 {
 	WaterBodyType = EWaterBodyType::Ocean;
-	
-#if WITH_EDITOR
-	ActorIcon = FWaterIconHelper::EnsureSpriteComponentCreated(this, TEXT("/Water/Icons/WaterBodyOceanSprite"));
-#endif
 
 #if WITH_EDITORONLY_DATA
 	CollisionExtents_DEPRECATED = FVector(50000.f, 50000.f, 10000.f);
diff --git a/Engine/Plugins/Experimental/Water/Source/Runtime/Private/WaterBodyOceanComponent.cpp b/Engine/Plugins/Experimental/Water/Source/Runtime/Private/WaterBodyOceanComponent.cpp
--- a/Engine/Plugins/Experimental/Water/Source/Runtime/Private/WaterBodyOceanComponent.cpp
+++ b/Engine/Plugins/Experimental/Water/Source/Runtime/Private/WaterBodyOceanComponent.cpp
@@ -96,6 +96,10 @@ void UWaterBodyOceanComponent::OnPostEditChangeProperty(FPropertyChangedEvent& P
 		bShapeOrPositionChanged = true;
 	}
 }
+const TCHAR* UWaterBodyOceanComponent::GetWaterSpriteTextureName() const
+{
+	return TEXT("/Water/Icons/WaterBodyOceanSprite");
+}
 #endif
 
 void UWaterBodyOceanComponent::Reset()
diff --git a/Engine/Plugins/Experimental/Water/Source/Runtime/Private/WaterBodyRiverActor.cpp b/Engine/Plugins/Experimental/Water/Source/Runtime/Private/WaterBodyRiverActor.cpp
--- a/Engine/Plugins/Experimental/Water/Source/Runtime/Private/WaterBodyRiverActor.cpp
+++ b/Engine/Plugins/Experimental/Water/Source/Runtime/Private/WaterBodyRiverActor.cpp
@@ -14,9 +14,6 @@ AWaterBodyRiver::AWaterBodyRiver(const FObjectInitializer& ObjectInitializer)
 	: Super(ObjectInitializer)
 {
 	WaterBodyType = EWaterBodyType::River;
-#if WITH_EDITOR
-	ActorIcon = FWaterIconHelper::EnsureSpriteComponentCreated(this, TEXT("/Water/Icons/WaterBodyRiverSprite"));
-#endif
 }
 
 void AWaterBodyRiver::PostLoad()
diff --git a/Engine/Plugins/Experimental/Water/Source/Runtime/Private/WaterBodyRiverComponent.cpp b/Engine/Plugins/Experimental/Water/Source/Runtime/Private/WaterBodyRiverComponent.cpp
--- a/Engine/Plugins/Experimental/Water/Source/Runtime/Private/WaterBodyRiverComponent.cpp
+++ b/Engine/Plugins/Experimental/Water/Source/Runtime/Private/WaterBodyRiverComponent.cpp
@@ -304,4 +304,8 @@ void UWaterBodyRiverComponent::OnPostEditChangeProperty(FPropertyChangedEvent& P
 		UpdateMaterialInstances();
 	}
 }
+const TCHAR* UWaterBodyRiverComponent::GetWaterSpriteTextureName() const
+{
+	return TEXT("/Water/Icons/WaterBodyRiverSprite");
+}
 #endif
\ No newline at end of file
diff --git a/Engine/Plugins/Experimental/Water/Source/Runtime/Private/WaterSubsystem.cpp b/Engine/Plugins/Experimental/Water/Source/Runtime/Private/WaterSubsystem.cpp
--- a/Engine/Plugins/Experimental/Water/Source/Runtime/Private/WaterSubsystem.cpp
+++ b/Engine/Plugins/Experimental/Water/Source/Runtime/Private/WaterSubsystem.cpp
@@ -314,10 +314,13 @@ void UWaterSubsystem::ApplyRuntimeSettings(const UWaterRuntimeSettings* Settings
 	MaterialParameterCollection = Settings->MaterialParameterCollection.LoadSynchronous();
 
 #if WITH_EDITOR
-	for (TActorIterator<AWaterBody> ActorItr(World); ActorItr; ++ActorItr)
+	// Update sprites since we may have changed the sprite Z offset setting.
+
+	ForEachWaterBodyComponent(World, [](UWaterBodyComponent* Component)
 	{
-		(*ActorItr)->UpdateActorIcon();
-	}
+		Component->UpdateWaterSpriteComponent();
+		return true;
+	});
 
 	for (TActorIterator<AWaterBodyIsland> ActorItr(World); ActorItr; ++ActorItr)
 	{
diff --git a/Engine/Plugins/Experimental/Water/Source/Runtime/Public/WaterBodyActor.h b/Engine/Plugins/Experimental/Water/Source/Runtime/Public/WaterBodyActor.h
--- a/Engine/Plugins/Experimental/Water/Source/Runtime/Public/WaterBodyActor.h
+++ b/Engine/Plugins/Experimental/Water/Source/Runtime/Public/WaterBodyActor.h
@@ -95,10 +95,6 @@ public:
 	UFUNCTION(BlueprintCallable, Category=Water)
 	UWaterBodyComponent* GetWaterBodyComponent() const { return WaterBodyComponent; }
 
-#if WITH_EDITOR
-	// #todo_water: all icon stuff can be moved to the component
-	void UpdateActorIcon();
-#endif // WITH_EDITOR
 protected:
 	/** Initializes the water body by creating the respective component for this water body type. */
 	virtual void InitializeBody();
@@ -126,17 +122,10 @@ protected:
 	UPROPERTY(EditAnywhere, BlueprintReadOnly, Instanced, Category = Wave, DisplayName = "Waves Source", meta = (Tooltip = ""))
 	UWaterWavesBase* WaterWaves = nullptr;
 
-#if WITH_EDITORONLY_DATA
-	UPROPERTY(Transient)
-	UBillboardComponent* ActorIcon;
-#endif
-
 #if WITH_EDITOR
 	virtual void PostEditMove(bool bFinished) override;
 	virtual void PreEditChange(FProperty* PropertyThatWillChange) override;
 	virtual void PostEditChangeProperty(FPropertyChangedEvent& PropertyChangedEvent) override;
-	/** Returns whether icon billboard is visible. */
-	virtual bool IsIconVisible() const;
 	virtual void PostDuplicate(bool bDuplicateForPIE) override;
 #endif // WITH_EDITOR
 
@@ -146,6 +135,9 @@ protected:
 // Deprecated
 
 public:
+	UE_DEPRECATED(5.1, "Moved to WaterBodyComponent")
+	virtual bool IsIconVisible() const { return WaterBodyComponent->IsIconVisible(); }
+
 	UE_DEPRECATED(4.27, "Moved to WaterBodyComponent")
 	virtual bool HasWaves() const final { return WaterBodyComponent->HasWaves(); }
 
diff --git a/Engine/Plugins/Experimental/Water/Source/Runtime/Public/WaterBodyComponent.h b/Engine/Plugins/Experimental/Water/Source/Runtime/Public/WaterBodyComponent.h
--- a/Engine/Plugins/Experimental/Water/Source/Runtime/Public/WaterBodyComponent.h
+++ b/Engine/Plugins/Experimental/Water/Source/Runtime/Public/WaterBodyComponent.h
@@ -96,6 +96,8 @@ public:
 	virtual ETextureRenderTargetFormat GetBrushRenderTargetFormat() const;
 	virtual void GetBrushRenderDependencies(TSet<UObject*>& OutDependencies) const;
 	virtual TArray<UPrimitiveComponent*> GetBrushRenderableComponents() const { return TArray<UPrimitiveComponent*>(); }
+	
+	virtual void UpdateWaterSpriteComponent();
 #endif //WITH_EDITOR
 
 	/** Returns whether the body supports waves */
@@ -368,6 +370,12 @@ protected:
 	virtual void CheckForErrors() override;
 
 	virtual TArray<TSharedRef<FTokenizedMessage>> CheckWaterBodyStatus() const;
+
+	virtual const TCHAR* GetWaterSpriteTextureName() const { return TEXT("/Water/Icons/WaterSprite"); }
+
+	virtual bool IsIconVisible() const { return true; }
+
+	virtual FVector GetWaterSpriteLocation() const { return GetComponentLocation(); }
 #endif // WITH_EDITOR
 
 	EWaterBodyQueryFlags CheckAndAjustQueryFlags(EWaterBodyQueryFlags InQueryFlags) const;
@@ -398,6 +406,8 @@ protected:
 
 	void OnWaterSplineMetadataChanged(UWaterSplineMetadata* InWaterSplineMetadata, FPropertyChangedEvent& PropertyChangedEvent);
 	void RegisterOnChangeWaterSplineMetadata(UWaterSplineMetadata* InWaterSplineMetadata, bool bRegister);
+
+	void CreateWaterSpriteComponent();
 #endif // WITH_EDITOR
 
 public:
diff --git a/Engine/Plugins/Experimental/Water/Source/Runtime/Public/WaterBodyCustomComponent.h b/Engine/Plugins/Experimental/Water/Source/Runtime/Public/WaterBodyCustomComponent.h
--- a/Engine/Plugins/Experimental/Water/Source/Runtime/Public/WaterBodyCustomComponent.h
+++ b/Engine/Plugins/Experimental/Water/Source/Runtime/Public/WaterBodyCustomComponent.h
@@ -30,6 +30,10 @@ protected:
 	virtual FPrimitiveSceneProxy* CreateSceneProxy() override;
 #if WITH_EDITOR
 	virtual TArray<TSharedRef<FTokenizedMessage>> CheckWaterBodyStatus() const override;
+
+	virtual const TCHAR* GetWaterSpriteTextureName() const override;
+
+	virtual bool IsIconVisible() const override { return false; }
 #endif // WITH_EDITOR
 
 protected:
diff --git a/Engine/Plugins/Experimental/Water/Source/Runtime/Public/WaterBodyLakeComponent.h b/Engine/Plugins/Experimental/Water/Source/Runtime/Public/WaterBodyLakeComponent.h
--- a/Engine/Plugins/Experimental/Water/Source/Runtime/Public/WaterBodyLakeComponent.h
+++ b/Engine/Plugins/Experimental/Water/Source/Runtime/Public/WaterBodyLakeComponent.h
@@ -31,6 +31,12 @@ protected:
 	virtual void Reset() override;
 	virtual void OnUpdateBody(bool bWithExclusionVolumes) override;
 
+#if WITH_EDITOR
+	virtual const TCHAR* GetWaterSpriteTextureName() const override;
+
+	virtual FVector GetWaterSpriteLocation() const override;
+#endif // WITH_EDITOR
+
 	UPROPERTY(NonPIEDuplicateTransient)
 	UStaticMeshComponent* LakeMeshComp;
 
diff --git a/Engine/Plugins/Experimental/Water/Source/Runtime/Public/WaterBodyOceanComponent.h b/Engine/Plugins/Experimental/Water/Source/Runtime/Public/WaterBodyOceanComponent.h
--- a/Engine/Plugins/Experimental/Water/Source/Runtime/Public/WaterBodyOceanComponent.h
+++ b/Engine/Plugins/Experimental/Water/Source/Runtime/Public/WaterBodyOceanComponent.h
@@ -39,6 +39,8 @@ protected:
 
 #if WITH_EDITOR
 	virtual void OnPostEditChangeProperty(FPropertyChangedEvent& PropertyChangedEvent, bool& bShapeOrPositionChanged, bool& bWeightmapSettingsChanged) override;
+
+	virtual const TCHAR* GetWaterSpriteTextureName() const override;
 #endif
 protected:
 	UPROPERTY(NonPIEDuplicateTransient)
diff --git a/Engine/Plugins/Experimental/Water/Source/Runtime/Public/WaterBodyRiverComponent.h b/Engine/Plugins/Experimental/Water/Source/Runtime/Public/WaterBodyRiverComponent.h
--- a/Engine/Plugins/Experimental/Water/Source/Runtime/Public/WaterBodyRiverComponent.h
+++ b/Engine/Plugins/Experimental/Water/Source/Runtime/Public/WaterBodyRiverComponent.h
@@ -42,6 +42,8 @@ protected:
 	virtual FBoxSphereBounds CalcBounds(const FTransform& LocalToWorld) const override;
 #if WITH_EDITOR
 	virtual void OnPostEditChangeProperty(FPropertyChangedEvent& PropertyChangedEvent, bool& bShapeOrPositionChanged, bool& bWeightmapSettingsChanged) override;
+
+	virtual const TCHAR* GetWaterSpriteTextureName() const override;
 #endif
 
 	void CreateOrUpdateLakeTransitionMID();
diff --git a/Engine/Plugins/Experimental/Water/Source/Runtime/Public/WaterRuntimeSettings.h b/Engine/Plugins/Experimental/Water/Source/Runtime/Public/WaterRuntimeSettings.h
--- a/Engine/Plugins/Experimental/Water/Source/Runtime/Public/WaterRuntimeSettings.h
+++ b/Engine/Plugins/Experimental/Water/Source/Runtime/Public/WaterRuntimeSettings.h
@@ -55,13 +55,9 @@ public:
 	UPROPERTY(EditAnywhere, config, Category = Rendering)
 	TSoftObjectPtr<UMaterialParameterCollection> MaterialParameterCollection;
 
-	/** Size of the water body icon in world-space. */
-	UPROPERTY(EditAnywhere, config, Category = Rendering)
-	float WaterBodyIconWorldSize = 1000.0f;
-
 	/** Offset in Z for the water body icon in world-space. */
 	UPROPERTY(EditAnywhere, config, Category = Rendering)
-	float WaterBodyIconWorldZOffset = 250.0f;
+	float WaterBodyIconWorldZOffset = 75.0f;
 
 #if WITH_EDITORONLY_DATA
 	// Delegate called whenever the curve data is updated
@@ -88,4 +84,9 @@ private:
 
 	UPROPERTY(EditAnywhere, Config, Category = Water,  meta = (MetaClass = "WaterBodyCustomComponent"))
 	TSubclassOf<UWaterBodyCustomComponent> WaterBodyCustomComponentClass;
+
+#if WITH_EDITORONLY_DATA
+	UPROPERTY()
+	float WaterBodyIconWorldSize_DEPRECATED;
+#endif // WITH_EDITORONLY_DATA
 };
