Fix for non-unity build and missing file

diff --git a/Engine/Plugins/Experimental/Water/Source/Runtime/Private/WaterBodyComponent.cpp b/Engine/Plugins/Experimental/Water/Source/Runtime/Private/WaterBodyComponent.cpp
--- a/Engine/Plugins/Experimental/Water/Source/Runtime/Private/WaterBodyComponent.cpp
+++ b/Engine/Plugins/Experimental/Water/Source/Runtime/Private/WaterBodyComponent.cpp
@@ -38,6 +38,8 @@
 #if WITH_EDITOR
 #include "WaterIconHelper.h"
 #include "Components/BillboardComponent.h"
+#include "Modules/ModuleManager.h"
+#include "WaterModule.h"
 #endif
 
 #define LOCTEXT_NAMESPACE "Water"
diff --git a/Engine/Plugins/Experimental/Water/Source/Runtime/Private/WaterIconHelper.cpp b/Engine/Plugins/Experimental/Water/Source/Runtime/Private/WaterIconHelper.cpp
--- a/Engine/Plugins/Experimental/Water/Source/Runtime/Private/WaterIconHelper.cpp
+++ b/Engine/Plugins/Experimental/Water/Source/Runtime/Private/WaterIconHelper.cpp
@@ -32,6 +32,7 @@ UBillboardComponent* FWaterIconHelper::EnsureSpriteComponentCreated_Internal(AAc
 		ActorIcon->bHiddenInGame = true;
 		ActorIcon->SpriteInfo.Category = TEXT("Water");
 		ActorIcon->SpriteInfo.DisplayName = NSLOCTEXT("SpriteCategory", "Water", "Water");
+		ActorIcon->SetVisibility(true);
 		ActorIcon->SetupAttachment(Actor->GetRootComponent());
 		UpdateSpriteComponent(Actor, ActorIcon->Sprite);
 	}
@@ -42,15 +43,9 @@ void FWaterIconHelper::UpdateSpriteComponent(AActor* Actor, UTexture2D* InTextur
 {
 	if (UBillboardComponent* ActorIcon = Actor->FindComponentByClass<UBillboardComponent>())
 	{
-		float TargetSize = GetDefault<UWaterRuntimeSettings>()->WaterBodyIconWorldSize;
 		FVector ZOffset(0.0f, 0.0f, GetDefault<UWaterRuntimeSettings>()->WaterBodyIconWorldZOffset);
-		if (InTexture != nullptr)
-		{
-			// Use the texture source's size as the texture might not having finished loading when this runs, in which case the default texture's size would be returned : 
-			int32 TextureSize = FMath::Max(InTexture->Source.GetSizeX(), InTexture->Source.GetSizeY());
-			float Scale = (TextureSize > 0) ? (TargetSize / (float)TextureSize) : 1.0f;
-			ActorIcon->SetRelativeScale3D(FVector(Scale));
-		}
+
+		ActorIcon->SetRelativeScale3D(FVector(1.f));
 		ActorIcon->Sprite = InTexture;
 		ActorIcon->SetRelativeLocation(ZOffset);
 		ActorIcon->bIsScreenSizeScaled = true;
diff --git a/Engine/Plugins/Experimental/Water/Source/Runtime/Public/WaterBodyActor.h b/Engine/Plugins/Experimental/Water/Source/Runtime/Public/WaterBodyActor.h
--- a/Engine/Plugins/Experimental/Water/Source/Runtime/Public/WaterBodyActor.h
+++ b/Engine/Plugins/Experimental/Water/Source/Runtime/Public/WaterBodyActor.h
@@ -135,9 +135,6 @@ protected:
 // Deprecated
 
 public:
-	UE_DEPRECATED(5.1, "Moved to WaterBodyComponent")
-	virtual bool IsIconVisible() const { return WaterBodyComponent->IsIconVisible(); }
-
 	UE_DEPRECATED(4.27, "Moved to WaterBodyComponent")
 	virtual bool HasWaves() const final { return WaterBodyComponent->HasWaves(); }
 
@@ -249,6 +246,11 @@ public:
 	UE_DEPRECATED(4.27, "Moved to WaterBodyComponent")
 	virtual ALandscapeProxy* FindLandscape() const final { return WaterBodyComponent->FindLandscape(); }
 
+#if WITH_EDITOR
+	UE_DEPRECATED(5.1, "Moved to WaterBodyComponent")
+	virtual bool IsIconVisible() const { return WaterBodyComponent->IsIconVisible(); }
+#endif // WITH_EDITOR
+
 protected:
 	UE_DEPRECATED(4.27, "Moved to WaterBodyComponent")
 	virtual bool IsBodyDynamic() const { return WaterBodyComponent->IsBodyDynamic(); }
