Fixed an issue where water sprite components were assumed to always be valid but were not created when in game worlds.

diff --git a/Engine/Plugins/Experimental/Water/Source/Runtime/Private/WaterBodyComponent.cpp b/Engine/Plugins/Experimental/Water/Source/Runtime/Private/WaterBodyComponent.cpp
--- a/Engine/Plugins/Experimental/Water/Source/Runtime/Private/WaterBodyComponent.cpp
+++ b/Engine/Plugins/Experimental/Water/Source/Runtime/Private/WaterBodyComponent.cpp
@@ -831,8 +831,8 @@ void UWaterBodyComponent::UpdateComponentVisibility(bool bAllowWaterMeshRebuild)
 				{
 					WaterZone->MarkForRebuild(EWaterZoneRebuildFlags::All);
 				}
-	 		}
-	 	}
+			}
+		}
 	}
 }
 
@@ -911,15 +911,15 @@ TArray<TSharedRef<FTokenizedMessage>> UWaterBodyComponent::CheckWaterBodyStatus(
 					->AddToken(FTextToken::Create(FText::Format(
 						LOCTEXT("MapCheck_Message_MissingWaterZone", "Water body {0} requires a WaterZone actor to be rendered. Please add one to the map. "),
 						FText::FromString(GetWaterBodyActor()->GetActorLabel())))));
+				}
 			}
-		}
 
 		if (AffectsLandscape() && (FindLandscape() == nullptr))
-		{
+			{
 			Result.Add(FTokenizedMessage::Create(EMessageSeverity::Error)
 				->AddToken(FUObjectToken::Create(this))
 				->AddToken(FTextToken::Create(FText::Format(
-					LOCTEXT("MapCheck_Message_MissingLandscape", "Water body {0} requires a Landscape to be rendered. Please add one to the map. "),
+						LOCTEXT("MapCheck_Message_MissingLandscape", "Water body {0} requires a Landscape to be rendered. Please add one to the map. "),
 					FText::FromString(GetWaterBodyActor()->GetActorLabel())))));
 		}
 	}
@@ -1447,8 +1447,12 @@ void UWaterBodyComponent::CreateWaterSpriteComponent()
 	bVisualizeComponent = true;
 	CreateSpriteComponent(Texture);
 
-	SpriteComponent->SetRelativeScale3D(FVector(1.f, 1.f, 1.f));
-	SpriteComponent->SetRelativeLocation(FVector(0.f, 0.f,  GetDefault<UWaterRuntimeSettings>()->WaterBodyIconWorldZOffset));
+	// CreateSpriteComponent will not create a component if we are in a game world such as PIE.
+	if (SpriteComponent != nullptr)
+	{
+		SpriteComponent->SetRelativeScale3D(FVector(1.f, 1.f, 1.f));
+		SpriteComponent->SetRelativeLocation(FVector(0.f, 0.f,  GetDefault<UWaterRuntimeSettings>()->WaterBodyIconWorldZOffset));
+	}
 }
 
 void UWaterBodyComponent::UpdateWaterSpriteComponent()
