Fixed typo in custom water body sprite path

diff --git a/Engine/Plugins/Experimental/Water/Source/Runtime/Private/WaterBodyCustomComponent.cpp b/Engine/Plugins/Experimental/Water/Source/Runtime/Private/WaterBodyCustomComponent.cpp
--- a/Engine/Plugins/Experimental/Water/Source/Runtime/Private/WaterBodyCustomComponent.cpp
+++ b/Engine/Plugins/Experimental/Water/Source/Runtime/Private/WaterBodyCustomComponent.cpp
@@ -142,7 +142,7 @@ TArray<TSharedRef<FTokenizedMessage>> UWaterBodyCustomComponent::CheckWaterBodyS
 }
 const TCHAR* UWaterBodyCustomComponent::GetWaterSpriteTextureName() const
 {
-	return TEXT("/Water/Icons/WaterBodyCustomComponent");
+	return TEXT("/Water/Icons/WaterBodyCustomSprite");
 }
 #endif // WITH_EDITOR
 
