Disabled collision on the WaterZone's bounds component

diff --git a/Engine/Plugins/Experimental/Water/Source/Runtime/Private/WaterZoneActor.cpp b/Engine/Plugins/Experimental/Water/Source/Runtime/Private/WaterZoneActor.cpp
--- a/Engine/Plugins/Experimental/Water/Source/Runtime/Private/WaterZoneActor.cpp
+++ b/Engine/Plugins/Experimental/Water/Source/Runtime/Private/WaterZoneActor.cpp
@@ -49,6 +49,7 @@ AWaterZone::AWaterZone(const FObjectInitializer& Initializer)
 		BoundsComponent = CreateDefaultSubobject<UBoxComponent>(TEXT("BoundsComponent"));
 		BoundsComponent->SetCollisionObjectType(ECC_WorldStatic);
 		BoundsComponent->SetCollisionResponseToAllChannels(ECR_Ignore);
+		BoundsComponent->SetCollisionEnabled(ECollisionEnabled::NoCollision);
 		BoundsComponent->SetGenerateOverlapEvents(false);
 		BoundsComponent->SetupAttachment(WaterMesh);
 		// Bounds component extent is half-extent, ZoneExtent is full extent.
