Fix build errors caused by removing implicit conversion to FVector.

diff --git a/Engine/Plugins/Experimental/Water/Source/Runtime/Private/WaterBodySceneProxy.cpp b/Engine/Plugins/Experimental/Water/Source/Runtime/Private/WaterBodySceneProxy.cpp
--- a/Engine/Plugins/Experimental/Water/Source/Runtime/Private/WaterBodySceneProxy.cpp
+++ b/Engine/Plugins/Experimental/Water/Source/Runtime/Private/WaterBodySceneProxy.cpp
@@ -69,7 +69,7 @@ static FColor PackWaterFlow(float VelocityMagnitude, float DirectionAngle)
 
 using namespace UE::Geometry;
 
-void FWaterBodySceneProxy::AddSectionFromDynamicMesh(const FDynamicMesh3& DynamicMesh, TFunctionRef<FDynamicMeshVertex(const FVector3f&)> VertTransformFunc)
+void FWaterBodySceneProxy::AddSectionFromDynamicMesh(const FDynamicMesh3& DynamicMesh, TFunctionRef<FDynamicMeshVertex(const FVector&)> VertTransformFunc)
 {
 	TArray<FDynamicMeshVertex> Vertices;
 	Vertices.Reserve(DynamicMesh.VertexCount());
@@ -131,8 +131,8 @@ void FWaterBodySceneProxy::GenerateLakeMesh(UWaterBodyLakeComponent* Component)
 		LakeMesh.AppendTriangle(Triangle);
 	}
 
-	AddSectionFromDynamicMesh(LakeMesh, [WaterBodyIndex = Component->GetWaterBodyIndex()](const FVector3f& In) {
-			FDynamicMeshVertex Result(FVector(In.X, In.Y, 0.f));
+	AddSectionFromDynamicMesh(LakeMesh, [WaterBodyIndex = Component->GetWaterBodyIndex()](const FVector& In) {
+			FDynamicMeshVertex Result(FVector3f(In.X, In.Y, 0.f));
 			Result.Color = FColor::Black;
 			Result.TextureCoordinate[0].X = WaterBodyIndex;
 			return Result;
@@ -158,8 +158,8 @@ void FWaterBodySceneProxy::GenerateLakeMesh(UWaterBodyLakeComponent* Component)
 	}
 	
 	// Push down the dilated region vertices by -ZOffset to prevent dilated regions from overwriting data from adjacent water bodies
-	AddSectionFromDynamicMesh(LakeMesh, [ZOffset = Component->GetShapeDilationZOffset()](const FVector3f& In) {
-			FDynamicMeshVertex Result(FVector(In.X, In.Y, ZOffset));
+	AddSectionFromDynamicMesh(LakeMesh, [ZOffset = Component->GetShapeDilationZOffset()](const FVector& In) {
+			FDynamicMeshVertex Result(FVector3f(In.X, In.Y, ZOffset));
 			Result.Color = FColor::Black;
 			Result.TextureCoordinate[0].X = -1;
 			return Result;
@@ -214,8 +214,8 @@ void FWaterBodySceneProxy::GenerateOceanMesh(UWaterBodyOceanComponent* Component
 	{
 		OceanMesh.AppendTriangle(Triangle);
 	}
-	AddSectionFromDynamicMesh(OceanMesh, [WaterBodyIndex = Component->GetWaterBodyIndex()](const FVector3f& In) {
-			FDynamicMeshVertex Result(FVector(In.X, In.Y, 0.f));
+	AddSectionFromDynamicMesh(OceanMesh, [WaterBodyIndex = Component->GetWaterBodyIndex()](const FVector& In) {
+			FDynamicMeshVertex Result(FVector3f(In.X, In.Y, 0.f));
 			Result.Color = FColor::Black;
 			Result.TextureCoordinate[0].X = WaterBodyIndex;
 			return Result;
@@ -241,8 +241,8 @@ void FWaterBodySceneProxy::GenerateOceanMesh(UWaterBodyOceanComponent* Component
 	}
 	
 	// Push down the dilated region vertices by -ZOffset to prevent dilated regions from overwriting data from adjacent water bodies
-	AddSectionFromDynamicMesh(OceanMesh, [ZOffset = Component->GetShapeDilationZOffset()](const FVector3f& In) {
-			FDynamicMeshVertex Result(FVector(In.X, In.Y, ZOffset));
+	AddSectionFromDynamicMesh(OceanMesh, [ZOffset = Component->GetShapeDilationZOffset()](const FVector& In) {
+			FDynamicMeshVertex Result(FVector3f(In.X, In.Y, ZOffset));
 			Result.Color = FColor::Black;
 			Result.TextureCoordinate[0].X = -1;
 			return Result;
@@ -271,16 +271,16 @@ static void AddVerticesForRiverSplineStep(float DistanceAlongSpline, const UWate
 	const float DilationAmount = Component->ShapeDilation;
 	const FVector DilationOffset = Normal * DilationAmount;
 
-	FDynamicMeshVertex Left(Pos - OutwardDistance);
-	FDynamicMeshVertex Right(Pos + OutwardDistance);
+	FDynamicMeshVertex Left(FVector3f(Pos - OutwardDistance));
+	FDynamicMeshVertex Right(FVector3f(Pos + OutwardDistance));
 
-	FDynamicMeshVertex DilatedFarLeft(Pos - OutwardDistance - DilationOffset);
+	FDynamicMeshVertex DilatedFarLeft(FVector3f(Pos - OutwardDistance - DilationOffset));
 	DilatedFarLeft.Position.Z += Component->GetShapeDilationZOffsetFar();
-	FDynamicMeshVertex DilatedLeft(Pos - OutwardDistance);
+	FDynamicMeshVertex DilatedLeft(FVector3f(Pos - OutwardDistance));
 	DilatedLeft.Position.Z += Component->GetShapeDilationZOffset();
-	FDynamicMeshVertex DilatedRight(Pos + OutwardDistance);
+	FDynamicMeshVertex DilatedRight(FVector3f(Pos + OutwardDistance));
 	DilatedRight.Position.Z += Component->GetShapeDilationZOffset();
-	FDynamicMeshVertex DilatedFarRight(Pos + OutwardDistance + DilationOffset);
+	FDynamicMeshVertex DilatedFarRight(FVector3f(Pos + OutwardDistance + DilationOffset));
 	DilatedFarRight.Position.Z += Component->GetShapeDilationZOffsetFar();
 
 	float FlowDirection = Tangent.HeadingAngle() + FMath::DegreesToRadians(Component->GetRelativeRotation().Yaw);
@@ -367,13 +367,13 @@ static void AddTerminalVerticesForRiverSpline(ERiverBoundaryEdge Edge, const UWa
 		TangentialOffset *= -1;
 	}
 
-	FDynamicMeshVertex BackLeft(Pos - OutwardDistance + TangentialOffset - DilationOffset);
+	FDynamicMeshVertex BackLeft(FVector3f(Pos - OutwardDistance + TangentialOffset - DilationOffset));
 	BackLeft.Position.Z += Component->GetShapeDilationZOffsetFar();
-	FDynamicMeshVertex Left(Pos - OutwardDistance + TangentialOffset);
+	FDynamicMeshVertex Left(FVector3f(Pos - OutwardDistance + TangentialOffset));
 	Left.Position.Z += Component->GetShapeDilationZOffsetFar();
-	FDynamicMeshVertex Right(Pos + OutwardDistance + TangentialOffset);
+	FDynamicMeshVertex Right(FVector3f(Pos + OutwardDistance + TangentialOffset));
 	Right.Position.Z += Component->GetShapeDilationZOffsetFar();
-	FDynamicMeshVertex BackRight(Pos + OutwardDistance + TangentialOffset + DilationOffset);
+	FDynamicMeshVertex BackRight(FVector3f(Pos + OutwardDistance + TangentialOffset + DilationOffset));
 	BackRight.Position.Z += Component->GetShapeDilationZOffsetFar();
 
 	// Initialize the vertex data to correct represent what we expect the dilated region to look like
diff --git a/Engine/Plugins/Experimental/Water/Source/Runtime/Public/WaterBodySceneProxy.h b/Engine/Plugins/Experimental/Water/Source/Runtime/Public/WaterBodySceneProxy.h
--- a/Engine/Plugins/Experimental/Water/Source/Runtime/Public/WaterBodySceneProxy.h
+++ b/Engine/Plugins/Experimental/Water/Source/Runtime/Public/WaterBodySceneProxy.h
@@ -36,7 +36,7 @@ private:
 	void GenerateOceanMesh(UWaterBodyOceanComponent* Component);
 	void GenerateRiverMesh(UWaterBodyRiverComponent* Component);
 
-	void AddSectionFromDynamicMesh(const UE::Geometry::FDynamicMesh3& DynamicMesh, TFunctionRef<FDynamicMeshVertex(const FVector3f&)> VertTransformFunc);
+	void AddSectionFromDynamicMesh(const UE::Geometry::FDynamicMesh3& DynamicMesh, TFunctionRef<FDynamicMeshVertex(const FVector&)> VertTransformFunc);
 
 	void InitResources(FWaterBodyMeshSection* Section);
 
