Water RayTracing initial support:
* Create a BLAS and dynamic instance for each water instance.
* WaterVertexFactory updates to support RayTracing and RayTracingDynamicGeometry.
* Cache BLASes so we can update them instead of building.
* Disabled by default

diff --git a/Engine/Plugins/Experimental/Water/Shaders/Private/WaterMeshVertexFactory.ush b/Engine/Plugins/Experimental/Water/Shaders/Private/WaterMeshVertexFactory.ush
--- a/Engine/Plugins/Experimental/Water/Shaders/Private/WaterMeshVertexFactory.ush
+++ b/Engine/Plugins/Experimental/Water/Shaders/Private/WaterMeshVertexFactory.ush
@@ -365,6 +365,56 @@ float3 VertexFactoryGetWorldNormal(FVertexFactoryInput Input, FVertexFactoryInte
 	return float3(0.0f, 0.0f, 1.0f);
 }
 
+#if RAYHITGROUPSHADER
+FVertexFactoryInput LoadVertexFactoryInputForHGS(uint TriangleIndex, int VertexIndex)
+{
+	FTriangleBaseAttributes TriangleAttributes = LoadTriangleBaseAttributes(TriangleIndex);
+
+	FVertexFactoryInput Input = (FVertexFactoryInput)0;
+
+	Input.Position = float4(TriangleAttributes.LocalPositions[VertexIndex], 1.0f);
+
+	Input.InstanceData0 = WaterRaytracingVF.InstanceData0;
+	Input.InstanceData1 = WaterRaytracingVF.InstanceData1;
+#if WITH_WATER_SELECTION_SUPPORT_VF
+	Input.InstanceData2 = 0.0f;
+#endif
+
+	VF_GPUSCENE_SET_INPUT_FOR_RT(Input, GetInstanceUserData(), 0U);
+
+	return Input;
+}
+#endif
+
+#if COMPUTESHADER
+FVertexFactoryInput LoadVertexFactoryInputForDynamicUpdate(uint TriangleIndex, int VertexIndex, uint PrimitiveId)
+{
+	FVertexFactoryInput Input = (FVertexFactoryInput)0;
+
+	const uint VertexId = TriangleIndex * 3 + VertexIndex;
+	const uint VertexOffset = VertexId * 4;
+	Input.Position.x = WaterRaytracingVF.VertexBuffer[VertexOffset + 0];
+	Input.Position.y = WaterRaytracingVF.VertexBuffer[VertexOffset + 1];
+	Input.Position.z = WaterRaytracingVF.VertexBuffer[VertexOffset + 2];
+	Input.Position.w = WaterRaytracingVF.VertexBuffer[VertexOffset + 3];
+
+	Input.InstanceData0 = WaterRaytracingVF.InstanceData0;
+	Input.InstanceData1 = WaterRaytracingVF.InstanceData1;
+#if WITH_WATER_SELECTION_SUPPORT_VF
+	Input.InstanceData2 = 0.0f;
+#endif
+	
+	VF_GPUSCENE_SET_INPUT_FOR_RT(Input, GetInstanceUserData(), 0U);
+
+	return Input;
+}
+
+uint GetNumRayTracingDynamicMeshVerticesIndirect()
+{
+	return 0;
+}
+#endif // RAYHITGROUP || COMPUTESHADER
+
 #if NEEDS_VERTEX_FACTORY_INTERPOLATION
 	struct FVertexFactoryRayTracingInterpolants
 	{
@@ -374,7 +424,7 @@ float3 VertexFactoryGetWorldNormal(FVertexFactoryInput Input, FVertexFactoryInte
 	float2 VertexFactoryGetRayTracingTextureCoordinate( FVertexFactoryRayTracingInterpolants Interpolants )
 	{
 #if NUM_MATERIAL_TEXCOORDS
-		return Interpolants.InterpolantsVSToPS.TexCoords.xy;
+		return Interpolants.InterpolantsVSToPS.TexCoords[0].xy;
 #else // #if NUM_MATERIAL_TEXCOORDS
 		return float2(0,0);
 #endif // #if NUM_MATERIAL_TEXCOORDS
diff --git a/Engine/Plugins/Experimental/Water/Source/Runtime/Private/WaterMeshSceneProxy.cpp b/Engine/Plugins/Experimental/Water/Source/Runtime/Private/WaterMeshSceneProxy.cpp
--- a/Engine/Plugins/Experimental/Water/Source/Runtime/Private/WaterMeshSceneProxy.cpp
+++ b/Engine/Plugins/Experimental/Water/Source/Runtime/Private/WaterMeshSceneProxy.cpp
@@ -11,6 +11,7 @@
 #include "Materials/MaterialInstanceDynamic.h"
 #include "DrawDebugHelpers.h"
 #include "Math/ColorList.h"
+#include "RayTracingInstance.h"
 
 DECLARE_STATS_GROUP(TEXT("Water Mesh"), STATGROUP_WaterMesh, STATCAT_Advanced);
 
@@ -74,6 +75,12 @@ static TAutoConsoleVariable<int32> CVarWaterMeshPreAllocStagingInstanceMemory(
 	TEXT("Pre-allocates staging instance data memory according to historical max. This reduces the overhead when the array needs to grow but may use more memory"),
 	ECVF_RenderThreadSafe);
 
+#if RHI_RAYTRACING
+static TAutoConsoleVariable<int32> CVarRayTracingGeometryWater(
+	TEXT("r.RayTracing.Geometry.Water"),
+	0,
+	TEXT("Include water in ray tracing effects (default = 0 (water disabled in ray tracing))"));
+#endif
 
 // ----------------------------------------------------------------------------------
 
@@ -118,6 +125,10 @@ FWaterMeshSceneProxy::FWaterMeshSceneProxy(UWaterMeshComponent* Component)
 	WaterMeshUserDataBuffers = new WaterMeshUserDataBuffersType(WaterInstanceDataBuffers);
 
 	WaterQuadTree.BuildMaterialIndices();
+
+#if RHI_RAYTRACING
+	RayTracingWaterData.SetNum(DensityCount);
+#endif
 }
 
 FWaterMeshSceneProxy::~FWaterMeshSceneProxy()
@@ -131,6 +142,38 @@ FWaterMeshSceneProxy::~FWaterMeshSceneProxy()
 	delete WaterInstanceDataBuffers;
 
 	delete WaterMeshUserDataBuffers;
+
+#if RHI_RAYTRACING
+	for (auto& WaterDataArray : RayTracingWaterData)
+	{
+		for (auto& WaterRayTracingItem : WaterDataArray)
+		{
+			WaterRayTracingItem.Geometry.ReleaseResource();
+			WaterRayTracingItem.DynamicVertexBuffer.Release();
+		}
+	}	
+#endif
+}
+
+FWaterMeshSceneProxy::FWaterLODParams FWaterMeshSceneProxy::GetWaterLODParams(const FVector& Position) const
+{
+	float WaterHeightForLOD = 0.0f;
+	WaterQuadTree.QueryInterpolatedTileBaseHeightAtLocation(FVector2D(Position), WaterHeightForLOD);
+
+	// Need to let the lowest LOD morph globally towards the next LOD. When the LOD is done morphing, simply clamp the LOD in the LOD selection to effectively promote the lowest LOD to the same LOD level as the one above
+	float DistToWater = FMath::Abs(Position.Z - WaterHeightForLOD) / LODScale;
+	DistToWater = FMath::Max(DistToWater - 2.0f, 0.0f);
+	DistToWater *= 2.0f;
+
+	// Clamp to WaterTileQuadTree.GetLODCount() - 1.0f prevents the last LOD to morph
+	const float FloatLOD = FMath::Clamp(FMath::Log2(DistToWater), 0.0f, WaterQuadTree.GetTreeDepth() - 1.0f);
+
+	FWaterLODParams WaterLODParams;
+	WaterLODParams.HeightLODFactor = FMath::Frac(FloatLOD);
+	WaterLODParams.LowestLOD = FMath::Clamp(FMath::FloorToInt(FloatLOD), 0, WaterQuadTree.GetTreeDepth() - 1);
+	WaterLODParams.WaterHeightForLOD = WaterHeightForLOD;
+
+	return WaterLODParams;
 }
 
 void FWaterMeshSceneProxy::GetDynamicMeshElements(const TArray<const FSceneView*>& Views, const FSceneViewFamily& ViewFamily, uint32 VisibilityMap, FMeshElementCollector& Collector) const
@@ -155,7 +198,7 @@ void FWaterMeshSceneProxy::GetDynamicMeshElements(const TArray<const FSceneView*
 	}
 #endif // WITH_WATER_SELECTION_SUPPORT
 
-	if (WaterQuadTree.GetNodeCount() == 0 || DensityCount == 0 || !FWaterUtils::IsWaterMeshRenderingEnabled(/*bIsRenderThread = */true))
+	if (!HasWaterData() || !FWaterUtils::IsWaterMeshRenderingEnabled(/*bIsRenderThread = */true))
 	{
 		return;
 	}
@@ -185,27 +228,16 @@ void FWaterMeshSceneProxy::GetDynamicMeshElements(const TArray<const FSceneView*
 			const FSceneView* View = Views[ViewIndex];
 
 			const FVector ObserverPosition = View->ViewMatrices.GetViewOrigin();
-
-			float WaterHeightForLOD = 0.0f;
-			WaterQuadTree.QueryInterpolatedTileBaseHeightAtLocation(FVector2D(ObserverPosition), WaterHeightForLOD);
-
-			// Need to let the lowest LOD morph globally towards the next LOD. When the LOD is done morphing, simply clamp the LOD in the LOD selection to effectively promote the lowest LOD to the same LOD level as the one above
-			float DistToWater = FMath::Abs(ObserverPosition.Z - WaterHeightForLOD) / LODScale;
-			DistToWater = FMath::Max(DistToWater - 2.0f, 0.0f);
-			DistToWater *= 2.0f;
-
-			// Clamp to WaterTileQuadTree.GetLODCount() - 1.0f prevents the last LOD to morph
-			float FloatLOD = FMath::Clamp(FMath::Log2(DistToWater), 0.0f, WaterQuadTree.GetTreeDepth() - 1.0f);
-			float HeightLODFactor = FMath::Frac(FloatLOD);
-			int32 LowestLOD = FMath::Clamp(FMath::FloorToInt(FloatLOD), 0, WaterQuadTree.GetTreeDepth() - 1);
+			
+			FWaterLODParams WaterLODParams = GetWaterLODParams(ObserverPosition);
 
 #if !(UE_BUILD_SHIPPING || UE_BUILD_TEST)
 			if (CVarWaterMeshShowLODLevels.GetValueOnRenderThread())
 			{
-				for (int32 i = LowestLOD; i < WaterQuadTree.GetTreeDepth(); i++)
+				for (int32 i = WaterLODParams.LowestLOD; i < WaterQuadTree.GetTreeDepth(); i++)
 				{
 					float LODDist = FWaterQuadTree::GetLODDistance(i, LODScale);
-					FVector Orig = FVector(FVector2D(ObserverPosition), WaterHeightForLOD);
+					FVector Orig = FVector(FVector2D(ObserverPosition), WaterLODParams.WaterHeightForLOD);
 
 					DrawCircle(Collector.GetPDI(ViewIndex), Orig, FVector::ForwardVector, FVector::RightVector, GColorList.GetFColorByIndex(i + 1), LODDist, 64, 0);
 				}
@@ -222,8 +254,8 @@ void FWaterMeshSceneProxy::GetDynamicMeshElements(const TArray<const FSceneView*
 			}
 
 			FWaterQuadTree::FTraversalDesc TraversalDesc;
-			TraversalDesc.LowestLOD = LowestLOD;
-			TraversalDesc.HeightMorph = HeightLODFactor;
+			TraversalDesc.LowestLOD = WaterLODParams.LowestLOD;
+			TraversalDesc.HeightMorph = WaterLODParams.HeightLODFactor;
 			TraversalDesc.LODCount = WaterQuadTree.GetTreeDepth();
 			TraversalDesc.DensityCount = DensityCount;
 			TraversalDesc.ForceCollapseDensityLevel = ForceCollapseDensityLevel;
@@ -388,6 +420,212 @@ void FWaterMeshSceneProxy::GetDynamicMeshElements(const TArray<const FSceneView*
 	WaterInstanceDataBuffers->Unlock();
 }
 
+#if RHI_RAYTRACING
+void FWaterMeshSceneProxy::SetupRayTracingInstances(int32 NumInstances, uint32 DensityIndex)
+{
+	TArray<FRayTracingWaterData>& WaterDataArray = RayTracingWaterData[DensityIndex];
+
+	if (WaterDataArray.Num() > NumInstances)
+	{
+		for (int32 Item = NumInstances; Item < WaterDataArray.Num(); Item++)
+		{
+			auto& WaterItem = WaterDataArray[Item];
+			WaterItem.Geometry.ReleaseResource();
+			WaterItem.DynamicVertexBuffer.Release();
+		}
+		WaterDataArray.SetNum(NumInstances);
+	}	
+
+	if (WaterDataArray.Num() < NumInstances)
+	{
+		FRayTracingGeometryInitializer Initializer;
+		static const FName DebugName("FWaterMeshSceneProxy");		
+		Initializer.IndexBuffer = WaterVertexFactories[DensityIndex]->IndexBuffer->IndexBufferRHI;
+		Initializer.GeometryType = RTGT_Triangles;
+		Initializer.bFastBuild = true;
+		Initializer.bAllowUpdate = true;
+		Initializer.TotalPrimitiveCount = 0;
+
+		WaterDataArray.Reserve(NumInstances);
+		const int32 StartIndex = WaterDataArray.Num();
+
+		for (int32 Item = StartIndex; Item < NumInstances; Item++)
+		{
+			FRayTracingWaterData& WaterData = WaterDataArray.AddDefaulted_GetRef();
+
+			Initializer.DebugName = FName(DebugName, Item);
+
+			WaterData.Geometry.SetInitializer(Initializer);
+			WaterData.Geometry.InitResource();
+		}
+	}
+}
+
+template <bool bWithWaterSelectionSupport>
+class TWaterVertexFactoryUserDataWrapper : public FOneFrameResource
+{
+public:
+	TWaterMeshUserData<bWithWaterSelectionSupport> UserData;
+};
+
+void FWaterMeshSceneProxy::GetDynamicRayTracingInstances(FRayTracingMaterialGatheringContext& Context, TArray<FRayTracingInstance>& OutRayTracingInstances)
+{
+	if (!HasWaterData() || !FWaterUtils::IsWaterMeshRenderingEnabled(/*bIsRenderThread = */true) || !CVarRayTracingGeometryWater.GetValueOnRenderThread())
+	{
+		return;
+	}
+
+	const FSceneView& SceneView = *Context.ReferenceView;
+	const FVector ObserverPosition = SceneView.ViewMatrices.GetViewOrigin();
+
+	FWaterLODParams WaterLODParams = GetWaterLODParams(ObserverPosition);
+
+	const int32 NumBuckets = WaterQuadTree.GetWaterMaterials().Num() * DensityCount;
+
+	FWaterQuadTree::FTraversalOutput WaterInstanceData;
+	WaterInstanceData.BucketInstanceCounts.Empty(NumBuckets);
+	WaterInstanceData.BucketInstanceCounts.AddZeroed(NumBuckets);
+
+	FWaterQuadTree::FTraversalDesc TraversalDesc;
+	TraversalDesc.LowestLOD = WaterLODParams.LowestLOD;
+	TraversalDesc.HeightMorph = WaterLODParams.HeightLODFactor;
+	TraversalDesc.LODCount = WaterQuadTree.GetTreeDepth();
+	TraversalDesc.DensityCount = DensityCount;
+	TraversalDesc.ForceCollapseDensityLevel = ForceCollapseDensityLevel;
+	TraversalDesc.PreViewTranslation = SceneView.ViewMatrices.GetPreViewTranslation();
+	TraversalDesc.ObserverPosition = ObserverPosition;
+	TraversalDesc.Frustum = FConvexVolume(); // Default volume to disable frustum culling
+	TraversalDesc.LODScale = LODScale;
+	TraversalDesc.bLODMorphingEnabled = !!CVarWaterMeshLODMorphEnabled.GetValueOnRenderThread();
+
+	WaterQuadTree.BuildWaterTileInstanceData(TraversalDesc, WaterInstanceData);
+
+	if (WaterInstanceData.InstanceCount == 0)
+	{
+		// no instance visible, early exit
+		return;
+	}
+
+	const int32 NumWaterMaterials = WaterQuadTree.GetWaterMaterials().Num();	
+
+	for (int32 DensityIndex = 0; DensityIndex < DensityCount; ++DensityIndex)
+	{
+		int32 DensityInstanceCount = 0;
+		for (int32 MaterialIndex = 0; MaterialIndex < NumWaterMaterials; ++MaterialIndex)
+		{
+			const int32 BucketIndex = MaterialIndex * DensityCount + DensityIndex;
+			const int32 InstanceCount = WaterInstanceData.BucketInstanceCounts[BucketIndex];
+			DensityInstanceCount += InstanceCount;
+		}
+
+		SetupRayTracingInstances(DensityInstanceCount, DensityIndex);
+	}
+
+	// Create per-bucket prefix sum and sort instance data so we can easily access per-instance data for each density
+	TArray<int32> BucketOffsets;
+	BucketOffsets.SetNumZeroed(NumBuckets);
+
+	for (int32 BucketIndex = 1; BucketIndex < NumBuckets; ++BucketIndex)
+	{
+		BucketOffsets[BucketIndex] = BucketOffsets[BucketIndex - 1] + WaterInstanceData.BucketInstanceCounts[BucketIndex - 1];
+	}
+	
+	WaterInstanceData.StagingInstanceData.StableSort([](const FWaterQuadTree::FStagingInstanceData& Lhs, const FWaterQuadTree::FStagingInstanceData& Rhs)
+		{
+			return Lhs.BucketIndex < Rhs.BucketIndex;
+		});
+
+	FMeshBatch BaseMesh;
+	BaseMesh.ReverseCulling = IsLocalToWorldDeterminantNegative();
+	BaseMesh.Type = PT_TriangleList;
+	BaseMesh.bUseForMaterial = true;
+	BaseMesh.CastShadow = false;
+	BaseMesh.CastRayTracedShadow = false;
+	BaseMesh.SegmentIndex = 0;
+	BaseMesh.Elements.AddZeroed();
+
+	for (int32 DensityIndex = 0; DensityIndex < DensityCount; ++DensityIndex)
+	{
+		int32 DensityInstanceIndex = 0;
+		
+		BaseMesh.VertexFactory = WaterVertexFactories[DensityIndex];
+
+		FMeshBatchElement& BatchElement = BaseMesh.Elements[0];
+
+		BatchElement.NumInstances = 1;
+
+		BatchElement.FirstIndex = 0;
+		BatchElement.NumPrimitives = WaterVertexFactories[DensityIndex]->IndexBuffer->GetIndexCount() / 3;
+		BatchElement.MinVertexIndex = 0;
+		BatchElement.MaxVertexIndex = WaterVertexFactories[DensityIndex]->VertexBuffer->GetVertexCount() - 1;
+
+		// Don't use primitive buffer
+		BatchElement.IndexBuffer = WaterVertexFactories[DensityIndex]->IndexBuffer;
+		BatchElement.PrimitiveIdMode = PrimID_ForceZero;
+		BatchElement.PrimitiveUniformBufferResource = &GIdentityPrimitiveUniformBuffer;
+
+		for (int32 MaterialIndex = 0; MaterialIndex < NumWaterMaterials; ++MaterialIndex)
+		{
+			const int32 BucketIndex = MaterialIndex * DensityCount + DensityIndex;
+			const int32 InstanceCount = WaterInstanceData.BucketInstanceCounts[BucketIndex];
+
+			if (!InstanceCount)
+			{
+				continue;
+			}
+
+			const FMaterialRenderProxy* MaterialRenderProxy = WaterQuadTree.GetWaterMaterials()[MaterialIndex];
+			check(MaterialRenderProxy != nullptr);
+
+			BaseMesh.MaterialRenderProxy = MaterialRenderProxy;
+
+			for (int32 InstanceIndex = 0; InstanceIndex < InstanceCount; ++InstanceIndex)
+			{
+				using FWaterVertexFactoryUserDataWrapperType = TWaterVertexFactoryUserDataWrapper<WITH_WATER_SELECTION_SUPPORT>;
+				FWaterVertexFactoryUserDataWrapperType& UserDataWrapper = Context.RayTracingMeshResourceCollector.AllocateOneFrameResource<FWaterVertexFactoryUserDataWrapperType>();
+
+				const int32 InstanceDataIndex = BucketOffsets[BucketIndex] + InstanceIndex;
+				const FWaterQuadTree::FStagingInstanceData& InstanceData = WaterInstanceData.StagingInstanceData[InstanceDataIndex];
+
+				FWaterVertexFactoryRaytracingParameters UniformBufferParams;
+				UniformBufferParams.VertexBuffer = WaterVertexFactories[DensityIndex]->VertexBuffer->GetSRV();
+				UniformBufferParams.InstanceData0 = InstanceData.Data[0];
+				UniformBufferParams.InstanceData1 = InstanceData.Data[1];
+
+				UserDataWrapper.UserData.InstanceDataBuffers = WaterMeshUserDataBuffers->GetUserData(EWaterMeshRenderGroupType::RG_RenderWaterTiles)->InstanceDataBuffers;
+				UserDataWrapper.UserData.RenderGroupType = EWaterMeshRenderGroupType::RG_RenderWaterTiles;
+				UserDataWrapper.UserData.WaterVertexFactoryRaytracingVFUniformBuffer = FWaterVertexFactoryRaytracingParametersRef::CreateUniformBufferImmediate(UniformBufferParams, UniformBuffer_SingleFrame);
+							
+				BatchElement.UserData = (void*)&UserDataWrapper.UserData;							
+
+				FRayTracingWaterData& WaterInstanceRayTracingData = RayTracingWaterData[DensityIndex][DensityInstanceIndex++];
+
+				FRayTracingInstance RayTracingInstance;
+				RayTracingInstance.Geometry = &WaterInstanceRayTracingData.Geometry;
+				RayTracingInstance.InstanceTransforms.Add(GetLocalToWorld());
+				RayTracingInstance.Materials.Add(BaseMesh);
+				RayTracingInstance.BuildInstanceMaskAndFlags(GetScene().GetFeatureLevel());
+				OutRayTracingInstances.Add(RayTracingInstance);
+
+				Context.DynamicRayTracingGeometriesToUpdate.Add(
+					FRayTracingDynamicGeometryUpdateParams
+					{
+						RayTracingInstance.Materials,
+						false,
+						uint32(WaterVertexFactories[DensityIndex]->VertexBuffer->GetVertexCount()),
+						uint32(WaterVertexFactories[DensityIndex]->VertexBuffer->GetVertexCount() * sizeof(FVector3f)),
+						uint32(WaterVertexFactories[DensityIndex]->IndexBuffer->GetIndexCount() / 3),
+						&WaterInstanceRayTracingData.Geometry,
+						nullptr,
+						true
+					}
+				);				
+			}
+		}
+	}
+}
+#endif // RHI_RAYTRACING
+
 FPrimitiveViewRelevance FWaterMeshSceneProxy::GetViewRelevance(const FSceneView* View) const
 {
 	FPrimitiveViewRelevance Result;
diff --git a/Engine/Plugins/Experimental/Water/Source/Runtime/Private/WaterVertexFactory.cpp b/Engine/Plugins/Experimental/Water/Source/Runtime/Private/WaterVertexFactory.cpp
--- a/Engine/Plugins/Experimental/Water/Source/Runtime/Private/WaterVertexFactory.cpp
+++ b/Engine/Plugins/Experimental/Water/Source/Runtime/Private/WaterVertexFactory.cpp
@@ -10,6 +10,7 @@
 #include "WaterInstanceDataBuffer.h"
 
 IMPLEMENT_GLOBAL_SHADER_PARAMETER_STRUCT(FWaterVertexFactoryParameters, "WaterVF");
+IMPLEMENT_GLOBAL_SHADER_PARAMETER_STRUCT(FWaterVertexFactoryRaytracingParameters, "WaterRaytracingVF");
 
 /**
  * Shader parameters for water vertex factory.
@@ -49,6 +50,13 @@ public:
 
 		ShaderBindings.Add(Shader->GetUniformBufferParameter<FWaterVertexFactoryParameters>(), VertexFactory->GetWaterVertexFactoryUniformBuffer(WaterMeshUserData->RenderGroupType));
 
+#if RHI_RAYTRACING
+		if (IsRayTracingEnabled())
+		{
+			ShaderBindings.Add(Shader->GetUniformBufferParameter<FWaterVertexFactoryRaytracingParameters>(), WaterMeshUserData->WaterVertexFactoryRaytracingVFUniformBuffer);
+		}
+#endif
+
 		if (VertexStreams.Num() > 0)
 		{
 			for (int32 i = 0; i < WaterInstanceDataBuffersType::NumBuffers; ++i)
@@ -74,11 +82,17 @@ public:
 // Always implement the basic vertex factory so that it's there for both editor and non-editor builds :
 IMPLEMENT_TEMPLATE_TYPE_LAYOUT(template<>, TWaterVertexFactoryShaderParameters</*bWithWaterSelectionSupport = */ false>);
 IMPLEMENT_VERTEX_FACTORY_PARAMETER_TYPE(TWaterVertexFactory</*bWithWaterSelectionSupport = */ false>, SF_Vertex, TWaterVertexFactoryShaderParameters</*bWithWaterSelectionSupport = */ false>);
+#if RHI_RAYTRACING
+IMPLEMENT_VERTEX_FACTORY_PARAMETER_TYPE(TWaterVertexFactory</*bWithWaterSelectionSupport = */ false>, SF_Compute, TWaterVertexFactoryShaderParameters</*bWithWaterSelectionSupport = */ false>);
+IMPLEMENT_VERTEX_FACTORY_PARAMETER_TYPE(TWaterVertexFactory</*bWithWaterSelectionSupport = */ false>, SF_RayHitGroup, TWaterVertexFactoryShaderParameters</*bWithWaterSelectionSupport = */ false>);
+#endif // RHI_RAYTRACING
 IMPLEMENT_TEMPLATE_VERTEX_FACTORY_TYPE(template<>, TWaterVertexFactory</*bWithWaterSelectionSupport = */ false>, "/Plugin/Water/Private/WaterMeshVertexFactory.ush",
 	  EVertexFactoryFlags::UsedWithMaterials
 	| EVertexFactoryFlags::SupportsDynamicLighting
 	| EVertexFactoryFlags::SupportsPrecisePrevWorldPos
 	| EVertexFactoryFlags::SupportsPrimitiveIdStream
+	| EVertexFactoryFlags::SupportsRayTracing
+	| EVertexFactoryFlags::SupportsRayTracingDynamicGeometry
 );
 
 #if WITH_WATER_SELECTION_SUPPORT
@@ -86,11 +100,17 @@ IMPLEMENT_TEMPLATE_VERTEX_FACTORY_TYPE(template<>, TWaterVertexFactory</*bWithWa
 // In editor builds, also implement the vertex factory that supports water selection:
 IMPLEMENT_TEMPLATE_TYPE_LAYOUT(template<>, TWaterVertexFactoryShaderParameters</*bWithWaterSelectionSupport = */ true>);
 IMPLEMENT_VERTEX_FACTORY_PARAMETER_TYPE(TWaterVertexFactory</*bWithWaterSelectionSupport = */ true>, SF_Vertex, TWaterVertexFactoryShaderParameters</*bWithWaterSelectionSupport = */ true>);
+#if RHI_RAYTRACING
+IMPLEMENT_VERTEX_FACTORY_PARAMETER_TYPE(TWaterVertexFactory</*bWithWaterSelectionSupport = */ true>, SF_Compute, TWaterVertexFactoryShaderParameters</*bWithWaterSelectionSupport = */ true>);
+IMPLEMENT_VERTEX_FACTORY_PARAMETER_TYPE(TWaterVertexFactory</*bWithWaterSelectionSupport = */ true>, SF_RayHitGroup, TWaterVertexFactoryShaderParameters</*bWithWaterSelectionSupport = */ true>);
+#endif
 IMPLEMENT_TEMPLATE_VERTEX_FACTORY_TYPE(template<>, TWaterVertexFactory</*bWithWaterSelectionSupport = */ true>, "/Plugin/Water/Private/WaterMeshVertexFactory.ush",
 	  EVertexFactoryFlags::UsedWithMaterials
 	| EVertexFactoryFlags::SupportsDynamicLighting
 	| EVertexFactoryFlags::SupportsPrecisePrevWorldPos
 	| EVertexFactoryFlags::SupportsPrimitiveIdStream
+	| EVertexFactoryFlags::SupportsRayTracing
+	| EVertexFactoryFlags::SupportsRayTracingDynamicGeometry
 );
 
 #endif // WITH_WATER_SELECTION_SUPPORT
diff --git a/Engine/Plugins/Experimental/Water/Source/Runtime/Public/WaterMeshSceneProxy.h b/Engine/Plugins/Experimental/Water/Source/Runtime/Public/WaterMeshSceneProxy.h
--- a/Engine/Plugins/Experimental/Water/Source/Runtime/Public/WaterMeshSceneProxy.h
+++ b/Engine/Plugins/Experimental/Water/Source/Runtime/Public/WaterMeshSceneProxy.h
@@ -55,7 +55,38 @@ public:
 	using WaterInstanceDataBuffersType = TWaterInstanceDataBuffers<WITH_WATER_SELECTION_SUPPORT>;
 	using WaterMeshUserDataBuffersType = TWaterMeshUserDataBuffers<WITH_WATER_SELECTION_SUPPORT>;
 
+#if RHI_RAYTRACING
+	virtual void GetDynamicRayTracingInstances(FRayTracingMaterialGatheringContext& Context, TArray<FRayTracingInstance>& OutRayTracingInstances) override final;
+	virtual bool HasRayTracingRepresentation() const override { return true; }
+	virtual bool IsRayTracingRelevant() const override { return true; }
+#endif
+
 private:
+	struct FWaterLODParams
+	{
+		int32 LowestLOD;
+		float HeightLODFactor;
+		float WaterHeightForLOD;
+	};
+
+#if RHI_RAYTRACING
+	struct FRayTracingWaterData
+	{
+		FRayTracingGeometry Geometry;
+		FRWBuffer DynamicVertexBuffer;
+	};
+#endif
+
+	bool HasWaterData() const 
+	{
+		return WaterQuadTree.GetNodeCount() != 0 && DensityCount != 0;
+	}
+
+	FWaterLODParams GetWaterLODParams(const FVector& Position) const;
+
+#if RHI_RAYTRACING
+	void SetupRayTracingInstances(int32 NumInstances, uint32 DensityIndex);
+#endif
 
 	FMaterialRelevance MaterialRelevance;
 
@@ -80,4 +111,9 @@ private:
 	int32 ForceCollapseDensityLevel = TNumericLimits<int32>::Max();
 
 	mutable int32 HistoricalMaxViewInstanceCount = 0;
+
+#if RHI_RAYTRACING
+	// Per density array of ray tracing geometries.
+	TArray<TArray<FRayTracingWaterData>> RayTracingWaterData;	
+#endif
 };
diff --git a/Engine/Plugins/Experimental/Water/Source/Runtime/Public/WaterVertexFactory.h b/Engine/Plugins/Experimental/Water/Source/Runtime/Public/WaterVertexFactory.h
--- a/Engine/Plugins/Experimental/Water/Source/Runtime/Public/WaterVertexFactory.h
+++ b/Engine/Plugins/Experimental/Water/Source/Runtime/Public/WaterVertexFactory.h
@@ -22,6 +22,14 @@ END_GLOBAL_SHADER_PARAMETER_STRUCT()
 
 typedef TUniformBufferRef<FWaterVertexFactoryParameters> FWaterVertexFactoryBufferRef;
 
+BEGIN_GLOBAL_SHADER_PARAMETER_STRUCT(FWaterVertexFactoryRaytracingParameters, )
+	SHADER_PARAMETER_SRV(Buffer<float>, VertexBuffer)
+	SHADER_PARAMETER(FVector4f, InstanceData0)
+	SHADER_PARAMETER(FVector4f, InstanceData1)
+END_GLOBAL_SHADER_PARAMETER_STRUCT()
+
+typedef TUniformBufferRef<FWaterVertexFactoryRaytracingParameters> FWaterVertexFactoryRaytracingParametersRef;
+
 class FWaterMeshIndexBuffer : public FIndexBuffer
 {
 public:
@@ -110,7 +118,7 @@ public:
 		NumVerts = NumVertsPerSide * NumVertsPerSide;
 
 		FRHIResourceCreateInfo CreateInfo(TEXT("FWaterMeshVertexBuffer"));
-		VertexBufferRHI = RHICreateBuffer(sizeof(FVector4f) * NumVerts, BUF_Static | BUF_VertexBuffer, 0, ERHIAccess::VertexOrIndexBuffer, CreateInfo);
+		VertexBufferRHI = RHICreateBuffer(sizeof(FVector4f) * NumVerts, BUF_Static | BUF_VertexBuffer | BUF_ShaderResource, 0, ERHIAccess::VertexOrIndexBuffer | ERHIAccess::SRVMask, CreateInfo);
 		FVector4f* DummyContents = (FVector4f*)RHILockBuffer(VertexBufferRHI, 0, sizeof(FVector4f) * NumVerts, RLM_WriteOnly);
 
 		for (uint32 VertY = 0; VertY < NumVertsPerSide; VertY++)
@@ -127,13 +135,24 @@ public:
 		}
 
 		RHIUnlockBuffer(VertexBufferRHI);
+
+		SRV = RHICreateShaderResourceView(VertexBufferRHI, sizeof(float), PF_R32_FLOAT);
+	}
+
+	virtual void ReleaseRHI() override
+	{
+		SRV.SafeRelease();
+		FVertexBuffer::ReleaseRHI();
 	}
 
 	int32 GetVertexCount() const { return NumVerts; }
+	FRHIShaderResourceView* GetSRV() { return SRV; }
 
 private:
 	int32 NumVerts = 0;
 	const int32 NumQuadsPerSide = 0;
+
+	FShaderResourceViewRHIRef SRV;
 };
 
 enum class EWaterMeshRenderGroupType : uint8
@@ -212,6 +231,10 @@ struct TWaterMeshUserData
 
 	EWaterMeshRenderGroupType RenderGroupType = EWaterMeshRenderGroupType::RG_RenderWaterTiles;
 	const TWaterInstanceDataBuffers<bWithWaterSelectionSupport>* InstanceDataBuffers = nullptr;
+
+#if RHI_RAYTRACING	
+	FUniformBufferRHIRef WaterVertexFactoryRaytracingVFUniformBuffer = nullptr;
+#endif
 };
 
 /**
diff --git a/Engine/Plugins/Experimental/Water/Source/Runtime/Public/WaterVertexFactory.inl b/Engine/Plugins/Experimental/Water/Source/Runtime/Public/WaterVertexFactory.inl
--- a/Engine/Plugins/Experimental/Water/Source/Runtime/Public/WaterVertexFactory.inl
+++ b/Engine/Plugins/Experimental/Water/Source/Runtime/Public/WaterVertexFactory.inl
@@ -160,6 +160,8 @@ void TWaterVertexFactory<bWithWaterSelectionSupport>::ModifyCompilationEnvironme
 		OutEnvironment.SetDefine(TEXT("USE_VERTEXFACTORY_HITPROXY_ID"), TEXT("1"));
 		OutEnvironment.SetDefine(TEXT("WITH_WATER_SELECTION_SUPPORT_VF"), TEXT("1"));
 	}
+
+	OutEnvironment.SetDefine(TEXT("RAY_TRACING_DYNAMIC_MESH_IN_LOCAL_SPACE"), TEXT("1"));
 }
 
 template <bool bWithWaterSelectionSupport>
