Fixed shader compilation issue for water info and merge passes in DX12

diff --git a/Engine/Plugins/Experimental/Water/Shaders/Private/WaterInfoFinalize.usf b/Engine/Plugins/Experimental/Water/Shaders/Private/WaterInfoFinalize.usf
--- a/Engine/Plugins/Experimental/Water/Shaders/Private/WaterInfoFinalize.usf
+++ b/Engine/Plugins/Experimental/Water/Shaders/Private/WaterInfoFinalize.usf
@@ -66,9 +66,10 @@ float2 VelocityBlur(float Radius, uint2 TextureCoordinates)
 }
  
 void Main(
-	in float4 SVPos : SV_POSITION,
+	in FScreenVertexOutput Input,
 	out float4 OutColor : SV_Target0)
 {
+	const float4 SVPos = Input.Position;
 	uint2 TextureCoordinates = floor(SVPos.xy);
 	
 	const FWaterInfo WaterInfo = MakeWaterInfo(WaterInfoTexture.Load(int3(TextureCoordinates, 0)));
diff --git a/Engine/Plugins/Experimental/Water/Shaders/Private/WaterInfoMerge.usf b/Engine/Plugins/Experimental/Water/Shaders/Private/WaterInfoMerge.usf
--- a/Engine/Plugins/Experimental/Water/Shaders/Private/WaterInfoMerge.usf
+++ b/Engine/Plugins/Experimental/Water/Shaders/Private/WaterInfoMerge.usf
@@ -33,10 +33,11 @@ float DecodeFloatRGB(float3 Enc)
 }
 
 void Main(
-	in float4 SVPos : SV_POSITION,
+	in FScreenVertexOutput Input,
 	out float4 OutColor : SV_Target0
 	)
 {
+	const float4 SVPos = Input.Position;
 	uint2 TextureCoordinates = floor(SVPos.xy);
 	
     const float GroundDeviceDepth = DecodeFloatRGB(DepthTexture.Load(int3(TextureCoordinates, 0)).rgb);
