Moved the generation of the water body scene proxy meshes onto the game thread to prevent having to rebuild the geometry data each time the scene proxies are recreated.

diff --git a/Engine/Plugins/Experimental/Water/Source/Runtime/Private/WaterBodyActor.cpp b/Engine/Plugins/Experimental/Water/Source/Runtime/Private/WaterBodyActor.cpp
--- a/Engine/Plugins/Experimental/Water/Source/Runtime/Private/WaterBodyActor.cpp
+++ b/Engine/Plugins/Experimental/Water/Source/Runtime/Private/WaterBodyActor.cpp
@@ -551,9 +551,10 @@ void AWaterBody::PostRegisterAllComponents()
 			WaterBodyComponent->WaterBodyIndex = Manager->AddWaterBodyComponent(WaterBodyComponent);
 			WaterBodyIndex = WaterBodyComponent->WaterBodyIndex;
 
-			// By this point, the water body will have already generated its scene proxy. Since we have just updated the water body index, we will
-			// need to regenerate the scene proxy so that its vertices contain the updated index.
-			WaterBodyComponent->MarkRenderStateDirty();
+			// By this point, the water body will have already generated a scene proxy.
+			// However, the existing scene proxy doesn't have the correct water body index as we have just been assigned one here from the water body manager.
+			// We need to regenerate the mesh and scene proxy so that the vertex colors contain the updated index.
+			WaterBodyComponent->UpdateWaterBodyRenderData();
 		}
 
 		// At this point, the water mesh actor should be ready and we can setup the MID accordingly : 
diff --git a/Engine/Plugins/Experimental/Water/Source/Runtime/Private/WaterBodyComponent.cpp b/Engine/Plugins/Experimental/Water/Source/Runtime/Private/WaterBodyComponent.cpp
--- a/Engine/Plugins/Experimental/Water/Source/Runtime/Private/WaterBodyComponent.cpp
+++ b/Engine/Plugins/Experimental/Water/Source/Runtime/Private/WaterBodyComponent.cpp
@@ -108,7 +108,7 @@ void UWaterBodyComponent::OnHiddenInGameChanged()
 
 FPrimitiveSceneProxy* UWaterBodyComponent::CreateSceneProxy()
 {
-	return new FWaterBodySceneProxy(this);
+	return new FWaterBodySceneProxy(this, WaterBodyMeshVertices, WaterBodyMeshIndices);
 }
 
 void UWaterBodyComponent::GetUsedMaterials(TArray<UMaterialInterface*>& OutMaterialInterfaces, bool bGetDebugMaterials) const
@@ -1162,7 +1162,7 @@ void UWaterBodyComponent::UpdateAll(bool bShapeOrPositionChanged)
 		{
 			FNavigationSystem::UpdateActorAndComponentData(*WaterBodyOwner);
 
-			MarkRenderStateDirty();
+			UpdateWaterBodyRenderData();
 		}
 
 		UpdateComponentVisibility(/* bAllowWaterMeshRebuild = */true);
@@ -1432,6 +1432,17 @@ AWaterZone* UWaterBodyComponent::GetWaterZone() const
 	return nullptr;
 }
 
+void UWaterBodyComponent::UpdateWaterBodyRenderData()
+{
+	GenerateWaterBodyMesh();
+	MarkRenderStateDirty();
+
+	if (AWaterZone* WaterZone = GetWaterZone())
+	{
+		WaterZone->MarkForRebuild(EWaterZoneRebuildFlags::UpdateWaterInfoTexture);
+	}
+}
+
 #if WITH_EDITOR
 
 void UWaterBodyComponent::CreateWaterSpriteComponent()
diff --git a/Engine/Plugins/Experimental/Water/Source/Runtime/Private/WaterBodyLakeComponent.cpp b/Engine/Plugins/Experimental/Water/Source/Runtime/Private/WaterBodyLakeComponent.cpp
--- a/Engine/Plugins/Experimental/Water/Source/Runtime/Private/WaterBodyLakeComponent.cpp
+++ b/Engine/Plugins/Experimental/Water/Source/Runtime/Private/WaterBodyLakeComponent.cpp
@@ -1,6 +1,7 @@
 // Copyright Epic Games, Inc. All Rights Reserved.
 
 #include "WaterBodyLakeComponent.h"
+#include "WaterModule.h"
 #include "WaterSplineComponent.h"
 #include "WaterSubsystem.h"
 #include "Components/SplineMeshComponent.h"
@@ -8,6 +9,10 @@
 #include "Engine/StaticMesh.h"
 #include "LakeCollisionComponent.h"
 #include "Materials/MaterialInstanceDynamic.h"
+#include "Polygon2.h"
+#include "ConstrainedDelaunay2.h"
+#include "Operations/InsetMeshRegion.h"
+#include "DynamicMesh/DynamicMeshChangeTracker.h"
 
 #if WITH_EDITOR
 #include "WaterIconHelper.h"
@@ -44,6 +49,100 @@ TArray<UPrimitiveComponent*> UWaterBodyLakeComponent::GetStandardRenderableCompo
 	return Result;
 }
 
+void UWaterBodyLakeComponent::GenerateWaterBodyMesh()
+{
+	TRACE_CPUPROFILER_EVENT_SCOPE(GenerateLakeMesh);
+
+	using namespace UE::Geometry;
+
+	WaterBodyMeshVertices.Empty();
+	WaterBodyMeshIndices.Empty();
+
+	const UWaterSplineComponent* SplineComp = GetWaterSpline();
+	if (SplineComp == nullptr || SplineComp->GetNumberOfSplineSegments() < 3)
+	{
+		return;
+	}
+
+	FPolygon2d LakePoly;
+	{
+		TArray<FVector> PolyLineVertices;
+		SplineComp->ConvertSplineToPolyLine(ESplineCoordinateSpace::Local, FMath::Square(10.f), PolyLineVertices);
+		
+		for (int32 i = 0; i < PolyLineVertices.Num() - 1; ++i) // skip the last vertex since it's the same as the first vertex
+		{
+			LakePoly.AppendVertex(FVector2D(PolyLineVertices[i]));
+		}
+	}
+	
+	FConstrainedDelaunay2d Triangulation;
+	Triangulation.FillRule = FConstrainedDelaunay2d::EFillRule::Positive;
+	Triangulation.Add(LakePoly);
+	Triangulation.Triangulate();
+
+	if (Triangulation.Triangles.Num() == 0)
+	{
+		return;
+	}
+
+
+	// This FDynamicMesh3 will only be used to compute the inset region for shape dilation
+	FDynamicMesh3 LakeMesh(EMeshComponents::None);
+	for (const FVector2d& Vertex : Triangulation.Vertices)
+	{
+		// push the set of undilated vertices to the persistent mesh
+		FDynamicMeshVertex MeshVertex(FVector(Vertex.X, Vertex.Y, 0.f));
+		MeshVertex.Color = FColor::Black;
+		MeshVertex.TextureCoordinate[0].X = WaterBodyIndex;
+		WaterBodyMeshVertices.Add(MeshVertex);
+
+		LakeMesh.AppendVertex(MeshVertex.Position);
+	}
+
+	for (const FIndex3i& Triangle : Triangulation.Triangles)
+	{
+		WaterBodyMeshIndices.Append({ (uint32)Triangle.A, (uint32)Triangle.B, (uint32)Triangle.C });
+		LakeMesh.AppendTriangle(Triangle);
+	}
+
+	if (ShapeDilation > 0.f)
+	{
+		// Inset the mesh by -ShapeDilation to effectively expand the mesh
+		FInsetMeshRegion Inset(&LakeMesh);
+		Inset.InsetDistance = -1 * ShapeDilation / 2.f;
+
+		for (const FIndex3i& Triangle : LakeMesh.GetTrianglesBuffer())
+		{
+			Inset.Triangles.Add(Triangle.A);
+			Inset.Triangles.Add(Triangle.B);
+			Inset.Triangles.Add(Triangle.C);
+		}
+		
+		if (Inset.Apply())
+		{
+			const uint32 IndexOffset = WaterBodyMeshVertices.Num();
+			for (const FVector3d& Vertex : LakeMesh.GetVerticesBuffer())
+			{
+				// push the set of dilated vertices to the persistent mesh
+				FDynamicMeshVertex MeshVertex(Vertex);
+				MeshVertex.Position.Z = ShapeDilationZOffset;
+				MeshVertex.Color = FColor::Black;
+				MeshVertex.TextureCoordinate[0].X = -1;
+				WaterBodyMeshVertices.Add(MeshVertex);
+			}
+
+			for (const FIndex3i& Triangle : LakeMesh.GetTrianglesBuffer())
+			{
+				WaterBodyMeshIndices.Append({ IndexOffset + Triangle.A, IndexOffset + Triangle.B, IndexOffset + Triangle.C });
+			}
+		}
+		else
+		{
+			UE_LOG(LogWater, Warning, TEXT("Failed to apply mesh inset for shape dilation (%s"), *GetOwner()->GetActorNameOrLabel());
+		}
+	}
+}
+
 FBoxSphereBounds UWaterBodyLakeComponent::CalcBounds(const FTransform& LocalToWorld) const
 {
 	FBox BoxExtent = GetWaterSpline()->GetLocalBounds().GetBox();
diff --git a/Engine/Plugins/Experimental/Water/Source/Runtime/Private/WaterBodyOceanComponent.cpp b/Engine/Plugins/Experimental/Water/Source/Runtime/Private/WaterBodyOceanComponent.cpp
--- a/Engine/Plugins/Experimental/Water/Source/Runtime/Private/WaterBodyOceanComponent.cpp
+++ b/Engine/Plugins/Experimental/Water/Source/Runtime/Private/WaterBodyOceanComponent.cpp
@@ -1,6 +1,7 @@
 // Copyright Epic Games, Inc. All Rights Reserved.
 
 #include "WaterBodyOceanComponent.h"
+#include "WaterModule.h"
 #include "WaterSplineComponent.h"
 #include "Components/SplineMeshComponent.h"
 #include "PhysicsEngine/ConvexElem.h"
@@ -11,6 +12,10 @@
 #include "WaterZoneActor.h"
 #include "UObject/UObjectIterator.h"
 #include "UObject/FortniteMainBranchObjectVersion.h"
+#include "Polygon2.h"
+#include "ConstrainedDelaunay2.h"
+#include "Operations/InsetMeshRegion.h"
+#include "DynamicMesh/DynamicMeshChangeTracker.h"
 
 #if WITH_EDITOR
 #include "WaterIconHelper.h"
@@ -65,7 +70,7 @@ void UWaterBodyOceanComponent::SetVisualExtents(FVector2D NewExtents)
 	if (VisualExtents != NewExtents)
 	{
 		VisualExtents = NewExtents;
-		MarkRenderStateDirty();
+		UpdateWaterBodyRenderData();
 		Modify();
 	}
 }
@@ -122,6 +127,107 @@ void UWaterBodyOceanComponent::Reset()
 	CollisionHullSets.Reset();
 }
 
+void UWaterBodyOceanComponent::GenerateWaterBodyMesh()
+{
+	TRACE_CPUPROFILER_EVENT_SCOPE(GenerateOceanMesh);
+
+	using namespace UE::Geometry;
+
+	WaterBodyMeshVertices.Empty();
+	WaterBodyMeshIndices.Empty();
+
+	const UWaterSplineComponent* SplineComp = GetWaterSpline();
+	
+	if (SplineComp->GetNumberOfSplineSegments() < 3)
+	{
+		return;
+	}
+
+	FPolygon2d Island;
+	TArray<FVector> PolyLineVertices;
+	SplineComp->ConvertSplineToPolyLine(ESplineCoordinateSpace::Local, FMath::Square(10.f), PolyLineVertices);
+	
+	// Construct a 2D polygon describing the central island
+	for (int32 i = PolyLineVertices.Num() - 2; i >= 0; --i) // skip the last vertex since it's the same as the first vertex
+	{
+		Island.AppendVertex(FVector2D(PolyLineVertices[i]));
+	}
+
+	FVector OceanLocation = GetComponentLocation();
+	FPolygon2d OceanBoundingPolygon = FPolygon2d::MakeRectangle(FVector2d(OceanLocation.X, OceanLocation.Y), VisualExtents.X, VisualExtents.Y);
+	FGeneralPolygon2d FinalPoly(OceanBoundingPolygon);
+	FinalPoly.AddHole(Island);
+
+	FConstrainedDelaunay2d Triangulation;
+	Triangulation.FillRule = FConstrainedDelaunay2d::EFillRule::Positive;
+	Triangulation.Add(FinalPoly);
+	Triangulation.Triangulate();
+
+	if (Triangulation.Triangles.Num() == 0)
+	{
+		return;
+	}
+
+	// This FDynamicMesh3 will only be used to compute the inset region for shape dilation
+	FDynamicMesh3 OceanMesh(EMeshComponents::None);
+	for (const FVector2d& Vertex : Triangulation.Vertices)
+	{
+		// push the set of undilated vertices to the persistent mesh
+		FDynamicMeshVertex MeshVertex(FVector(Vertex.X, Vertex.Y, 0.f));
+		MeshVertex.Color = FColor::Black;
+		MeshVertex.TextureCoordinate[0].X = WaterBodyIndex;
+		WaterBodyMeshVertices.Add(MeshVertex);
+
+		OceanMesh.AppendVertex(MeshVertex.Position);
+	}
+
+	for (const FIndex3i& Triangle : Triangulation.Triangles)
+	{
+		WaterBodyMeshIndices.Append({ (uint32)Triangle.A, (uint32)Triangle.B, (uint32)Triangle.C });
+		OceanMesh.AppendTriangle(Triangle);
+	}
+
+
+	if (ShapeDilation > 0.f)
+	{
+		TRACE_CPUPROFILER_EVENT_SCOPE(DilateOceanMesh);
+
+		// Inset the mesh by -ShapeDilation to effectively expand the mesh
+		FInsetMeshRegion Inset(&OceanMesh);
+		Inset.InsetDistance = -1 * ShapeDilation / 2.f;
+
+		for (FIndex3i Triangle : OceanMesh.GetTrianglesBuffer())
+		{
+			Inset.Triangles.Add(Triangle.A);
+			Inset.Triangles.Add(Triangle.B);
+			Inset.Triangles.Add(Triangle.C);
+		}
+		
+		if (Inset.Apply())
+		{
+			const uint32 IndexOffset = WaterBodyMeshVertices.Num();
+			for (const FVector3d& Vertex : OceanMesh.GetVerticesBuffer())
+			{
+				// push the set of dilated vertices to the persistent mesh
+				FDynamicMeshVertex MeshVertex(Vertex);
+				MeshVertex.Position.Z = ShapeDilationZOffset;
+				MeshVertex.Color = FColor::Black;
+				MeshVertex.TextureCoordinate[0].X = -1;
+				WaterBodyMeshVertices.Add(MeshVertex);
+			}
+
+			for (const FIndex3i& Triangle : OceanMesh.GetTrianglesBuffer())
+			{
+				WaterBodyMeshIndices.Append({ IndexOffset + Triangle.A, IndexOffset + Triangle.B, IndexOffset + Triangle.C });
+			}
+		}
+		else
+		{
+			UE_LOG(LogWater, Warning, TEXT("Failed to apply mesh inset for shape dilation (%s"), *GetOwner()->GetActorNameOrLabel());
+		}
+	}	
+}
+
 void UWaterBodyOceanComponent::PostLoad()
 {
 	Super::PostLoad();
diff --git a/Engine/Plugins/Experimental/Water/Source/Runtime/Private/WaterBodyRiverComponent.cpp b/Engine/Plugins/Experimental/Water/Source/Runtime/Private/WaterBodyRiverComponent.cpp
--- a/Engine/Plugins/Experimental/Water/Source/Runtime/Private/WaterBodyRiverComponent.cpp
+++ b/Engine/Plugins/Experimental/Water/Source/Runtime/Private/WaterBodyRiverComponent.cpp
@@ -82,6 +82,238 @@ void UWaterBodyRiverComponent::OnUpdateBody(bool bWithExclusionVolumes)
 	}
 }
 
+// ----------------------------------------------------------------------------------
+
+static FColor PackWaterFlow(float VelocityMagnitude, float DirectionAngle)
+{
+	check((DirectionAngle >= 0.f) && (DirectionAngle <= TWO_PI));
+
+	const float MaxVelocity = FWaterUtils::GetWaterMaxFlowVelocity(false);
+
+	float NormalizedMagnitude = FMath::Clamp(VelocityMagnitude, 0.f, MaxVelocity) / MaxVelocity;
+	float NormalizedAngle = DirectionAngle / TWO_PI;
+	float MappedMagnitude = NormalizedMagnitude * TNumericLimits<uint16>::Max();
+	float MappedAngle = NormalizedAngle * TNumericLimits<uint16>::Max();
+
+	uint16 ResultMag = (uint16)MappedMagnitude;
+	uint16 ResultAngle = (uint16)MappedAngle;
+	
+	FColor Result;
+	Result.R = (ResultMag >> 8) & 0xFF;
+	Result.G = (ResultMag >> 0) & 0xFF;
+	Result.B = (ResultAngle >> 8) & 0xFF;
+	Result.A = (ResultAngle >> 0) & 0xFF;
+
+	return Result;
+}
+
+static void AddVerticesForRiverSplineStep(float DistanceAlongSpline, const UWaterBodyRiverComponent* Component, TArray<FDynamicMeshVertex>& Vertices, TArray<uint32>& Indices)
+{
+	const UWaterSplineComponent* SplineComp = Component->GetWaterSpline();
+	
+	const FVector Tangent = SplineComp->GetTangentAtDistanceAlongSpline(DistanceAlongSpline, ESplineCoordinateSpace::Local).GetSafeNormal();
+	const FVector Up = SplineComp->GetUpVectorAtDistanceAlongSpline(DistanceAlongSpline, ESplineCoordinateSpace::Local).GetSafeNormal();
+
+	const FVector Normal = FVector::CrossProduct(Tangent, Up).GetSafeNormal();
+	const FVector Pos = SplineComp->GetLocationAtDistanceAlongSpline(DistanceAlongSpline, ESplineCoordinateSpace::Local);
+
+	const float Key = SplineComp->SplineCurves.ReparamTable.Eval(DistanceAlongSpline, 0.f);
+	const float HalfWidth = Cast<UWaterSplineMetadata>(SplineComp->GetSplinePointsMetadata())->RiverWidth.Eval(Key) / 2;
+	float Velocity = Cast<UWaterSplineMetadata>(SplineComp->GetSplinePointsMetadata())->WaterVelocityScalar.Eval(Key);
+
+	// Distance from the center of the spline to place our first vertices
+	FVector OutwardDistance = Normal * HalfWidth;
+	// Prevent there being a relative height difference between the two vertices even when the spline has a slight roll to it
+	OutwardDistance.Z = 0.f;
+
+	const float DilationAmount = Component->ShapeDilation;
+	const FVector DilationOffset = Normal * DilationAmount;
+
+	FDynamicMeshVertex Left(Pos - OutwardDistance);
+	FDynamicMeshVertex Right(Pos + OutwardDistance);
+
+	FDynamicMeshVertex DilatedFarLeft(Pos - OutwardDistance - DilationOffset);
+	DilatedFarLeft.Position.Z += Component->GetShapeDilationZOffsetFar();
+	FDynamicMeshVertex DilatedLeft(Pos - OutwardDistance);
+	DilatedLeft.Position.Z += Component->GetShapeDilationZOffset();
+	FDynamicMeshVertex DilatedRight(Pos + OutwardDistance);
+	DilatedRight.Position.Z += Component->GetShapeDilationZOffset();
+	FDynamicMeshVertex DilatedFarRight(Pos + OutwardDistance + DilationOffset);
+	DilatedFarRight.Position.Z += Component->GetShapeDilationZOffsetFar();
+
+	float FlowDirection = Tangent.HeadingAngle() + FMath::DegreesToRadians(Component->GetRelativeRotation().Yaw);
+	// Convert negative angles into positive angles
+	if (FlowDirection < 0.f)
+	{
+		FlowDirection = TWO_PI + FlowDirection;
+	}
+
+	// If negative velocity, inverse the direction and change the velocity back to positive.
+	if (Velocity < 0.f)
+	{
+		Velocity *= -1.f;
+		FlowDirection = FMath::Fmod(PI + FlowDirection, TWO_PI);
+	}
+
+	const FColor EmptyFlowData = FColor(0.f);
+
+	const FColor PackedFlowData = PackWaterFlow(Velocity, FlowDirection);
+	Left.Color = PackedFlowData;
+	Right.Color = PackedFlowData;
+
+	DilatedFarLeft.Color = EmptyFlowData;
+	DilatedLeft.Color = EmptyFlowData;
+	DilatedRight.Color = EmptyFlowData;
+	DilatedFarRight.Color = EmptyFlowData;
+
+	// Embed the water body index in the vertex data so that we can distinguish between dilated and undilated regions of the texture
+	const int32 WaterBodyIndex = Component->GetWaterBodyIndex();
+	Left.TextureCoordinate[0].X = WaterBodyIndex;
+	Right.TextureCoordinate[0].X = WaterBodyIndex;
+	
+	DilatedFarLeft.TextureCoordinate[0].X = -1;
+	DilatedLeft.TextureCoordinate[0].X = -1;
+	DilatedRight.TextureCoordinate[0].X = -1;
+	DilatedFarRight.TextureCoordinate[0].X = -1;
+
+	/* River segment geometry:
+		6---7,8---9,10--11
+		| /  |  /  |  / |
+		0---1,2---3,4---5
+	*/
+	const uint32 BaseIndex = Vertices.Num();
+	Vertices.Append({ DilatedFarLeft, DilatedLeft, Left, Right, DilatedRight, DilatedFarRight });
+	// Append left dilation quad
+	Indices.Append({ BaseIndex, BaseIndex + 7, BaseIndex + 1, BaseIndex, BaseIndex + 6, BaseIndex + 7 });
+	// Append main quad
+	Indices.Append({ BaseIndex + 2, BaseIndex + 9, BaseIndex + 3, BaseIndex + 2, BaseIndex + 8, BaseIndex + 9});
+	// Append right dilation quad
+	Indices.Append({ BaseIndex + 4, BaseIndex + 11, BaseIndex + 5, BaseIndex + 4, BaseIndex + 10, BaseIndex + 11 });
+}
+
+enum class ERiverBoundaryEdge {
+	Start,
+	End,
+};
+
+static void AddTerminalVerticesForRiverSpline(ERiverBoundaryEdge Edge, const UWaterBodyRiverComponent* Component, TArray<FDynamicMeshVertex>& Vertices, TArray<uint32>& Indices)
+{
+	const UWaterSplineComponent* SplineComp = Component->GetWaterSpline();
+
+	const float DistanceAlongSpline = (Edge == ERiverBoundaryEdge::Start) ? 0.f : SplineComp->GetSplineLength();
+	
+	const FVector Tangent = SplineComp->GetTangentAtDistanceAlongSpline(DistanceAlongSpline, ESplineCoordinateSpace::Local).GetSafeNormal();
+	const FVector Up = SplineComp->GetUpVectorAtDistanceAlongSpline(DistanceAlongSpline, ESplineCoordinateSpace::Local).GetSafeNormal();
+	
+	const FVector Normal = FVector::CrossProduct(Tangent, Up).GetSafeNormal();
+	const FVector Pos = SplineComp->GetLocationAtDistanceAlongSpline(DistanceAlongSpline, ESplineCoordinateSpace::Local);
+	
+	const float Key = SplineComp->SplineCurves.ReparamTable.Eval(DistanceAlongSpline, 0.f);
+	const float HalfWidth = Cast<UWaterSplineMetadata>(SplineComp->GetSplinePointsMetadata())->RiverWidth.Eval(Key) / 2;
+
+	const float DilationAmount = Component->ShapeDilation;
+	const FVector DilationOffset = Normal * DilationAmount;
+	FVector OutwardDistance = Normal * HalfWidth;
+	OutwardDistance.Z = 0.f;
+
+	FVector TangentialOffset = Tangent * DilationAmount;
+	TangentialOffset.Z = 0.f;
+
+	// For the starting edge the tangential offset is negative to push it backwards
+	if (Edge == ERiverBoundaryEdge::Start)
+	{
+		TangentialOffset *= -1;
+	}
+
+	FDynamicMeshVertex BackLeft(Pos - OutwardDistance + TangentialOffset - DilationOffset);
+	BackLeft.Position.Z += Component->GetShapeDilationZOffsetFar();
+	FDynamicMeshVertex Left(Pos - OutwardDistance + TangentialOffset);
+	Left.Position.Z += Component->GetShapeDilationZOffsetFar();
+	FDynamicMeshVertex Right(Pos + OutwardDistance + TangentialOffset);
+	Right.Position.Z += Component->GetShapeDilationZOffsetFar();
+	FDynamicMeshVertex BackRight(Pos + OutwardDistance + TangentialOffset + DilationOffset);
+	BackRight.Position.Z += Component->GetShapeDilationZOffsetFar();
+
+	// Initialize the vertex data to correct represent what we expect the dilated region to look like
+	// (no color and -1 UVs[0].x
+	BackLeft.Color = FColor(0.f);
+	BackLeft.TextureCoordinate[0].X = -1;
+	Left.Color = FColor(0.f);
+	Left.TextureCoordinate[0].X = -1;
+	Right.Color = FColor(0.f);
+	Right.TextureCoordinate[0].X = -1;
+	BackRight.Color = FColor(0.f);
+	BackRight.TextureCoordinate[0].X = -1;
+	
+
+	Vertices.Append({ BackLeft, Left, Right, BackRight });
+	if (Edge == ERiverBoundaryEdge::Start)
+	{
+		/* Dilated front segment geometry: (ignore vertices 6-7 since they are the non-dilated vertices)
+			4---5,6---7,8---9
+			|    |     |    |
+			0----1-----2----3
+		*/
+		Indices.Append({ 0, 5, 1, 0, 4, 5});
+		Indices.Append({ 1, 8, 2, 1, 5, 8});
+		Indices.Append({ 2, 9, 3, 2, 8, 9});
+	}
+	else
+	{
+		/* Dilated back segment geometry: (ignore vertices 6-7 since they are the non-dilated vertices)
+			6----7-----8----9
+			|    |     |    |
+			0---1,2---3,4---5
+		*/
+		check(Vertices.Num() >= 10);
+		const uint32 BaseIndex = Vertices.Num() - 10;
+		// Since we don't know if we're at the end point or not in the AddVerticesForRiverSplineStep function,
+		// we end up adding a bunch of indices which link to the next set of vertices but they don't exist so we must trim them here.
+		Indices.SetNum(Indices.Num() - 18);
+
+		Indices.Append({ BaseIndex + 0, BaseIndex + 6, BaseIndex + 7, BaseIndex + 0, BaseIndex + 7, BaseIndex + 1 });
+		Indices.Append({ BaseIndex + 1, BaseIndex + 7, BaseIndex + 8, BaseIndex + 1, BaseIndex + 8, BaseIndex + 4 });
+		Indices.Append({ BaseIndex + 4, BaseIndex + 9, BaseIndex + 5, BaseIndex + 4, BaseIndex + 8, BaseIndex + 9 });
+	}
+}
+
+// ----------------------------------------------------------------------------------
+
+void UWaterBodyRiverComponent::GenerateWaterBodyMesh()
+{
+	TRACE_CPUPROFILER_EVENT_SCOPE(GenerateRiverMesh);
+
+	WaterBodyMeshVertices.Empty();
+	WaterBodyMeshIndices.Empty();
+
+	const UWaterSplineComponent* SplineComp = GetWaterSpline();
+	if (SplineComp == nullptr)
+	{
+		return;
+	}
+
+	TArray<FDynamicMeshVertex> Vertices;
+	TArray<uint32> Indices;
+
+	// Add an extra point at the start to dilate starting edge
+	AddTerminalVerticesForRiverSpline(ERiverBoundaryEdge::Start, this, Vertices, Indices);
+
+	TArray<double> Distances;
+	TArray<FVector> Points;
+	SplineComp->DivideSplineIntoPolylineRecursiveWithDistances(0.f, SplineComp->GetSplineLength(), ESplineCoordinateSpace::Local, FMath::Square(10.f), Points, Distances);
+	
+	for (double DistanceAlongSpline : Distances)
+	{
+		AddVerticesForRiverSplineStep(DistanceAlongSpline, this, Vertices, Indices);
+	}
+	
+	// Add an extra point at the end to dilate ending edge
+	AddTerminalVerticesForRiverSpline(ERiverBoundaryEdge::End, this, Vertices, Indices);
+
+	WaterBodyMeshVertices = MoveTemp(Vertices);
+	WaterBodyMeshIndices = MoveTemp(Indices);
+}
+
 FBoxSphereBounds UWaterBodyRiverComponent::CalcBounds(const FTransform& LocalToWorld) const
 {
 	FBox BoxExtent(ForceInit);
diff --git a/Engine/Plugins/Experimental/Water/Source/Runtime/Private/WaterBodySceneProxy.cpp b/Engine/Plugins/Experimental/Water/Source/Runtime/Private/WaterBodySceneProxy.cpp
--- a/Engine/Plugins/Experimental/Water/Source/Runtime/Private/WaterBodySceneProxy.cpp
+++ b/Engine/Plugins/Experimental/Water/Source/Runtime/Private/WaterBodySceneProxy.cpp
@@ -5,17 +5,6 @@
 #include "WaterModule.h"
 #include "WaterUtils.h"
 #include "WaterBodyComponent.h"
-#include "WaterBodyOceanComponent.h"
-#include "WaterBodyRiverComponent.h"
-#include "WaterBodyLakeComponent.h"
-#include "WaterSplineComponent.h"
-#include "Algo/Transform.h"
-#include "MeshBoundaryLoops.h"
-#include "Curve/GeneralPolygon2.h"
-#include "Operations/MinimalHoleFiller.h"
-#include "DynamicMesh/DynamicMeshChangeTracker.h"
-#include "Operations/InsetMeshRegion.h"
-#include "ConstrainedDelaunay2.h"
 #include "Materials/MaterialInstanceDynamic.h"
 #include "SceneInterface.h"
 
@@ -40,420 +29,6 @@ public:
 	{}
 };
 
-// ----------------------------------------------------------------------------------
-
-static FColor PackWaterFlow(float VelocityMagnitude, float DirectionAngle)
-{
-	check((DirectionAngle >= 0.f) && (DirectionAngle <= TWO_PI));
-
-	const float MaxVelocity = FWaterUtils::GetWaterMaxFlowVelocity(false);
-
-	float NormalizedMagnitude = FMath::Clamp(VelocityMagnitude, 0.f, MaxVelocity) / MaxVelocity;
-	float NormalizedAngle = DirectionAngle / TWO_PI;
-	float MappedMagnitude = NormalizedMagnitude * TNumericLimits<uint16>::Max();
-	float MappedAngle = NormalizedAngle * TNumericLimits<uint16>::Max();
-
-	uint16 ResultMag = (uint16)MappedMagnitude;
-	uint16 ResultAngle = (uint16)MappedAngle;
-	
-	FColor Result;
-	Result.R = (ResultMag >> 8) & 0xFF;
-	Result.G = (ResultMag >> 0) & 0xFF;
-	Result.B = (ResultAngle >> 8) & 0xFF;
-	Result.A = (ResultAngle >> 0) & 0xFF;
-
-	return Result;
-}
-
-// ----------------------------------------------------------------------------------
-
-using namespace UE::Geometry;
-
-void FWaterBodySceneProxy::AddSectionFromDynamicMesh(const FDynamicMesh3& DynamicMesh, TFunctionRef<FDynamicMeshVertex(const FVector&)> VertTransformFunc)
-{
-	TArray<FDynamicMeshVertex> Vertices;
-	Vertices.Reserve(DynamicMesh.VertexCount());
-
-	Algo::Transform(DynamicMesh.GetVerticesBuffer(), Vertices, VertTransformFunc);
-
-	FWaterBodyMeshSection* NewSection = new FWaterBodyMeshSection(GetScene().GetFeatureLevel());
-	NewSection->IndexBuffer.Indices.Reserve(DynamicMesh.TriangleCount() * 3);
-	for (const FIndex3i& Index : DynamicMesh.GetTrianglesBuffer())
-	{
-		NewSection->IndexBuffer.Indices.Add(Index.A);
-		NewSection->IndexBuffer.Indices.Add(Index.B);
-		NewSection->IndexBuffer.Indices.Add(Index.C);
-	}
-	
-	NewSection->VertexBuffers.InitFromDynamicVertex(&NewSection->VertexFactory, Vertices);
-	InitResources(NewSection);
-
-	Sections.Add(NewSection);
-}
-
-void FWaterBodySceneProxy::GenerateLakeMesh(UWaterBodyLakeComponent* Component)
-{
-	TRACE_CPUPROFILER_EVENT_SCOPE(GenerateLakeMesh);
-
-	const UWaterSplineComponent* SplineComp = Component->GetWaterSpline();
-	if (SplineComp->GetNumberOfSplineSegments() < 3)
-	{
-		return;
-	}
-	FPolygon2d LakePoly;
-	{
-		TArray<FVector> PolyLineVertices;
-		SplineComp->ConvertSplineToPolyLine(ESplineCoordinateSpace::Local, FMath::Square(10.f), PolyLineVertices);
-		
-		for (int32 i = 0; i < PolyLineVertices.Num() - 1; ++i) // skip the last vertex since it's the same as the first vertex
-		{
-			LakePoly.AppendVertex(FVector2D(PolyLineVertices[i]));
-		}
-	}
-	
-	FConstrainedDelaunay2d Triangulation;
-	Triangulation.FillRule = FConstrainedDelaunay2d::EFillRule::Positive;
-	Triangulation.Add(LakePoly);
-	Triangulation.Triangulate();
-
-	if (Triangulation.Triangles.Num() == 0)
-	{
-		return;
-	}
-
-	FDynamicMesh3 LakeMesh(EMeshComponents::None);
-	for (const FVector2d& Vertex : Triangulation.Vertices)
-	{
-		LakeMesh.AppendVertex(FVector3d(Vertex.X, Vertex.Y, 0));
-	}
-	for (const FIndex3i& Triangle : Triangulation.Triangles)
-	{
-		LakeMesh.AppendTriangle(Triangle);
-	}
-
-	AddSectionFromDynamicMesh(LakeMesh, [WaterBodyIndex = Component->GetWaterBodyIndex()](const FVector& In) {
-			FDynamicMeshVertex Result(FVector3f(In.X, In.Y, 0.f));
-			Result.Color = FColor::Black;
-			Result.TextureCoordinate[0].X = WaterBodyIndex;
-			return Result;
-		});
-	
-	if (Component->ShapeDilation > 0.f)
-	{
-		// Inset the mesh by -ShapeDilation to effectively expand the mesh
-		FInsetMeshRegion Inset(&LakeMesh);
-		Inset.InsetDistance = -1 * Component->ShapeDilation / 2.f;
-
-		for (const FIndex3i& Triangle : LakeMesh.GetTrianglesBuffer())
-		{
-			Inset.Triangles.Add(Triangle.A);
-			Inset.Triangles.Add(Triangle.B);
-			Inset.Triangles.Add(Triangle.C);
-		}
-		
-		if (!Inset.Apply())
-		{
-			UE_LOG(LogWater, Warning, TEXT("Failed to apply mesh inset for shape dilation (%s"), *Component->GetOwner()->GetActorNameOrLabel());
-		}
-	}
-	
-	// Push down the dilated region vertices by -ZOffset to prevent dilated regions from overwriting data from adjacent water bodies
-	AddSectionFromDynamicMesh(LakeMesh, [ZOffset = Component->GetShapeDilationZOffset()](const FVector& In) {
-			FDynamicMeshVertex Result(FVector3f(In.X, In.Y, ZOffset));
-			Result.Color = FColor::Black;
-			Result.TextureCoordinate[0].X = -1;
-			return Result;
-		});
-}
-
-void FWaterBodySceneProxy::GenerateOceanMesh(UWaterBodyOceanComponent* Component)
-{
-	TRACE_CPUPROFILER_EVENT_SCOPE(GenerateOceanMesh);
-
-	const UWaterSplineComponent* SplineComp = Component->GetWaterSpline();
-	
-	if (SplineComp->GetNumberOfSplineSegments() < 3)
-	{
-		return;
-	}
-
-	FPolygon2d Island;
-	{
-		TArray<FVector> PolyLineVertices;
-		SplineComp->ConvertSplineToPolyLine(ESplineCoordinateSpace::Local, FMath::Square(10.f), PolyLineVertices);
-		
-		// Construct a 2D polygon describing the central island
-		for (int32 i = PolyLineVertices.Num() - 2; i >= 0; --i) // skip the last vertex since it's the same as the first vertex
-		{
-			Island.AppendVertex(FVector2D(PolyLineVertices[i]));
-		}
-	}
-
-	FVector OceanLocation = Component->GetComponentLocation();
-	FVector2D OceanExtent = Component->GetVisualExtents();
-	FPolygon2d OceanBoundingPolygon = FPolygon2d::MakeRectangle(FVector2d(OceanLocation.X, OceanLocation.Y), OceanExtent.X, OceanExtent.Y);
-	FGeneralPolygon2d FinalPoly(OceanBoundingPolygon);
-	FinalPoly.AddHole(Island);
-
-	FConstrainedDelaunay2d Triangulation;
-	Triangulation.FillRule = FConstrainedDelaunay2d::EFillRule::Positive;
-	Triangulation.Add(FinalPoly);
-	Triangulation.Triangulate();
-
-	if (Triangulation.Triangles.Num() == 0)
-	{
-		return;
-	}
-
-	FDynamicMesh3 OceanMesh(EMeshComponents::None);
-	for (FVector2d Vertex : Triangulation.Vertices)
-	{
-		OceanMesh.AppendVertex(FVector3d(Vertex.X, Vertex.Y, 0));
-	}
-	for (FIndex3i Triangle : Triangulation.Triangles)
-	{
-		OceanMesh.AppendTriangle(Triangle);
-	}
-	AddSectionFromDynamicMesh(OceanMesh, [WaterBodyIndex = Component->GetWaterBodyIndex()](const FVector& In) {
-			FDynamicMeshVertex Result(FVector3f(In.X, In.Y, 0.f));
-			Result.Color = FColor::Black;
-			Result.TextureCoordinate[0].X = WaterBodyIndex;
-			return Result;
-		});
-
-	if (Component->ShapeDilation > 0.f)
-	{
-		// Inset the mesh by -ShapeDilation to effectively expand the mesh
-		FInsetMeshRegion Inset(&OceanMesh);
-		Inset.InsetDistance = -1 * Component->ShapeDilation / 2.f;
-
-		for (FIndex3i Triangle : OceanMesh.GetTrianglesBuffer())
-		{
-			Inset.Triangles.Add(Triangle.A);
-			Inset.Triangles.Add(Triangle.B);
-			Inset.Triangles.Add(Triangle.C);
-		}
-		
-		if (!Inset.Apply())
-		{
-			UE_LOG(LogWater, Warning, TEXT("Failed to apply mesh inset for shape dilation (%s"), *Component->GetOwner()->GetActorNameOrLabel());
-		}
-	}
-	
-	// Push down the dilated region vertices by -ZOffset to prevent dilated regions from overwriting data from adjacent water bodies
-	AddSectionFromDynamicMesh(OceanMesh, [ZOffset = Component->GetShapeDilationZOffset()](const FVector& In) {
-			FDynamicMeshVertex Result(FVector3f(In.X, In.Y, ZOffset));
-			Result.Color = FColor::Black;
-			Result.TextureCoordinate[0].X = -1;
-			return Result;
-		});
-}
-
-static void AddVerticesForRiverSplineStep(float DistanceAlongSpline, const UWaterBodyRiverComponent* Component, TArray<FDynamicMeshVertex>& Vertices, TArray<uint32>& Indices)
-{
-	const UWaterSplineComponent* SplineComp = Component->GetWaterSpline();
-	
-	const FVector Tangent = SplineComp->GetTangentAtDistanceAlongSpline(DistanceAlongSpline, ESplineCoordinateSpace::Local).GetSafeNormal();
-	const FVector Up = SplineComp->GetUpVectorAtDistanceAlongSpline(DistanceAlongSpline, ESplineCoordinateSpace::Local).GetSafeNormal();
-
-	const FVector Normal = FVector::CrossProduct(Tangent, Up).GetSafeNormal();
-	const FVector Pos = SplineComp->GetLocationAtDistanceAlongSpline(DistanceAlongSpline, ESplineCoordinateSpace::Local);
-
-	const float Key = SplineComp->SplineCurves.ReparamTable.Eval(DistanceAlongSpline, 0.f);
-	const float HalfWidth = Cast<UWaterSplineMetadata>(SplineComp->GetSplinePointsMetadata())->RiverWidth.Eval(Key) / 2;
-	float Velocity = Cast<UWaterSplineMetadata>(SplineComp->GetSplinePointsMetadata())->WaterVelocityScalar.Eval(Key);
-
-	// Distance from the center of the spline to place our first vertices
-	FVector OutwardDistance = Normal * HalfWidth;
-	// Prevent there being a relative height difference between the two vertices even when the spline has a slight roll to it
-	OutwardDistance.Z = 0.f;
-
-	const float DilationAmount = Component->ShapeDilation;
-	const FVector DilationOffset = Normal * DilationAmount;
-
-	FDynamicMeshVertex Left(FVector3f(Pos - OutwardDistance));
-	FDynamicMeshVertex Right(FVector3f(Pos + OutwardDistance));
-
-	FDynamicMeshVertex DilatedFarLeft(FVector3f(Pos - OutwardDistance - DilationOffset));
-	DilatedFarLeft.Position.Z += Component->GetShapeDilationZOffsetFar();
-	FDynamicMeshVertex DilatedLeft(FVector3f(Pos - OutwardDistance));
-	DilatedLeft.Position.Z += Component->GetShapeDilationZOffset();
-	FDynamicMeshVertex DilatedRight(FVector3f(Pos + OutwardDistance));
-	DilatedRight.Position.Z += Component->GetShapeDilationZOffset();
-	FDynamicMeshVertex DilatedFarRight(FVector3f(Pos + OutwardDistance + DilationOffset));
-	DilatedFarRight.Position.Z += Component->GetShapeDilationZOffsetFar();
-
-	float FlowDirection = Tangent.HeadingAngle() + FMath::DegreesToRadians(Component->GetRelativeRotation().Yaw);
-	// Convert negative angles into positive angles
-	if (FlowDirection < 0.f)
-	{
-		FlowDirection = TWO_PI + FlowDirection;
-	}
-
-	// If negative velocity, inverse the direction and change the velocity back to positive.
-	if (Velocity < 0.f)
-	{
-		Velocity *= -1.f;
-		FlowDirection = FMath::Fmod(PI + FlowDirection, TWO_PI);
-	}
-
-	const FColor EmptyFlowData = FColor(0.f);
-
-	const FColor PackedFlowData = PackWaterFlow(Velocity, FlowDirection);
-	Left.Color = PackedFlowData;
-	Right.Color = PackedFlowData;
-
-	DilatedFarLeft.Color = EmptyFlowData;
-	DilatedLeft.Color = EmptyFlowData;
-	DilatedRight.Color = EmptyFlowData;
-	DilatedFarRight.Color = EmptyFlowData;
-
-	// Embed the water body index in the vertex data so that we can distinguish between dilated and undilated regions of the texture
-	const int32 WaterBodyIndex = Component->GetWaterBodyIndex();
-	Left.TextureCoordinate[0].X = WaterBodyIndex;
-	Right.TextureCoordinate[0].X = WaterBodyIndex;
-	
-	DilatedFarLeft.TextureCoordinate[0].X = -1;
-	DilatedLeft.TextureCoordinate[0].X = -1;
-	DilatedRight.TextureCoordinate[0].X = -1;
-	DilatedFarRight.TextureCoordinate[0].X = -1;
-
-	/* River segment geometry:
-		6---7,8---9,10--11
-		| /  |  /  |  / |
-		0---1,2---3,4---5
-	*/
-	const uint32 BaseIndex = Vertices.Num();
-	Vertices.Append({ DilatedFarLeft, DilatedLeft, Left, Right, DilatedRight, DilatedFarRight });
-	// Append left dilation quad
-	Indices.Append({ BaseIndex, BaseIndex + 7, BaseIndex + 1, BaseIndex, BaseIndex + 6, BaseIndex + 7 });
-	// Append main quad
-	Indices.Append({ BaseIndex + 2, BaseIndex + 9, BaseIndex + 3, BaseIndex + 2, BaseIndex + 8, BaseIndex + 9});
-	// Append right dilation quad
-	Indices.Append({ BaseIndex + 4, BaseIndex + 11, BaseIndex + 5, BaseIndex + 4, BaseIndex + 10, BaseIndex + 11 });
-}
-
-enum class ERiverBoundaryEdge {
-	Start,
-	End,
-};
-
-static void AddTerminalVerticesForRiverSpline(ERiverBoundaryEdge Edge, const UWaterBodyRiverComponent* Component, TArray<FDynamicMeshVertex>& Vertices, TArray<uint32>& Indices)
-{
-	const UWaterSplineComponent* SplineComp = Component->GetWaterSpline();
-
-	const float DistanceAlongSpline = (Edge == ERiverBoundaryEdge::Start) ? 0.f : SplineComp->GetSplineLength();
-	
-	const FVector Tangent = SplineComp->GetTangentAtDistanceAlongSpline(DistanceAlongSpline, ESplineCoordinateSpace::Local).GetSafeNormal();
-	const FVector Up = SplineComp->GetUpVectorAtDistanceAlongSpline(DistanceAlongSpline, ESplineCoordinateSpace::Local).GetSafeNormal();
-	
-	const FVector Normal = FVector::CrossProduct(Tangent, Up).GetSafeNormal();
-	const FVector Pos = SplineComp->GetLocationAtDistanceAlongSpline(DistanceAlongSpline, ESplineCoordinateSpace::Local);
-	
-	const float Key = SplineComp->SplineCurves.ReparamTable.Eval(DistanceAlongSpline, 0.f);
-	const float HalfWidth = Cast<UWaterSplineMetadata>(SplineComp->GetSplinePointsMetadata())->RiverWidth.Eval(Key) / 2;
-
-	const float DilationAmount = Component->ShapeDilation;
-	const FVector DilationOffset = Normal * DilationAmount;
-	FVector OutwardDistance = Normal * HalfWidth;
-	OutwardDistance.Z = 0.f;
-
-	FVector TangentialOffset = Tangent * DilationAmount;
-	TangentialOffset.Z = 0.f;
-
-	// For the starting edge the tangential offset is negative to push it backwards
-	if (Edge == ERiverBoundaryEdge::Start)
-	{
-		TangentialOffset *= -1;
-	}
-
-	FDynamicMeshVertex BackLeft(FVector3f(Pos - OutwardDistance + TangentialOffset - DilationOffset));
-	BackLeft.Position.Z += Component->GetShapeDilationZOffsetFar();
-	FDynamicMeshVertex Left(FVector3f(Pos - OutwardDistance + TangentialOffset));
-	Left.Position.Z += Component->GetShapeDilationZOffsetFar();
-	FDynamicMeshVertex Right(FVector3f(Pos + OutwardDistance + TangentialOffset));
-	Right.Position.Z += Component->GetShapeDilationZOffsetFar();
-	FDynamicMeshVertex BackRight(FVector3f(Pos + OutwardDistance + TangentialOffset + DilationOffset));
-	BackRight.Position.Z += Component->GetShapeDilationZOffsetFar();
-
-	// Initialize the vertex data to correct represent what we expect the dilated region to look like
-	// (no color and -1 UVs[0].x
-	BackLeft.Color = FColor(0.f);
-	BackLeft.TextureCoordinate[0].X = -1;
-	Left.Color = FColor(0.f);
-	Left.TextureCoordinate[0].X = -1;
-	Right.Color = FColor(0.f);
-	Right.TextureCoordinate[0].X = -1;
-	BackRight.Color = FColor(0.f);
-	BackRight.TextureCoordinate[0].X = -1;
-	
-
-	Vertices.Append({ BackLeft, Left, Right, BackRight });
-	if (Edge == ERiverBoundaryEdge::Start)
-	{
-		/* Dilated front segment geometry: (ignore vertices 6-7 since they are the non-dilated vertices)
-			4---5,6---7,8---9
-			|    |     |    |
-			0----1-----2----3
-		*/
-		Indices.Append({ 0, 5, 1, 0, 4, 5});
-		Indices.Append({ 1, 8, 2, 1, 5, 8});
-		Indices.Append({ 2, 9, 3, 2, 8, 9});
-	}
-	else
-	{
-		/* Dilated back segment geometry: (ignore vertices 6-7 since they are the non-dilated vertices)
-			6----7-----8----9
-			|    |     |    |
-			0---1,2---3,4---5
-		*/
-		check(Vertices.Num() >= 10);
-		const uint32 BaseIndex = Vertices.Num() - 10;
-		// Since we don't know if we're at the end point or not in the AddVerticesForRiverSplineStep function,
-		// we end up adding a bunch of indices which link to the next set of vertices but they don't exist so we must trim them here.
-		Indices.SetNum(Indices.Num() - 18);
-
-		Indices.Append({ BaseIndex + 0, BaseIndex + 6, BaseIndex + 7, BaseIndex + 0, BaseIndex + 7, BaseIndex + 1 });
-		Indices.Append({ BaseIndex + 1, BaseIndex + 7, BaseIndex + 8, BaseIndex + 1, BaseIndex + 8, BaseIndex + 4 });
-		Indices.Append({ BaseIndex + 4, BaseIndex + 9, BaseIndex + 5, BaseIndex + 4, BaseIndex + 8, BaseIndex + 9 });
-	}
-}
-
-void FWaterBodySceneProxy::GenerateRiverMesh(UWaterBodyRiverComponent* Component)
-{
-	TRACE_CPUPROFILER_EVENT_SCOPE(GenerateRiverMesh);
-
-	const UWaterSplineComponent* SplineComp = Component->GetWaterSpline();
-
-	TArray<FDynamicMeshVertex> Vertices;
-	TArray<uint32> Indices;
-
-	// Add an extra point at the start to dilate starting edge
-	AddTerminalVerticesForRiverSpline(ERiverBoundaryEdge::Start, Component, Vertices, Indices);
-
-	TArray<double> Distances;
-	TArray<FVector> Points;
-	SplineComp->DivideSplineIntoPolylineRecursiveWithDistances(0.f, SplineComp->GetSplineLength(), ESplineCoordinateSpace::Local, FMath::Square(10.f), Points, Distances);
-	
-	for (double DistanceAlongSpline : Distances)
-	{
-		AddVerticesForRiverSplineStep(DistanceAlongSpline, Component, Vertices, Indices);
-	}
-	
-	// Add an extra point at the end to dilate ending edge
-	AddTerminalVerticesForRiverSpline(ERiverBoundaryEdge::End, Component, Vertices, Indices);
-
-
-	FWaterBodyMeshSection* NewSection = new FWaterBodyMeshSection(GetScene().GetFeatureLevel());
-	NewSection->IndexBuffer.Indices = MoveTemp(Indices);
-	
-	NewSection->VertexBuffers.InitFromDynamicVertex(&NewSection->VertexFactory, Vertices);
-
-	InitResources(NewSection);
-
-	Sections.Add(NewSection);
-}
-
 void FWaterBodySceneProxy::InitResources(FWaterBodyMeshSection* Section)
 {
 	check(Section != nullptr);
@@ -464,27 +39,20 @@ void FWaterBodySceneProxy::InitResources(FWaterBodyMeshSection* Section)
 	BeginInitResource(&Section->VertexFactory);
 }
 
-FWaterBodySceneProxy::FWaterBodySceneProxy(UWaterBodyComponent* Component)
+FWaterBodySceneProxy::FWaterBodySceneProxy(UWaterBodyComponent* Component, const TArray<FDynamicMeshVertex>& Vertices, TArray<uint32>& Indices)
 	: FPrimitiveSceneProxy(Component)
 {
-	if (Component == nullptr || Component->GetWaterSpline() == nullptr)
+	if (Vertices.Num() > 0 && Indices.Num() > 0)
 	{
-		return;
-	}
+		FWaterBodyMeshSection* NewSection = new FWaterBodyMeshSection(GetScene().GetFeatureLevel());
+		NewSection->IndexBuffer.Indices = Indices;
 
-	switch(Component->GetWaterBodyType())
-	{
-		case EWaterBodyType::Lake:
-			GenerateLakeMesh(CastChecked<UWaterBodyLakeComponent>(Component));
-			break;
-		case EWaterBodyType::Ocean:
-			GenerateOceanMesh(CastChecked<UWaterBodyOceanComponent>(Component));
-			break;
-		case EWaterBodyType::River:
-			GenerateRiverMesh(CastChecked<UWaterBodyRiverComponent>(Component));
-			break;
-		default:
-			break;
+		TArray<FDynamicMeshVertex> SectionVertices = Vertices;
+		NewSection->VertexBuffers.InitFromDynamicVertex(&NewSection->VertexFactory, SectionVertices);
+
+		InitResources(NewSection);
+
+		Sections.Add(NewSection);
 	}
 
 	if (UMaterialInstance* WaterInfoMaterial = Component->GetWaterInfoMaterialInstance())
diff --git a/Engine/Plugins/Experimental/Water/Source/Runtime/Private/WaterZoneActor.cpp b/Engine/Plugins/Experimental/Water/Source/Runtime/Private/WaterZoneActor.cpp
--- a/Engine/Plugins/Experimental/Water/Source/Runtime/Private/WaterZoneActor.cpp
+++ b/Engine/Plugins/Experimental/Water/Source/Runtime/Private/WaterZoneActor.cpp
@@ -285,6 +285,6 @@ void AWaterZone::UpdateWaterInfoTexture()
 			Component->UpdateMaterialInstances();
 		}
 
-		UE_LOG(LogWater, Log, TEXT("Queued Water Info texture update"));
+		UE_LOG(LogWater, Verbose, TEXT("Queued Water Info texture update"));
 	}
 }
diff --git a/Engine/Plugins/Experimental/Water/Source/Runtime/Public/WaterBodyComponent.h b/Engine/Plugins/Experimental/Water/Source/Runtime/Public/WaterBodyComponent.h
--- a/Engine/Plugins/Experimental/Water/Source/Runtime/Public/WaterBodyComponent.h
+++ b/Engine/Plugins/Experimental/Water/Source/Runtime/Public/WaterBodyComponent.h
@@ -16,6 +16,7 @@
 #include "Interfaces/Interface_PostProcessVolume.h"
 #include "WaterWaves.h"
 #include "WaterBodyTypes.h"
+#include "DynamicMeshBuilder.h"
 #include "WaterBodyComponent.generated.h"
 
 class UWaterSplineComponent;
@@ -321,6 +322,8 @@ public:
 	/** Gets the water zone to which this component belongs */
 	AWaterZone* GetWaterZone() const;
 
+	void UpdateWaterBodyRenderData();
+
 	virtual FPrimitiveSceneProxy* CreateSceneProxy() override;
 	virtual void GetUsedMaterials(TArray<UMaterialInterface*>& OutMaterials, bool bGetDebugMaterials) const override;
 protected:
@@ -349,6 +352,9 @@ protected:
 
 	virtual void OnUpdateBody(bool bWithExclusionVolumes) {}
 
+	/* Generates the mesh used to render the Water Info texture */
+	virtual void GenerateWaterBodyMesh() {}
+
 	/** Returns navigation area class */
 	TSubclassOf<UNavAreaBase> GetNavAreaClass() const { return WaterNavAreaClass; }
 
@@ -539,6 +545,9 @@ protected:
 	UPROPERTY(EditAnywhere, BlueprintReadWrite, AdvancedDisplay, Category = Rendering)
 	float ShapeDilationZOffset = -64.f;
 
+	TArray<FDynamicMeshVertex> WaterBodyMeshVertices;
+	TArray<uint32> WaterBodyMeshIndices;
+
 #if WITH_EDITORONLY_DATA
 	UPROPERTY()
 	bool bOverrideWaterMesh_DEPRECATED;
diff --git a/Engine/Plugins/Experimental/Water/Source/Runtime/Public/WaterBodyLakeComponent.h b/Engine/Plugins/Experimental/Water/Source/Runtime/Public/WaterBodyLakeComponent.h
--- a/Engine/Plugins/Experimental/Water/Source/Runtime/Public/WaterBodyLakeComponent.h
+++ b/Engine/Plugins/Experimental/Water/Source/Runtime/Public/WaterBodyLakeComponent.h
@@ -23,6 +23,7 @@ public:
 	virtual EWaterBodyType GetWaterBodyType() const override { return EWaterBodyType::Lake; }
 	virtual TArray<UPrimitiveComponent*> GetCollisionComponents() const override;
 	virtual TArray<UPrimitiveComponent*> GetStandardRenderableComponents() const override;
+	virtual void GenerateWaterBodyMesh() override;
 
 	virtual FBoxSphereBounds CalcBounds(const FTransform& LocalToWorld) const override;
 
diff --git a/Engine/Plugins/Experimental/Water/Source/Runtime/Public/WaterBodyOceanComponent.h b/Engine/Plugins/Experimental/Water/Source/Runtime/Public/WaterBodyOceanComponent.h
--- a/Engine/Plugins/Experimental/Water/Source/Runtime/Public/WaterBodyOceanComponent.h
+++ b/Engine/Plugins/Experimental/Water/Source/Runtime/Public/WaterBodyOceanComponent.h
@@ -33,6 +33,7 @@ protected:
 	virtual void BeginUpdateWaterBody() override;
 	virtual void OnUpdateBody(bool bWithExclusionVolumes) override;
 	virtual void Reset() override;
+	virtual void GenerateWaterBodyMesh() override;
 
 	virtual void PostLoad() override;
 	virtual FBoxSphereBounds CalcBounds(const FTransform& LocalToWorld) const;
diff --git a/Engine/Plugins/Experimental/Water/Source/Runtime/Public/WaterBodyRiverComponent.h b/Engine/Plugins/Experimental/Water/Source/Runtime/Public/WaterBodyRiverComponent.h
--- a/Engine/Plugins/Experimental/Water/Source/Runtime/Public/WaterBodyRiverComponent.h
+++ b/Engine/Plugins/Experimental/Water/Source/Runtime/Public/WaterBodyRiverComponent.h
@@ -38,6 +38,7 @@ protected:
 	virtual void Reset() override;
 	virtual void UpdateMaterialInstances() override;
 	virtual void OnUpdateBody(bool bWithExclusionVolumes) override;
+	virtual void GenerateWaterBodyMesh() override;
 
 	virtual FBoxSphereBounds CalcBounds(const FTransform& LocalToWorld) const override;
 #if WITH_EDITOR
diff --git a/Engine/Plugins/Experimental/Water/Source/Runtime/Public/WaterBodySceneProxy.h b/Engine/Plugins/Experimental/Water/Source/Runtime/Public/WaterBodySceneProxy.h
--- a/Engine/Plugins/Experimental/Water/Source/Runtime/Public/WaterBodySceneProxy.h
+++ b/Engine/Plugins/Experimental/Water/Source/Runtime/Public/WaterBodySceneProxy.h
@@ -18,7 +18,7 @@ class FMaterialRenderProxy;
 class FWaterBodySceneProxy final : public FPrimitiveSceneProxy
 {
 public:
-	FWaterBodySceneProxy(UWaterBodyComponent* Component);
+	FWaterBodySceneProxy(UWaterBodyComponent* Component, const TArray<FDynamicMeshVertex>& Vertices, TArray<uint32>& Indices);
 	virtual ~FWaterBodySceneProxy();
 	virtual void GetDynamicMeshElements(const TArray<const FSceneView*>& Views, const FSceneViewFamily& ViewFamily, uint32 VisibilityMap, FMeshElementCollector& Collector) const override;
 	virtual FPrimitiveViewRelevance GetViewRelevance(const FSceneView* View) const override;
@@ -32,12 +32,6 @@ public:
 	void SetWithinWaterInfoPass(bool bInWithinWaterInfoPass);
 
 private:
-	void GenerateLakeMesh(UWaterBodyLakeComponent* Component);
-	void GenerateOceanMesh(UWaterBodyOceanComponent* Component);
-	void GenerateRiverMesh(UWaterBodyRiverComponent* Component);
-
-	void AddSectionFromDynamicMesh(const UE::Geometry::FDynamicMesh3& DynamicMesh, TFunctionRef<FDynamicMeshVertex(const FVector&)> VertTransformFunc);
-
 	void InitResources(FWaterBodyMeshSection* Section);
 
 	TArray<FWaterBodyMeshSection*> Sections;
