Fix incorrect river visual bounds calculation

diff --git a/Engine/Plugins/Experimental/Water/Source/Runtime/Private/WaterBodyRiverComponent.cpp b/Engine/Plugins/Experimental/Water/Source/Runtime/Private/WaterBodyRiverComponent.cpp
--- a/Engine/Plugins/Experimental/Water/Source/Runtime/Private/WaterBodyRiverComponent.cpp
+++ b/Engine/Plugins/Experimental/Water/Source/Runtime/Private/WaterBodyRiverComponent.cpp
@@ -321,14 +321,14 @@ FBoxSphereBounds UWaterBodyRiverComponent::CalcBounds(const FTransform& LocalToW
 	{
 		if (SplineMeshComponent != nullptr)
 		{
-			BoxExtent += SplineMeshComponent->CalcLocalBounds().BoxExtent;
+			const FBox SplineMeshComponentBounds = SplineMeshComponent->CalcBounds(SplineMeshComponent->GetRelativeTransform()).GetBox();
+			BoxExtent += SplineMeshComponentBounds;
 		}
 	}
 	// Spline mesh components aren't storing our vertical bounds so account for that with the ChannelDepth parameter.
 	BoxExtent.Max.Z += MaxWaveHeightOffset;
 	BoxExtent.Min.Z -= GetChannelDepth();
-	FBoxSphereBounds Result(BoxExtent);
-	return Result.TransformBy(LocalToWorld);
+	return FBoxSphereBounds(BoxExtent).TransformBy(LocalToWorld);
 }
 
 void UWaterBodyRiverComponent::UpdateMaterialInstances()
