Fixed implicit conversions from FVector to FVector3f in water body mesh generation

diff --git a/Engine/Plugins/Experimental/Water/Source/Runtime/Private/WaterBodyLakeComponent.cpp b/Engine/Plugins/Experimental/Water/Source/Runtime/Private/WaterBodyLakeComponent.cpp
--- a/Engine/Plugins/Experimental/Water/Source/Runtime/Private/WaterBodyLakeComponent.cpp
+++ b/Engine/Plugins/Experimental/Water/Source/Runtime/Private/WaterBodyLakeComponent.cpp
@@ -91,7 +91,7 @@ void UWaterBodyLakeComponent::GenerateWaterBodyMesh()
 	for (const FVector2d& Vertex : Triangulation.Vertices)
 	{
 		// push the set of undilated vertices to the persistent mesh
-		FDynamicMeshVertex MeshVertex(FVector(Vertex.X, Vertex.Y, 0.f));
+		FDynamicMeshVertex MeshVertex(FVector3f(Vertex.X, Vertex.Y, 0.f));
 		MeshVertex.Color = FColor::Black;
 		MeshVertex.TextureCoordinate[0].X = WaterBodyIndex;
 		WaterBodyMeshVertices.Add(MeshVertex);
@@ -124,7 +124,7 @@ void UWaterBodyLakeComponent::GenerateWaterBodyMesh()
 			for (const FVector3d& Vertex : LakeMesh.GetVerticesBuffer())
 			{
 				// push the set of dilated vertices to the persistent mesh
-				FDynamicMeshVertex MeshVertex(Vertex);
+				FDynamicMeshVertex MeshVertex(FVector3f(Vertex.X, Vertex.Y, 0.f));
 				MeshVertex.Position.Z = ShapeDilationZOffset;
 				MeshVertex.Color = FColor::Black;
 				MeshVertex.TextureCoordinate[0].X = -1;
diff --git a/Engine/Plugins/Experimental/Water/Source/Runtime/Private/WaterBodyOceanComponent.cpp b/Engine/Plugins/Experimental/Water/Source/Runtime/Private/WaterBodyOceanComponent.cpp
--- a/Engine/Plugins/Experimental/Water/Source/Runtime/Private/WaterBodyOceanComponent.cpp
+++ b/Engine/Plugins/Experimental/Water/Source/Runtime/Private/WaterBodyOceanComponent.cpp
@@ -173,7 +173,7 @@ void UWaterBodyOceanComponent::GenerateWaterBodyMesh()
 	for (const FVector2d& Vertex : Triangulation.Vertices)
 	{
 		// push the set of undilated vertices to the persistent mesh
-		FDynamicMeshVertex MeshVertex(FVector(Vertex.X, Vertex.Y, 0.f));
+		FDynamicMeshVertex MeshVertex(FVector3f(Vertex.X, Vertex.Y, 0.f));
 		MeshVertex.Color = FColor::Black;
 		MeshVertex.TextureCoordinate[0].X = WaterBodyIndex;
 		WaterBodyMeshVertices.Add(MeshVertex);
@@ -209,7 +209,7 @@ void UWaterBodyOceanComponent::GenerateWaterBodyMesh()
 			for (const FVector3d& Vertex : OceanMesh.GetVerticesBuffer())
 			{
 				// push the set of dilated vertices to the persistent mesh
-				FDynamicMeshVertex MeshVertex(Vertex);
+				FDynamicMeshVertex MeshVertex(FVector3f(Vertex.X, Vertex.Y, 0.f));
 				MeshVertex.Position.Z = ShapeDilationZOffset;
 				MeshVertex.Color = FColor::Black;
 				MeshVertex.TextureCoordinate[0].X = -1;
diff --git a/Engine/Plugins/Experimental/Water/Source/Runtime/Private/WaterBodyRiverComponent.cpp b/Engine/Plugins/Experimental/Water/Source/Runtime/Private/WaterBodyRiverComponent.cpp
--- a/Engine/Plugins/Experimental/Water/Source/Runtime/Private/WaterBodyRiverComponent.cpp
+++ b/Engine/Plugins/Experimental/Water/Source/Runtime/Private/WaterBodyRiverComponent.cpp
@@ -129,16 +129,16 @@ static void AddVerticesForRiverSplineStep(float DistanceAlongSpline, const UWate
 	const float DilationAmount = Component->ShapeDilation;
 	const FVector DilationOffset = Normal * DilationAmount;
 
-	FDynamicMeshVertex Left(Pos - OutwardDistance);
-	FDynamicMeshVertex Right(Pos + OutwardDistance);
+	FDynamicMeshVertex Left(FVector3f(Pos - OutwardDistance));
+	FDynamicMeshVertex Right(FVector3f(Pos + OutwardDistance));
 
-	FDynamicMeshVertex DilatedFarLeft(Pos - OutwardDistance - DilationOffset);
+	FDynamicMeshVertex DilatedFarLeft(FVector3f(Pos - OutwardDistance - DilationOffset));
 	DilatedFarLeft.Position.Z += Component->GetShapeDilationZOffsetFar();
-	FDynamicMeshVertex DilatedLeft(Pos - OutwardDistance);
+	FDynamicMeshVertex DilatedLeft(FVector3f(Pos - OutwardDistance));
 	DilatedLeft.Position.Z += Component->GetShapeDilationZOffset();
-	FDynamicMeshVertex DilatedRight(Pos + OutwardDistance);
+	FDynamicMeshVertex DilatedRight(FVector3f(Pos + OutwardDistance));
 	DilatedRight.Position.Z += Component->GetShapeDilationZOffset();
-	FDynamicMeshVertex DilatedFarRight(Pos + OutwardDistance + DilationOffset);
+	FDynamicMeshVertex DilatedFarRight(FVector3f(Pos + OutwardDistance + DilationOffset));
 	DilatedFarRight.Position.Z += Component->GetShapeDilationZOffsetFar();
 
 	float FlowDirection = Tangent.HeadingAngle() + FMath::DegreesToRadians(Component->GetRelativeRotation().Yaw);
@@ -225,13 +225,13 @@ static void AddTerminalVerticesForRiverSpline(ERiverBoundaryEdge Edge, const UWa
 		TangentialOffset *= -1;
 	}
 
-	FDynamicMeshVertex BackLeft(Pos - OutwardDistance + TangentialOffset - DilationOffset);
+	FDynamicMeshVertex BackLeft(FVector3f(Pos - OutwardDistance + TangentialOffset - DilationOffset));
 	BackLeft.Position.Z += Component->GetShapeDilationZOffsetFar();
-	FDynamicMeshVertex Left(Pos - OutwardDistance + TangentialOffset);
+	FDynamicMeshVertex Left(FVector3f(Pos - OutwardDistance + TangentialOffset));
 	Left.Position.Z += Component->GetShapeDilationZOffsetFar();
-	FDynamicMeshVertex Right(Pos + OutwardDistance + TangentialOffset);
+	FDynamicMeshVertex Right(FVector3f(Pos + OutwardDistance + TangentialOffset));
 	Right.Position.Z += Component->GetShapeDilationZOffsetFar();
-	FDynamicMeshVertex BackRight(Pos + OutwardDistance + TangentialOffset + DilationOffset);
+	FDynamicMeshVertex BackRight(FVector3f(Pos + OutwardDistance + TangentialOffset + DilationOffset));
 	BackRight.Position.Z += Component->GetShapeDilationZOffsetFar();
 
 	// Initialize the vertex data to correct represent what we expect the dilated region to look like
