Missing explicit conversion in water body mesh generation.

diff --git a/Engine/Plugins/Experimental/Water/Source/Runtime/Private/WaterBodyLakeComponent.cpp b/Engine/Plugins/Experimental/Water/Source/Runtime/Private/WaterBodyLakeComponent.cpp
--- a/Engine/Plugins/Experimental/Water/Source/Runtime/Private/WaterBodyLakeComponent.cpp
+++ b/Engine/Plugins/Experimental/Water/Source/Runtime/Private/WaterBodyLakeComponent.cpp
@@ -96,7 +96,7 @@ void UWaterBodyLakeComponent::GenerateWaterBodyMesh()
 		MeshVertex.TextureCoordinate[0].X = WaterBodyIndex;
 		WaterBodyMeshVertices.Add(MeshVertex);
 
-		LakeMesh.AppendVertex(MeshVertex.Position);
+		LakeMesh.AppendVertex(FVector3d(Vertex, 0.0));
 	}
 
 	for (const FIndex3i& Triangle : Triangulation.Triangles)
diff --git a/Engine/Plugins/Experimental/Water/Source/Runtime/Private/WaterBodyOceanComponent.cpp b/Engine/Plugins/Experimental/Water/Source/Runtime/Private/WaterBodyOceanComponent.cpp
--- a/Engine/Plugins/Experimental/Water/Source/Runtime/Private/WaterBodyOceanComponent.cpp
+++ b/Engine/Plugins/Experimental/Water/Source/Runtime/Private/WaterBodyOceanComponent.cpp
@@ -178,7 +178,7 @@ void UWaterBodyOceanComponent::GenerateWaterBodyMesh()
 		MeshVertex.TextureCoordinate[0].X = WaterBodyIndex;
 		WaterBodyMeshVertices.Add(MeshVertex);
 
-		OceanMesh.AppendVertex(MeshVertex.Position);
+		OceanMesh.AppendVertex(FVector3d(Vertex, 0.0));
 	}
 
 	for (const FIndex3i& Triangle : Triangulation.Triangles)
