The Water Info texture should be re-renderer whenever the landscape heightmap changes

diff --git a/Engine/Plugins/Experimental/Water/Source/Editor/Private/WaterLandscapeBrush.cpp b/Engine/Plugins/Experimental/Water/Source/Editor/Private/WaterLandscapeBrush.cpp
--- a/Engine/Plugins/Experimental/Water/Source/Editor/Private/WaterLandscapeBrush.cpp
+++ b/Engine/Plugins/Experimental/Water/Source/Editor/Private/WaterLandscapeBrush.cpp
@@ -490,6 +490,11 @@ void AWaterLandscapeBrush::SetTargetLandscape(ALandscape* InTargetLandscape)
 void AWaterLandscapeBrush::OnFullHeightmapRenderDone(UTextureRenderTarget2D* InHeightmapRenderTarget)
 {
 	bRenderTargetsDirty = false;
+	// #todo_water [roey]: This needs to be changed when the WaterZone can maintain it's own list of "ground actors" so that we don't needlessly update all water zones.
+	if (UWaterSubsystem* WaterSubsystem = UWaterSubsystem::GetWaterSubsystem(GetWorld()))
+	{
+		WaterSubsystem->MarkAllWaterZonesForRebuild(EWaterZoneRebuildFlags::UpdateWaterInfoTexture);
+	}
 }
 
 void AWaterLandscapeBrush::SetOwningLandscape(ALandscape* InOwningLandscape)
diff --git a/Engine/Plugins/Experimental/Water/Source/Runtime/Private/WaterSubsystem.cpp b/Engine/Plugins/Experimental/Water/Source/Runtime/Private/WaterSubsystem.cpp
--- a/Engine/Plugins/Experimental/Water/Source/Runtime/Private/WaterSubsystem.cpp
+++ b/Engine/Plugins/Experimental/Water/Source/Runtime/Private/WaterSubsystem.cpp
@@ -5,7 +5,6 @@
 #include "HAL/IConsoleManager.h"
 #include "Engine/World.h"
 #include "TimerManager.h"
-#include "WaterZoneActor.h"
 #include "WaterMeshComponent.h"
 #include "Engine/Engine.h"
 #include "EngineUtils.h"
@@ -515,13 +514,13 @@ float UWaterSubsystem::GetOceanBaseHeight() const
 	return TNumericLimits<float>::Lowest();
 }
 
-void UWaterSubsystem::MarkAllWaterZonesForRebuild()
+void UWaterSubsystem::MarkAllWaterZonesForRebuild(EWaterZoneRebuildFlags RebuildFlags)
 {
 	if (UWorld* World = GetWorld())
 	{
 		for (AWaterZone* WaterZone : TActorRange<AWaterZone>(World))
 		{
-			WaterZone->MarkForRebuild(EWaterZoneRebuildFlags::All);
+			WaterZone->MarkForRebuild(RebuildFlags);
 		}
 	}
 }
diff --git a/Engine/Plugins/Experimental/Water/Source/Runtime/Public/WaterSubsystem.h b/Engine/Plugins/Experimental/Water/Source/Runtime/Public/WaterSubsystem.h
--- a/Engine/Plugins/Experimental/Water/Source/Runtime/Public/WaterSubsystem.h
+++ b/Engine/Plugins/Experimental/Water/Source/Runtime/Public/WaterSubsystem.h
@@ -10,6 +10,7 @@
 #include "Engine/EngineTypes.h"
 #include "Interfaces/Interface_PostProcessVolume.h"
 #include "WaterBodyManager.h"
+#include "WaterZoneActor.h"
 #include "WaterSubsystem.generated.h"
 
 DECLARE_STATS_GROUP(TEXT("Water"), STATGROUP_Water, STATCAT_Advanced);
@@ -17,7 +18,6 @@ DECLARE_STATS_GROUP(TEXT("Water"), STATGROUP_Water, STATCAT_Advanced);
 DECLARE_DYNAMIC_MULTICAST_DELEGATE_TwoParams(FOnCameraUnderwaterStateChanged, bool, bIsUnderWater, float, DepthUnderwater);
 DECLARE_DYNAMIC_MULTICAST_DELEGATE(FOnWaterScalabilityChanged);
 
-class AWaterZone;
 class UWaterBodyComponent;
 class UMaterialParameterCollection;
 class UWaterRuntimeSettings;
@@ -165,7 +165,7 @@ public:
 
 	UMaterialParameterCollection* GetMaterialParameterCollection() const {	return MaterialParameterCollection; }
 	
-	void MarkAllWaterZonesForRebuild();
+	void MarkAllWaterZonesForRebuild(EWaterZoneRebuildFlags RebuildFlags = EWaterZoneRebuildFlags::All);
 
 #if WITH_EDITOR
 	/** Little scope object to temporarily change the value of bAllowWaterSubsystemOnPreviewWorld */
