Fix for incorrect angle unwind logic causing cooking crash
#rnx

diff --git a/Engine/Plugins/Experimental/Water/Source/Runtime/Private/WaterBodyRiverComponent.cpp b/Engine/Plugins/Experimental/Water/Source/Runtime/Private/WaterBodyRiverComponent.cpp
--- a/Engine/Plugins/Experimental/Water/Source/Runtime/Private/WaterBodyRiverComponent.cpp
+++ b/Engine/Plugins/Experimental/Water/Source/Runtime/Private/WaterBodyRiverComponent.cpp
@@ -143,7 +143,7 @@ static void AddVerticesForRiverSplineStep(float DistanceAlongSpline, const UWate
 
 	float FlowDirection = Tangent.HeadingAngle() + FMath::DegreesToRadians(Component->GetRelativeRotation().Yaw);
 	// Convert negative angles into positive angles
-	if (FlowDirection < 0.f)
+	while (FlowDirection < 0.f)
 	{
 		FlowDirection = TWO_PI + FlowDirection;
 	}
