Fixed missing Far Mesh due to its position not being sent to the shader in Translated World Position

#rnx

diff --git a/Engine/Plugins/Experimental/Water/Source/Runtime/Private/WaterQuadTree.cpp b/Engine/Plugins/Experimental/Water/Source/Runtime/Private/WaterQuadTree.cpp
--- a/Engine/Plugins/Experimental/Water/Source/Runtime/Private/WaterQuadTree.cpp
+++ b/Engine/Plugins/Experimental/Water/Source/Runtime/Private/WaterQuadTree.cpp
@@ -567,12 +567,8 @@ void FWaterQuadTree::AddFarMesh(const UMaterialInterface* InFarMeshMaterial, dou
 		TileScale.X = (TileOffets[i].X == 0.0) ? WaterSize.X : InFarDistanceMeshExtent;
 		TileScale.Y = (TileOffets[i].Y == 0.0) ? WaterSize.Y : InFarDistanceMeshExtent;
 
-		// Build instance data
-		FarMeshData.InstanceData[i].Data[0] = FVector4f(FVector4(TilePos, FVector2D(InFarDistanceMeshHeight, 0.0))); // LWC_TODO: precision loss
-		FarMeshData.InstanceData[i].Data[1] = FVector4f(FVector4(FVector2D::ZeroVector, TileScale)); // LWC_TODO: precision loss
-#if WITH_WATER_SELECTION_SUPPORT
-		FarMeshData.InstanceData[i].Data[2] = FVector4f(FHitProxyId::InvisibleHitProxyId.GetColor().ReinterpretAsLinear());
-#endif // WITH_WATER_SELECTION_SUPPORT
+		FarMeshData.InstanceData[i].WorldPosition = FVector(TilePos, InFarDistanceMeshHeight);
+		FarMeshData.InstanceData[i].Scale = FVector2f(TileScale);
 	}
 }
 
@@ -638,7 +634,15 @@ void FWaterQuadTree::BuildWaterTileInstanceData(const FTraversalDesc& InTraversa
 
 		for (int32 i = 0; i < FarMeshTileCount; i++)
 		{
-			Output.StagingInstanceData[StartIndex + i] = FarMeshData.InstanceData[i];
+			// Build instance data
+			// Transform worldposition to Translated World Position
+			const FVector TranslatedWorldPosition(FarMeshData.InstanceData[i].WorldPosition + InTraversalDesc.PreViewTranslation);
+			Output.StagingInstanceData[StartIndex + i].Data[0] = FVector4f(FVector4(TranslatedWorldPosition, 0.0));
+			Output.StagingInstanceData[StartIndex + i].Data[1] = FVector4f(FVector2f::ZeroVector, FarMeshData.InstanceData[i].Scale);
+#if WITH_WATER_SELECTION_SUPPORT
+			Output.StagingInstanceData[StartIndex + i].Data[2] = FHitProxyId::InvisibleHitProxyId.GetColor().ReinterpretAsLinear();
+#endif // WITH_WATER_SELECTION_SUPPORT
+
 			Output.StagingInstanceData[StartIndex + i].BucketIndex = BucketIndex;
 		}
 	}
diff --git a/Engine/Plugins/Experimental/Water/Source/Runtime/Public/WaterQuadTree.h b/Engine/Plugins/Experimental/Water/Source/Runtime/Public/WaterQuadTree.h
--- a/Engine/Plugins/Experimental/Water/Source/Runtime/Public/WaterQuadTree.h
+++ b/Engine/Plugins/Experimental/Water/Source/Runtime/Public/WaterQuadTree.h
@@ -275,8 +275,14 @@ private:
 	/** Contains everything needed to render the far mesh. This data lives outside the quadtree structure itself */
 	struct FFarMeshData
 	{
-		/** Stored data for rendering all far mesh instances (8 or them if used). This is built oncce when the quadtree is built and then read each time the quadtree is traversed */
-		TArray<FStagingInstanceData> InstanceData;
+		struct FFarMeshInstanceData
+		{
+			FVector WorldPosition = FVector::ZeroVector;
+			FVector2f Scale = FVector2f(1.0f, 1.0f);
+		};
+
+		/** Stored data for rendering all far mesh instances (8 or them if used). This is built once when the quadtree is built and then read and transformed each time the quadtree is traversed */
+		TArray<FFarMeshInstanceData> InstanceData;
 
 		/** Material for the Far Distance Mesh, its material render proxy will be cached in WaterMaterials when BuildMaterialIndices is called */
 		const UMaterialInterface* Material = nullptr;
