Fix issue where all water bodies were drawing selection outlines when any other water body was selected.

#rnx

diff --git a/Engine/Plugins/Experimental/Water/Shaders/Private/WaterMeshVertexFactory.ush b/Engine/Plugins/Experimental/Water/Shaders/Private/WaterMeshVertexFactory.ush
--- a/Engine/Plugins/Experimental/Water/Shaders/Private/WaterMeshVertexFactory.ush
+++ b/Engine/Plugins/Experimental/Water/Shaders/Private/WaterMeshVertexFactory.ush
@@ -250,14 +250,14 @@ FVertexFactoryIntermediates GetVertexFactoryIntermediates(FVertexFactoryInput In
 		Intermediates.MorphedTranslatedWorldPos = TranslatedWorldPosition;
 	}
 
-	Intermediates.MorphedWorldPos = LWCSubtract(Intermediates.MorphedTranslatedWorldPos, ResolvedView.PreViewTranslation);
-
 #if WITH_WATER_SELECTION_SUPPORT_VF
 	float SelectedValue = Input.InstanceData2.w;
 	float IsVisible = WaterVF.bRenderSelected * SelectedValue + WaterVF.bRenderUnselected * (1-SelectedValue);
-	Intermediates.MorphedWorldPos = LWCMultiply(Intermediates.MorphedWorldPos, IsVisible);
+	Intermediates.MorphedTranslatedWorldPos *= IsVisible;
 #endif // WITH_WATER_SELECTION_SUPPORT_VF
 
+	Intermediates.MorphedWorldPos = LWCSubtract(Intermediates.MorphedTranslatedWorldPos, ResolvedView.PreViewTranslation);
+
 	return Intermediates;
 }
 
