Fix issues rendering landscape with the default material.

- Add JumpStepMaterial, FindEdgesMaterial and BlurEdgesMaterial to AWaterBrushManager's render dependencies.  This way we compile a complete shader map for these materials before we render.
- Add TRACE_CPUPROFILER_EVENT_SCOPE to AWaterBrushManager::Render_Native to make it easier to see when the rendering actually happens in Unreal Insights.

diff --git a/Engine/Plugins/Experimental/Water/Source/Editor/Private/WaterBrushManager.cpp b/Engine/Plugins/Experimental/Water/Source/Editor/Private/WaterBrushManager.cpp
--- a/Engine/Plugins/Experimental/Water/Source/Editor/Private/WaterBrushManager.cpp
+++ b/Engine/Plugins/Experimental/Water/Source/Editor/Private/WaterBrushManager.cpp
@@ -489,6 +489,9 @@ void AWaterBrushManager::GetRenderDependencies(TSet<UObject*>& OutDependencies)
 	AddDependencyIfValid(CompositeWaterBodyTextureMaterial, OutDependencies);
 	AddDependencyIfValid(IslandFalloffMaterial, OutDependencies);
 	AddDependencyIfValid(FinalizeVelocityHeightMaterial, OutDependencies);
+	AddDependencyIfValid(JumpStepMaterial, OutDependencies);
+	AddDependencyIfValid(FindEdgesMaterial, OutDependencies);
+	AddDependencyIfValid(BlurEdgesMaterial, OutDependencies);
 }
 
 void AWaterBrushManager::UpdateTransform(const FTransform& Transform)
@@ -1380,6 +1383,7 @@ void AWaterBrushManager::SetupDefaultMaterials()
 UTextureRenderTarget2D* AWaterBrushManager::Render_Native(bool InIsHeightmap, UTextureRenderTarget2D* InCombinedResult, FName const& InWeightmapLayerName)
 {
 	QUICK_SCOPE_CYCLE_COUNTER(AWaterBrushManager_Render_Native);
+	TRACE_CPUPROFILER_EVENT_SCOPE(AWaterBrushManager::Render_Native);
 
 	LandscapeRTRef = InCombinedResult;
 
