Fixed crash when "saving as..." a map with a river water body : the water body component doesn't have a valid water spline metadata until post register

diff --git a/Engine/Plugins/Experimental/Water/Source/Runtime/Private/WaterBodyRiverComponent.cpp b/Engine/Plugins/Experimental/Water/Source/Runtime/Private/WaterBodyRiverComponent.cpp
--- a/Engine/Plugins/Experimental/Water/Source/Runtime/Private/WaterBodyRiverComponent.cpp
+++ b/Engine/Plugins/Experimental/Water/Source/Runtime/Private/WaterBodyRiverComponent.cpp
@@ -107,9 +107,9 @@ static FColor PackWaterFlow(float VelocityMagnitude, float DirectionAngle)
 	return Result;
 }
 
-static void AddVerticesForRiverSplineStep(float DistanceAlongSpline, const UWaterBodyRiverComponent* Component, TArray<FDynamicMeshVertex>& Vertices, TArray<uint32>& Indices)
+static void AddVerticesForRiverSplineStep(float DistanceAlongSpline, const UWaterBodyRiverComponent* Component, const UWaterSplineComponent* SplineComp, const UWaterSplineMetadata* WaterSplineMetadata, TArray<FDynamicMeshVertex>& Vertices, TArray<uint32>& Indices)
 {
-	const UWaterSplineComponent* SplineComp = Component->GetWaterSpline();
+	check((Component != nullptr) && (SplineComp != nullptr) && (WaterSplineMetadata != nullptr));
 	
 	const FVector Tangent = SplineComp->GetTangentAtDistanceAlongSpline(DistanceAlongSpline, ESplineCoordinateSpace::Local).GetSafeNormal();
 	const FVector Up = SplineComp->GetUpVectorAtDistanceAlongSpline(DistanceAlongSpline, ESplineCoordinateSpace::Local).GetSafeNormal();
@@ -118,8 +118,8 @@ static void AddVerticesForRiverSplineStep(float DistanceAlongSpline, const UWate
 	const FVector Pos = SplineComp->GetLocationAtDistanceAlongSpline(DistanceAlongSpline, ESplineCoordinateSpace::Local);
 
 	const float Key = SplineComp->SplineCurves.ReparamTable.Eval(DistanceAlongSpline, 0.f);
-	const float HalfWidth = Cast<UWaterSplineMetadata>(SplineComp->GetSplinePointsMetadata())->RiverWidth.Eval(Key) / 2;
-	float Velocity = Cast<UWaterSplineMetadata>(SplineComp->GetSplinePointsMetadata())->WaterVelocityScalar.Eval(Key);
+	const float HalfWidth = WaterSplineMetadata->RiverWidth.Eval(Key) / 2;
+	float Velocity = WaterSplineMetadata->WaterVelocityScalar.Eval(Key);
 
 	// Distance from the center of the spline to place our first vertices
 	FVector OutwardDistance = Normal * HalfWidth;
@@ -196,9 +196,9 @@ enum class ERiverBoundaryEdge {
 	End,
 };
 
-static void AddTerminalVerticesForRiverSpline(ERiverBoundaryEdge Edge, const UWaterBodyRiverComponent* Component, TArray<FDynamicMeshVertex>& Vertices, TArray<uint32>& Indices)
+static void AddTerminalVerticesForRiverSpline(ERiverBoundaryEdge Edge, const UWaterBodyRiverComponent* Component, const UWaterSplineComponent* SplineComp, const UWaterSplineMetadata* WaterSplineMetadata, TArray<FDynamicMeshVertex>& Vertices, TArray<uint32>& Indices)
 {
-	const UWaterSplineComponent* SplineComp = Component->GetWaterSpline();
+	check((Component != nullptr) && (SplineComp != nullptr) && (WaterSplineMetadata != nullptr));
 
 	const float DistanceAlongSpline = (Edge == ERiverBoundaryEdge::Start) ? 0.f : SplineComp->GetSplineLength();
 	
@@ -209,7 +209,7 @@ static void AddTerminalVerticesForRiverSpline(ERiverBoundaryEdge Edge, const UWa
 	const FVector Pos = SplineComp->GetLocationAtDistanceAlongSpline(DistanceAlongSpline, ESplineCoordinateSpace::Local);
 	
 	const float Key = SplineComp->SplineCurves.ReparamTable.Eval(DistanceAlongSpline, 0.f);
-	const float HalfWidth = Cast<UWaterSplineMetadata>(SplineComp->GetSplinePointsMetadata())->RiverWidth.Eval(Key) / 2;
+	const float HalfWidth = WaterSplineMetadata->RiverWidth.Eval(Key) / 2;
 
 	const float DilationAmount = Component->ShapeDilation;
 	const FVector DilationOffset = Normal * DilationAmount;
@@ -287,7 +287,7 @@ void UWaterBodyRiverComponent::GenerateWaterBodyMesh()
 	WaterBodyMeshIndices.Empty();
 
 	const UWaterSplineComponent* SplineComp = GetWaterSpline();
-	if (SplineComp == nullptr)
+	if ((SplineComp == nullptr) || (WaterSplineMetadata == nullptr))
 	{
 		return;
 	}
@@ -296,7 +296,7 @@ void UWaterBodyRiverComponent::GenerateWaterBodyMesh()
 	TArray<uint32> Indices;
 
 	// Add an extra point at the start to dilate starting edge
-	AddTerminalVerticesForRiverSpline(ERiverBoundaryEdge::Start, this, Vertices, Indices);
+	AddTerminalVerticesForRiverSpline(ERiverBoundaryEdge::Start, this, SplineComp, WaterSplineMetadata, Vertices, Indices);
 
 	TArray<double> Distances;
 	TArray<FVector> Points;
@@ -304,11 +304,11 @@ void UWaterBodyRiverComponent::GenerateWaterBodyMesh()
 	
 	for (double DistanceAlongSpline : Distances)
 	{
-		AddVerticesForRiverSplineStep(DistanceAlongSpline, this, Vertices, Indices);
+		AddVerticesForRiverSplineStep(DistanceAlongSpline, this, SplineComp, WaterSplineMetadata, Vertices, Indices);
 	}
 	
 	// Add an extra point at the end to dilate ending edge
-	AddTerminalVerticesForRiverSpline(ERiverBoundaryEdge::End, this, Vertices, Indices);
+	AddTerminalVerticesForRiverSpline(ERiverBoundaryEdge::End, this, SplineComp, WaterSplineMetadata, Vertices, Indices);
 
 	WaterBodyMeshVertices = MoveTemp(Vertices);
 	WaterBodyMeshIndices = MoveTemp(Indices);
