Added early outs for river mesh generation if the river contains fewer than 2 spline points or no distances are returned by DivideSplineIntoPolyLine

diff --git a/Engine/Plugins/Experimental/Water/Source/Runtime/Private/WaterBodyRiverComponent.cpp b/Engine/Plugins/Experimental/Water/Source/Runtime/Private/WaterBodyRiverComponent.cpp
--- a/Engine/Plugins/Experimental/Water/Source/Runtime/Private/WaterBodyRiverComponent.cpp
+++ b/Engine/Plugins/Experimental/Water/Source/Runtime/Private/WaterBodyRiverComponent.cpp
@@ -287,21 +287,25 @@ void UWaterBodyRiverComponent::GenerateWaterBodyMesh()
 	WaterBodyMeshIndices.Empty();
 
 	const UWaterSplineComponent* SplineComp = GetWaterSpline();
-	if ((SplineComp == nullptr) || (WaterSplineMetadata == nullptr))
+	if ((SplineComp == nullptr) || (WaterSplineMetadata == nullptr) || (SplineComp->GetNumberOfSplinePoints() < 2))
 	{
 		return;
 	}
 
+	TArray<double> Distances;
+	TArray<FVector> Points;
+	SplineComp->DivideSplineIntoPolylineRecursiveWithDistances(0.f, SplineComp->GetSplineLength(), ESplineCoordinateSpace::Local, FMath::Square(10.f), Points, Distances);
+	if (Distances.Num() == 0)
+	{
+		return;
+	}
+	
 	TArray<FDynamicMeshVertex> Vertices;
 	TArray<uint32> Indices;
 
 	// Add an extra point at the start to dilate starting edge
 	AddTerminalVerticesForRiverSpline(ERiverBoundaryEdge::Start, this, SplineComp, WaterSplineMetadata, Vertices, Indices);
 
-	TArray<double> Distances;
-	TArray<FVector> Points;
-	SplineComp->DivideSplineIntoPolylineRecursiveWithDistances(0.f, SplineComp->GetSplineLength(), ESplineCoordinateSpace::Local, FMath::Square(10.f), Points, Distances);
-	
 	for (double DistanceAlongSpline : Distances)
 	{
 		AddVerticesForRiverSplineStep(DistanceAlongSpline, this, SplineComp, WaterSplineMetadata, Vertices, Indices);
