Moved the responsibility of rendering the water info texture to a new WaterViewExtension.
The View extension can submit the rendering commands at the beginning of the next frame after all scene proxies were created.
This solves the issue where if a water body scene proxy was invalidated, it would enqueue the render update on the same frame, before the proxy was regenerated.

diff --git a/Engine/Plugins/Experimental/Water/Source/Runtime/Private/WaterInfoRendering.cpp b/Engine/Plugins/Experimental/Water/Source/Runtime/Private/WaterInfoRendering.cpp
--- a/Engine/Plugins/Experimental/Water/Source/Runtime/Private/WaterInfoRendering.cpp
+++ b/Engine/Plugins/Experimental/Water/Source/Runtime/Private/WaterInfoRendering.cpp
@@ -18,7 +18,13 @@
 #include "Rendering/NaniteStreamingManager.h"
 #include "SceneViewExtension.h"
 
-namespace WaterInfo
+static int32 RenderCaptureNextWaterInfoDraws = 0;
+static FAutoConsoleVariableRef CVarRenderCaptureNextWaterInfoDraws(
+	TEXT("r.Water.WaterInfo.RenderCaptureNextWaterInfoDraws"),
+	RenderCaptureNextWaterInfoDraws,
+	TEXT("Enable capturing of the water info texture for the next N draws"));
+
+namespace UE::WaterInfo
 {
 
 struct FUpdateWaterInfoParams
@@ -535,10 +541,12 @@ static FSceneRenderer* CreateWaterInfoColorRenderer(
 
 void UpdateWaterInfoRendering(
 	FSceneInterface* Scene,
-	UTextureRenderTarget2D* TextureRenderTarget,
 	const WaterInfo::FRenderingContext& Context)
 {
-	if (TextureRenderTarget == nullptr || Scene == nullptr)
+	RenderCaptureInterface::FScopedCapture RenderCapture((RenderCaptureNextWaterInfoDraws != 0), TEXT("RenderWaterInfo"));
+	RenderCaptureNextWaterInfoDraws = FMath::Max(0, RenderCaptureNextWaterInfoDraws - 1);
+
+	if (Context.TextureRenderTarget == nullptr || Scene == nullptr)
 	{
 		return;
 	}
@@ -554,13 +562,13 @@ void UpdateWaterInfoRendering(
 	ViewRotationMat = ViewRotationMat.RemoveTranslation();
 	ViewRotationMat.RemoveScaling();
 	
-	const FIntPoint CaptureExtent(TextureRenderTarget->GetSurfaceWidth(), TextureRenderTarget->GetSurfaceHeight());
+	const FIntPoint CaptureExtent(Context.TextureRenderTarget->GetSurfaceWidth(), Context.TextureRenderTarget->GetSurfaceHeight());
 
 	const FMatrix OrthoProj = BuildOrthoMatrix(ZoneExtent.X, ZoneExtent.Y);
 
 	FSceneRenderer* DepthRenderer = CreateWaterInfoDepthRenderer(
 		Scene,
-		TextureRenderTarget->GameThread_GetRenderTargetResource(),
+		Context.TextureRenderTarget->GameThread_GetRenderTargetResource(),
 		Context,
 		CaptureExtent,
 		ViewRotationMat,
@@ -569,14 +577,14 @@ void UpdateWaterInfoRendering(
 	
 	FSceneRenderer* ColorRenderer = CreateWaterInfoColorRenderer(
 		Scene,
-		TextureRenderTarget->GameThread_GetRenderTargetResource(),
+		Context.TextureRenderTarget->GameThread_GetRenderTargetResource(),
 		Context,
 		CaptureExtent,
 		ViewRotationMat,
 		ViewLocation,
 		OrthoProj);
 
-	FTextureRenderTargetResource* TextureRenderTargetResource = TextureRenderTarget->GameThread_GetRenderTargetResource();
+	FTextureRenderTargetResource* TextureRenderTargetResource = Context.TextureRenderTarget->GameThread_GetRenderTargetResource();
 
 	FUpdateWaterInfoParams Params;
 	Params.DepthRenderer = DepthRenderer;
diff --git a/Engine/Plugins/Experimental/Water/Source/Runtime/Private/WaterInfoRendering.h b/Engine/Plugins/Experimental/Water/Source/Runtime/Private/WaterInfoRendering.h
--- a/Engine/Plugins/Experimental/Water/Source/Runtime/Private/WaterInfoRendering.h
+++ b/Engine/Plugins/Experimental/Water/Source/Runtime/Private/WaterInfoRendering.h
@@ -10,11 +10,12 @@ class FSceneInterface;
 class UTextureRenderTarget2D;
 class AActor;
 
-namespace WaterInfo
+namespace UE::WaterInfo
 {
 	struct FRenderingContext
 	{
 		AWaterZone* ZoneToRender = nullptr;
+		UTextureRenderTarget2D* TextureRenderTarget;
 		TArray<UWaterBodyComponent*> WaterBodies;
 		TArray<AActor*> GroundActors;
 		float CaptureZ;
@@ -22,6 +23,5 @@ namespace WaterInfo
 	
 void UpdateWaterInfoRendering(
 	FSceneInterface* Scene,
-	UTextureRenderTarget2D* TextureRenderTarget,
 	const FRenderingContext& Context);
 }
diff --git a/Engine/Plugins/Experimental/Water/Source/Runtime/Private/WaterSubsystem.cpp b/Engine/Plugins/Experimental/Water/Source/Runtime/Private/WaterSubsystem.cpp
--- a/Engine/Plugins/Experimental/Water/Source/Runtime/Private/WaterSubsystem.cpp
+++ b/Engine/Plugins/Experimental/Water/Source/Runtime/Private/WaterSubsystem.cpp
@@ -24,6 +24,7 @@
 #include "WaterRuntimeSettings.h"
 #include "Engine/CollisionProfile.h"
 #include "Engine/StaticMesh.h"
+#include "WaterViewExtension.h"
 
 // ----------------------------------------------------------------------------------
 
@@ -244,6 +245,7 @@ void UWaterSubsystem::Initialize(FSubsystemCollectionBase& Collection)
 	check(World != nullptr);
 
 	WaterBodyManager.Initialize(World);
+	WaterViewExtension = FSceneViewExtensions::NewExtension<FWaterViewExtension>(World);
 
 	bUsingSmoothedTime = false;
 	FConsoleVariableDelegate NotifyWaterScalabilityChanged = FConsoleVariableDelegate::CreateUObject(this, &UWaterSubsystem::NotifyWaterScalabilityChangedInternal);
@@ -525,6 +527,11 @@ void UWaterSubsystem::MarkAllWaterZonesForRebuild(EWaterZoneRebuildFlags Rebuild
 	}
 }
 
+void UWaterSubsystem::MarkWaterInfoTextureForRebuild(const UE::WaterInfo::FRenderingContext& WaterInfoContext)
+{
+	WaterViewExtension->MarkWaterInfoTextureForRebuild(WaterInfoContext);
+}
+
 void UWaterSubsystem::NotifyWaterScalabilityChangedInternal(IConsoleVariable* CVar)
 {
 	OnWaterScalabilityChanged.Broadcast();
diff --git a/Engine/Plugins/Experimental/Water/Source/Runtime/Private/WaterViewExtension.cpp b/Engine/Plugins/Experimental/Water/Source/Runtime/Private/WaterViewExtension.cpp
new file mode 100644
--- /dev/null
+++ b/Engine/Plugins/Experimental/Water/Source/Runtime/Private/WaterViewExtension.cpp
@@ -0,0 +1,45 @@
+// Copyright Epic Games, Inc. All Rights Reserved.
+
+#include "WaterViewExtension.h"
+#include "WaterBodyComponent.h"
+#include "WaterSubsystem.h"
+#include "WaterModule.h"
+#include "WaterZoneActor.h"
+#include "WaterInfoRendering.h"
+
+FWaterViewExtension::FWaterViewExtension(const FAutoRegister& AutoReg, UWorld* InWorld)
+	: FWorldSceneViewExtension(AutoReg, InWorld)
+{
+}
+
+FWaterViewExtension::~FWaterViewExtension()
+{
+}
+
+void FWaterViewExtension::SetupViewFamily(FSceneViewFamily& InViewFamily)
+{
+	if (WaterInfoContextsToRender.Num() > 0)
+	{
+		const TWeakObjectPtr<UWorld> WorldPtr = GetWorld();
+		check(WorldPtr.IsValid())
+
+		// Move out of the cache list to prevent re-entering this function when we set up the view extensions for the water info update itself
+		TMap<AWaterZone*, UE::WaterInfo::FRenderingContext> ContextsToRender = MoveTemp(WaterInfoContextsToRender);
+		WaterInfoContextsToRender.Empty();
+
+		for (const TPair<AWaterZone*, UE::WaterInfo::FRenderingContext>& Pair : ContextsToRender)
+		{
+			UE::WaterInfo::UpdateWaterInfoRendering(WorldPtr.Get()->Scene, Pair.Value);
+		}
+	}
+}
+
+void FWaterViewExtension::PreRenderViewFamily_RenderThread(FRHICommandListImmediate& RHICmdList, FSceneViewFamily& InViewFamily)
+{
+}
+
+void FWaterViewExtension::MarkWaterInfoTextureForRebuild(const UE::WaterInfo::FRenderingContext& RenderContext)
+{
+	WaterInfoContextsToRender.Emplace(RenderContext.ZoneToRender, RenderContext);
+}
+
diff --git a/Engine/Plugins/Experimental/Water/Source/Runtime/Private/WaterZoneActor.cpp b/Engine/Plugins/Experimental/Water/Source/Runtime/Private/WaterZoneActor.cpp
--- a/Engine/Plugins/Experimental/Water/Source/Runtime/Private/WaterZoneActor.cpp
+++ b/Engine/Plugins/Experimental/Water/Source/Runtime/Private/WaterZoneActor.cpp
@@ -22,12 +22,6 @@
 #include "Components/BoxComponent.h"
 #endif // WITH_EDITOR
 
-static int32 RenderCaptureNextWaterInfoDraws = 0;
-static FAutoConsoleVariableRef CVarRenderCaptureNextWaterInfoDraws(
-	TEXT("r.Water.WaterInfo.RenderCaptureNextWaterInfoDraws"),
-	RenderCaptureNextWaterInfoDraws,
-	TEXT("Enable capturing of the water info texture for the next N draws"));
-
 static int32 ForceUpdateWaterInfoNextFrames = 0;
 static FAutoConsoleVariableRef CVarForceUpdateWaterInfoNextFrames(
 	TEXT("r.Water.WaterInfo.ForceUpdateWaterInfoNextFrames"),
@@ -230,9 +224,6 @@ void AWaterZone::UpdateWaterInfoTexture()
 {
 	if (UWorld* World = GetWorld(); World && FApp::CanEverRender())
 	{
-		RenderCaptureInterface::FScopedCapture RenderCapture((RenderCaptureNextWaterInfoDraws != 0), TEXT("RenderWaterInfo"));
-		RenderCaptureNextWaterInfoDraws = FMath::Max(0, RenderCaptureNextWaterInfoDraws - 1);
-
 		float WaterZMin(TNumericLimits<double>::Max());
 		float WaterZMax(TNumericLimits<double>::Lowest());
 	
@@ -270,15 +261,20 @@ void AWaterZone::UpdateWaterInfoTexture()
 			GroundActors.Add(LandscapeProxy);
 		}
 
-		WaterInfo::FRenderingContext Context;
+		const ETextureRenderTargetFormat Format = bHalfPrecisionTexture ? ETextureRenderTargetFormat::RTF_RGBA16f : RTF_RGBA32f;
+		WaterInfoTexture = FWaterUtils::GetOrCreateTransientRenderTarget2D(WaterInfoTexture, TEXT("WaterInfoTexture"), RenderTargetResolution, Format);
+
+		UE::WaterInfo::FRenderingContext Context;
 		Context.ZoneToRender = this;
 		Context.WaterBodies = WaterBodies;
 		Context.GroundActors = MoveTemp(GroundActors);
 		Context.CaptureZ = WaterZMax;
+		Context.TextureRenderTarget = WaterInfoTexture;
 
-		const ETextureRenderTargetFormat Format = bHalfPrecisionTexture ? ETextureRenderTargetFormat::RTF_RGBA16f : RTF_RGBA32f;
-		WaterInfoTexture = FWaterUtils::GetOrCreateTransientRenderTarget2D(WaterInfoTexture, TEXT("WaterInfoTexture"), RenderTargetResolution, Format);
-		WaterInfo::UpdateWaterInfoRendering(World->Scene, WaterInfoTexture, Context);
+		if (UWaterSubsystem* WaterSubsystem = UWaterSubsystem::GetWaterSubsystem(World))
+		{
+			WaterSubsystem->MarkWaterInfoTextureForRebuild(Context);
+		}
 
 		for (UWaterBodyComponent* Component : WaterBodies)
 		{
diff --git a/Engine/Plugins/Experimental/Water/Source/Runtime/Public/WaterSubsystem.h b/Engine/Plugins/Experimental/Water/Source/Runtime/Public/WaterSubsystem.h
--- a/Engine/Plugins/Experimental/Water/Source/Runtime/Public/WaterSubsystem.h
+++ b/Engine/Plugins/Experimental/Water/Source/Runtime/Public/WaterSubsystem.h
@@ -26,6 +26,9 @@ class UTexture2D;
 struct FUnderwaterPostProcessDebugInfo;
 enum class EWaterBodyQueryFlags;
 class ABuoyancyManager;
+class FWaterViewExtension;
+
+namespace UE::WaterInfo { struct FRenderingContext; }
 
 bool IsWaterEnabled(bool bIsRenderThread);
 
@@ -167,6 +170,8 @@ public:
 	
 	void MarkAllWaterZonesForRebuild(EWaterZoneRebuildFlags RebuildFlags = EWaterZoneRebuildFlags::All);
 
+	void MarkWaterInfoTextureForRebuild(const UE::WaterInfo::FRenderingContext& WaterInfoContext);
+
 #if WITH_EDITOR
 	/** Little scope object to temporarily change the value of bAllowWaterSubsystemOnPreviewWorld */
 	struct WATER_API FScopedAllowWaterSubsystemOnPreviewWorld
@@ -242,6 +247,8 @@ private:
 
 	FUnderwaterPostProcessVolume UnderwaterPostProcessVolume;
 
+	TSharedPtr<FWaterViewExtension, ESPMode::ThreadSafe> WaterViewExtension;
+
 #if WITH_EDITOR
 	/** By default, there is no water subsystem allowed on preview worlds except when explicitly requested : */
 	static bool bAllowWaterSubsystemOnPreviewWorld;
diff --git a/Engine/Plugins/Experimental/Water/Source/Runtime/Public/WaterViewExtension.h b/Engine/Plugins/Experimental/Water/Source/Runtime/Public/WaterViewExtension.h
new file mode 100644
--- /dev/null
+++ b/Engine/Plugins/Experimental/Water/Source/Runtime/Public/WaterViewExtension.h
@@ -0,0 +1,31 @@
+// Copyright Epic Games, Inc. All Rights Reserved.
+
+#pragma once
+
+#include "CoreMinimal.h"
+#include "SceneViewExtension.h"
+#include "WaterInfoRendering.h"
+
+
+class AWaterZone;
+
+class FWaterViewExtension : public FWorldSceneViewExtension
+{
+public:
+	FWaterViewExtension(const FAutoRegister& AutoReg, UWorld* InWorld);
+	~FWaterViewExtension();
+
+	// FSceneViewExtensionBase implementation : 
+	virtual void SetupView(FSceneViewFamily& InViewFamily, FSceneView& InView) override {}
+	virtual void BeginRenderViewFamily(FSceneViewFamily& InViewFamily) override {}
+
+	virtual void SetupViewFamily(FSceneViewFamily& InViewFamily) override;
+	virtual void PreRenderViewFamily_RenderThread(FRHICommandListImmediate& RHICmdList, FSceneViewFamily& InViewFamily) override;
+	// End FSceneViewExtensionBase implementation
+
+	void MarkWaterInfoTextureForRebuild(const UE::WaterInfo::FRenderingContext& RenderContext);
+
+private:
+	// Store contexts in a map to prevent multiple update entries for the same zone. We only need to keep the most recent.
+	TMap<AWaterZone*, UE::WaterInfo::FRenderingContext> WaterInfoContextsToRender;
+};
