Explicitly enable 128 bit output on the WaterInfo pixel shaders when the platform requires it.

diff --git a/Engine/Plugins/Experimental/Water/Source/Runtime/Private/WaterInfoRendering.cpp b/Engine/Plugins/Experimental/Water/Source/Runtime/Private/WaterInfoRendering.cpp
--- a/Engine/Plugins/Experimental/Water/Source/Runtime/Private/WaterInfoRendering.cpp
+++ b/Engine/Plugins/Experimental/Water/Source/Runtime/Private/WaterInfoRendering.cpp
@@ -60,6 +60,16 @@ public:
 		SHADER_PARAMETER(float, CaptureZ)
 		RENDER_TARGET_BINDING_SLOTS()
 	END_SHADER_PARAMETER_STRUCT()
+
+	static void ModifyCompilationEnvironment(const FGlobalShaderPermutationParameters& Parameters, FShaderCompilerEnvironment& OutEnvironment)
+	{
+		// Water info merge unconditionally requires a 128 bit render target. Some platforms require explicitly enabling this output mode.
+		bool bPlatformRequiresExplicit128bitRT = FDataDrivenShaderPlatformInfo::GetRequiresExplicit128bitRT(Parameters.Platform);
+		if (bPlatformRequiresExplicit128bitRT)
+		{
+			OutEnvironment.SetRenderTargetOutputFormat(0, PF_A32B32G32R32F);
+		}
+	}
 };
 
 IMPLEMENT_GLOBAL_SHADER(FWaterInfoMergePS, "/Plugin/Water/Private/WaterInfoMerge.usf", "Main", SF_Pixel);
@@ -145,6 +155,33 @@ public:
 		SHADER_PARAMETER(int, BlurRadius)
 		RENDER_TARGET_BINDING_SLOTS()
 	END_SHADER_PARAMETER_STRUCT()
+
+	class FEnable128BitRT : SHADER_PERMUTATION_BOOL("ENABLE_128_BIT");
+	using FPermutationDomain = TShaderPermutationDomain<FEnable128BitRT>;
+
+	static FPermutationDomain GetPermutationVector(bool bUse128BitRT)
+	{
+		FPermutationDomain PermutationVector;
+		PermutationVector.Set<FEnable128BitRT>(bUse128BitRT);
+		return PermutationVector;
+	}
+
+	static bool ShouldCompilePermutation(const FGlobalShaderPermutationParameters& Parameters) 
+	{
+		FPermutationDomain PermutationVector(Parameters.PermutationId);
+
+		bool bPlatformRequiresExplicit128bitRT = FDataDrivenShaderPlatformInfo::GetRequiresExplicit128bitRT(Parameters.Platform);
+		return (!PermutationVector.Get<FEnable128BitRT>() || bPlatformRequiresExplicit128bitRT);
+	}
+
+	static void ModifyCompilationEnvironment(const FGlobalShaderPermutationParameters& Parameters, FShaderCompilerEnvironment& OutEnvironment)
+	{
+		const FPermutationDomain PermutationVector(Parameters.PermutationId);
+		if (PermutationVector.Get<FEnable128BitRT>())
+		{
+			OutEnvironment.SetRenderTargetOutputFormat(0, PF_A32B32G32R32F);
+		}
+	}
 };
 
 IMPLEMENT_GLOBAL_SHADER(FWaterInfoFinalizePS, "/Plugin/Water/Private/WaterInfoFinalize.usf", "Main", SF_Pixel);
@@ -165,6 +202,9 @@ static void FinalizeWaterInfo(
 	GraphicsPSOInit.DepthStencilState = TStaticDepthStencilState<false, CF_Always>::GetRHI();
 	GraphicsPSOInit.BlendState = TStaticBlendState<>::GetRHI();
 
+	const bool bUse128BitRT = PlatformRequires128bitRT(OutputTexture->Desc.Format);
+	const FWaterInfoFinalizePS::FPermutationDomain PixelPermutationVector = FWaterInfoFinalizePS::GetPermutationVector(bUse128BitRT);
+
 	{
 		FWaterInfoFinalizePS::FParameters* PassParameters = GraphBuilder.AllocParameters<FWaterInfoFinalizePS::FParameters>();
 		PassParameters->View = View.ViewUniformBuffer;
@@ -177,7 +217,7 @@ static void FinalizeWaterInfo(
 		PassParameters->GroundZMin = Params.GroundZMin;
 
 		TShaderMapRef<FScreenVS> VertexShader(View.ShaderMap);
-		TShaderMapRef<FWaterInfoFinalizePS> PixelShader(View.ShaderMap);
+		TShaderMapRef<FWaterInfoFinalizePS> PixelShader(View.ShaderMap, PixelPermutationVector);
 
 		GraphicsPSOInit.BoundShaderState.VertexDeclarationRHI = GFilterVertexDeclaration.VertexDeclarationRHI;
 		GraphicsPSOInit.BoundShaderState.VertexShaderRHI = VertexShader.GetVertexShader();
