Merged WaterMesh.ShowTileBounds and WaterMesh.TileBoundsColor as well as re-added the water body type coloring mode.

diff --git a/Engine/Plugins/Experimental/Water/Source/Runtime/Private/WaterMeshSceneProxy.cpp b/Engine/Plugins/Experimental/Water/Source/Runtime/Private/WaterMeshSceneProxy.cpp
--- a/Engine/Plugins/Experimental/Water/Source/Runtime/Private/WaterMeshSceneProxy.cpp
+++ b/Engine/Plugins/Experimental/Water/Source/Runtime/Private/WaterMeshSceneProxy.cpp
@@ -58,14 +58,7 @@ static TAutoConsoleVariable<int32> CVarWaterMeshShowLODLevels(
 static TAutoConsoleVariable<int32> CVarWaterMeshShowTileBounds(
 	TEXT("r.Water.WaterMesh.ShowTileBounds"),
 	0,
-	TEXT("Shows the tile bounds colored by r.Water.WaterMesh.TileBoundsColor"),
-	ECVF_RenderThreadSafe
-);
-
-static TAutoConsoleVariable<int32> CVarWaterMeshTileBoundsColor(
-	TEXT("r.Water.WaterMesh.TileBoundsColor"),
-	1,
-	TEXT("Color of the tile bounds visualized by r.Water.WaterMesh.ShowTileBounds. 0 is by LOD, 1 is by water body type"),
+	TEXT("Shows the tile bounds with optional color modes: 0 is disabled, 1 is by water body type, 2 is by LOD, 3 is by density index"),
 	ECVF_RenderThreadSafe
 );
 
@@ -268,8 +261,7 @@ void FWaterMeshSceneProxy::GetDynamicMeshElements(const TArray<const FSceneView*
 #if !(UE_BUILD_SHIPPING || UE_BUILD_TEST)
 			//Debug
 			TraversalDesc.DebugPDI = Collector.GetPDI(ViewIndex);
-			TraversalDesc.DebugShowTile = !!CVarWaterMeshShowTileBounds.GetValueOnRenderThread();
-			TraversalDesc.DebugShowTypeColor = !!CVarWaterMeshTileBoundsColor.GetValueOnRenderThread();
+			TraversalDesc.DebugShowTile = CVarWaterMeshShowTileBounds.GetValueOnRenderThread();
 #endif
 			WaterQuadTree.BuildWaterTileInstanceData(TraversalDesc, WaterInstanceData);
 
diff --git a/Engine/Plugins/Experimental/Water/Source/Runtime/Private/WaterQuadTree.cpp b/Engine/Plugins/Experimental/Water/Source/Runtime/Private/WaterQuadTree.cpp
--- a/Engine/Plugins/Experimental/Water/Source/Runtime/Private/WaterQuadTree.cpp
+++ b/Engine/Plugins/Experimental/Water/Source/Runtime/Private/WaterQuadTree.cpp
@@ -94,17 +94,22 @@ void FWaterQuadTree::FNode::AddNodeForRender(const FNodeData& InNodeData, const
 
 #if !(UE_BUILD_SHIPPING || UE_BUILD_TEST)
 	// Debug drawing
-	if (InTraversalDesc.DebugShowTile)
+	if (InTraversalDesc.DebugShowTile != 0)
 	{
 		FColor Color;
-		if (InTraversalDesc.DebugShowTypeColor)
+		if (InTraversalDesc.DebugShowTile == 1)
 		{
-			//static FColor WaterTypeColor[] = { FColor::Red, FColor::Green, FColor::Blue, FColor::Yellow, FColor::Purple };
-			Color = GColorList.GetFColorByIndex(DensityIndex + 1);
+			static FColor WaterTypeColor[] = { FColor::Red, FColor::Green, FColor::Blue, FColor::Yellow, FColor::Purple };
+			Color = WaterTypeColor[TileDebugID];
 		}
-		else
+		else if (InTraversalDesc.DebugShowTile == 2)
+		{
+			Color = GColorList.GetFColorByIndex(InLODLevel + 1);
+		}
+		else if (InTraversalDesc.DebugShowTile == 3)
 		{
-			Color = GColorList.GetFColorByIndex(InLODLevel+1);
+
+			Color = GColorList.GetFColorByIndex(DensityIndex + 1);
 		}
 
 		DrawWireBox(InTraversalDesc.DebugPDI, Bounds.ExpandBy(FVector(-20.0f, -20.0f, 0.0f)), Color, SDPG_World);
diff --git a/Engine/Plugins/Experimental/Water/Source/Runtime/Public/WaterQuadTree.h b/Engine/Plugins/Experimental/Water/Source/Runtime/Public/WaterQuadTree.h
--- a/Engine/Plugins/Experimental/Water/Source/Runtime/Public/WaterQuadTree.h
+++ b/Engine/Plugins/Experimental/Water/Source/Runtime/Public/WaterQuadTree.h
@@ -110,8 +110,7 @@ struct FWaterQuadTree
 
 #if !(UE_BUILD_SHIPPING || UE_BUILD_TEST)
 		// Debug
-		bool DebugShowTile = false;
-		bool DebugShowTypeColor = true;
+		int32 DebugShowTile = 0;
 		class FPrimitiveDrawInterface* DebugPDI = nullptr;
 #endif
 	};
