Fixed water info rendering on mobile path due to fog pass adding artificial velocity everywhere.

diff --git a/Engine/Plugins/Experimental/Water/Source/Runtime/Private/WaterInfoRendering.cpp b/Engine/Plugins/Experimental/Water/Source/Runtime/Private/WaterInfoRendering.cpp
--- a/Engine/Plugins/Experimental/Water/Source/Runtime/Private/WaterInfoRendering.cpp
+++ b/Engine/Plugins/Experimental/Water/Source/Runtime/Private/WaterInfoRendering.cpp
@@ -419,14 +419,7 @@ static void UpdateWaterInfoRendering_RenderThread(
 	RHICmdList.Transition(FRHITransitionInfo(Params.OutputTexture->TextureRHI, ERHIAccess::RTV, ERHIAccess::SRVMask));
 }
 
-static FSceneRenderer* CreateWaterInfoDepthRenderer(
-	FSceneInterface* Scene,
-	FRenderTarget* RenderTarget,
-	const WaterInfo::FRenderingContext& Context,
-	FIntPoint RenderTargetSize,
-	const FMatrix& ViewRotationMatrix,
-	const FVector& ViewLocation,
-	const FMatrix& ProjectionMatrix)
+static FEngineShowFlags GetWaterInfoBaseShowFlags()
 {
 	FEngineShowFlags ShowFlags(ESFIM_Game);
 	ShowFlags.NaniteMeshes = 0;
@@ -437,6 +430,23 @@ static FSceneRenderer* CreateWaterInfoDepthRenderer(
 	ShowFlags.Translucency = 0;
 	ShowFlags.SeparateTranslucency = 0;
 	ShowFlags.AntiAliasing = 0;
+	ShowFlags.Fog = 0;
+	ShowFlags.VolumetricFog = 0;
+	ShowFlags.DynamicShadows = 0;
+	return ShowFlags;
+}
+
+static FSceneRenderer* CreateWaterInfoDepthRenderer(
+	FSceneInterface* Scene,
+	FRenderTarget* RenderTarget,
+	const WaterInfo::FRenderingContext& Context,
+	FIntPoint RenderTargetSize,
+	const FMatrix& ViewRotationMatrix,
+	const FVector& ViewLocation,
+	const FMatrix& ProjectionMatrix)
+{
+	FEngineShowFlags ShowFlags = GetWaterInfoBaseShowFlags();
+
 	FSceneViewFamilyContext DepthViewFamily(FSceneViewFamily::ConstructionValues(
 		RenderTarget,
 		Scene,
@@ -511,15 +521,7 @@ static FSceneRenderer* CreateWaterInfoColorRenderer(
 	const FVector& ViewLocation,
 	const FMatrix& ProjectionMatrix)
 {
-	FEngineShowFlags ShowFlags(ESFIM_Game);
-	ShowFlags.NaniteMeshes = 0;
-	ShowFlags.Atmosphere = 0;
-	ShowFlags.Lighting = 0;
-	ShowFlags.Bloom = 0;
-	ShowFlags.ScreenPercentage = 0;
-	ShowFlags.Translucency = 0;
-	ShowFlags.SeparateTranslucency = 0;
-	ShowFlags.AntiAliasing = 0;
+	FEngineShowFlags ShowFlags = GetWaterInfoBaseShowFlags();
 	
 	FSceneViewFamilyContext ColorViewFamily(FSceneViewFamily::ConstructionValues(
 		RenderTarget,
