Improve precision of the water depth by removing the CaptureZOffset from the range within which we normalize the water height values.

diff --git a/Engine/Plugins/Experimental/Water/Source/Runtime/Private/WaterZoneActor.cpp b/Engine/Plugins/Experimental/Water/Source/Runtime/Private/WaterZoneActor.cpp
--- a/Engine/Plugins/Experimental/Water/Source/Runtime/Private/WaterZoneActor.cpp
+++ b/Engine/Plugins/Experimental/Water/Source/Runtime/Private/WaterZoneActor.cpp
@@ -245,8 +245,6 @@ void AWaterZone::UpdateWaterInfoTexture()
 			return;
 		}
 
-		WaterZMax += CaptureZOffset;
-
 		checkf(WaterZMin != WaterZMax, TEXT("Water has a height extent of 0 which can cause divide by zero errors in the shaders! Ensure water bodies have proper bounding boxes."));
 		WaterHeightExtents = FVector2f(WaterZMin, WaterZMax);
 
@@ -268,7 +266,7 @@ void AWaterZone::UpdateWaterInfoTexture()
 		Context.ZoneToRender = this;
 		Context.WaterBodies = WaterBodies;
 		Context.GroundActors = MoveTemp(GroundActors);
-		Context.CaptureZ = WaterZMax;
+		Context.CaptureZ = WaterZMax + CaptureZOffset;
 		Context.TextureRenderTarget = WaterInfoTexture;
 
 		if (UWaterSubsystem* WaterSubsystem = UWaterSubsystem::GetWaterSubsystem(World))
