Added a way to delay updates to the WaterInfoTexture.
This is necessary when the DrawWaterInfo materials haven't yet been compiled under ODSC.
Additionally, this can be useful for preventing the waterinfo from drawing when landscape mips aren't ready in-editor.

diff --git a/Engine/Plugins/Experimental/Water/Source/Runtime/Private/WaterZoneActor.cpp b/Engine/Plugins/Experimental/Water/Source/Runtime/Private/WaterZoneActor.cpp
--- a/Engine/Plugins/Experimental/Water/Source/Runtime/Private/WaterZoneActor.cpp
+++ b/Engine/Plugins/Experimental/Water/Source/Runtime/Private/WaterZoneActor.cpp
@@ -119,8 +119,10 @@ void AWaterZone::Update()
 	if (bNeedsWaterInfoRebuild || (ForceUpdateWaterInfoNextFrames != 0))
 	{
 		ForceUpdateWaterInfoNextFrames = (ForceUpdateWaterInfoNextFrames < 0) ? ForceUpdateWaterInfoNextFrames : FMath::Max(0, ForceUpdateWaterInfoNextFrames - 1);
-		UpdateWaterInfoTexture();
-		bNeedsWaterInfoRebuild = false;
+		if (UpdateWaterInfoTexture())
+		{
+			bNeedsWaterInfoRebuild = false;
+		}
 	}
 	
 	if (WaterMesh)
@@ -220,18 +222,32 @@ void AWaterZone::OnBoundsComponentModified()
 }
 #endif // WITH_EDITOR
 
-void AWaterZone::UpdateWaterInfoTexture()
+bool AWaterZone::UpdateWaterInfoTexture()
 {
 	if (UWorld* World = GetWorld(); World && FApp::CanEverRender())
 	{
 		float WaterZMin(TNumericLimits<double>::Max());
 		float WaterZMax(TNumericLimits<double>::Lowest());
 	
+		bool bHasIncompleteShaderMaps = false;
 		// #todo_water [roey]: we need to store which actors this zone is responsible for rendering once we implement support for multi zones.
 		// For now whenever we update the water info texture we will just collect all water bodies and pass those to the renderer.
 		TArray<UWaterBodyComponent*> WaterBodies;
-		UWaterSubsystem::ForEachWaterBodyComponent(World, [&WaterBodies, &WaterZMax, &WaterZMin](UWaterBodyComponent* Component)
+		UWaterSubsystem::ForEachWaterBodyComponent(World, [World, &WaterBodies, &WaterZMax, &WaterZMin, &bHasIncompleteShaderMaps](UWaterBodyComponent* Component)
 		{
+			if (UMaterialInterface* WaterInfoMaterial = Component->GetWaterInfoMaterialInstance())
+			{
+				if (FMaterialResource* MaterialResource = WaterInfoMaterial->GetMaterialResource(World->Scene->GetFeatureLevel()))
+				{
+					if (!MaterialResource->IsGameThreadShaderMapComplete())
+					{
+						MaterialResource->SubmitCompileJobs_GameThread(EShaderCompileJobPriority::ForceLocal);
+						bHasIncompleteShaderMaps = true;
+						return true;
+					}
+				}
+			}
+
 			WaterBodies.Add(Component);
 			const FBox WaterBodyBounds = Component->CalcBounds(Component->GetComponentToWorld()).GetBox();
 			WaterZMax = FMath::Max(WaterZMax, WaterBodyBounds.Max.Z);
@@ -239,10 +255,15 @@ void AWaterZone::UpdateWaterInfoTexture()
 			return true;
 		});
 
+		if (bHasIncompleteShaderMaps)
+		{
+			return false;
+		}
+
 		// If we don't have any water bodies we don't need to do anything.
 		if (WaterBodies.Num() == 0)
 		{
-			return;
+			return true;
 		}
 
 		checkf(WaterZMin != WaterZMax, TEXT("Water has a height extent of 0 which can cause divide by zero errors in the shaders! Ensure water bodies have proper bounding boxes."));
@@ -281,4 +302,6 @@ void AWaterZone::UpdateWaterInfoTexture()
 
 		UE_LOG(LogWater, Verbose, TEXT("Queued Water Info texture update"));
 	}
+
+	return true;
 }
diff --git a/Engine/Plugins/Experimental/Water/Source/Runtime/Public/WaterZoneActor.h b/Engine/Plugins/Experimental/Water/Source/Runtime/Public/WaterZoneActor.h
--- a/Engine/Plugins/Experimental/Water/Source/Runtime/Public/WaterZoneActor.h
+++ b/Engine/Plugins/Experimental/Water/Source/Runtime/Public/WaterZoneActor.h
@@ -51,7 +51,7 @@ public:
 	UTextureRenderTarget2D* WaterInfoTexture;
 private:
 
-	void UpdateWaterInfoTexture();
+	bool UpdateWaterInfoTexture();
 
 	void OnExtentChanged();
 	
