Support water bodies with zero channel depth.
This is required to properly support water bodies which may not affect the landscape.
Also unified the way dilated regions are pushed below the terrain using the material rather than handling it in the merge shader as well.

diff --git a/Engine/Plugins/Experimental/Water/Shaders/Private/WaterInfoMerge.usf b/Engine/Plugins/Experimental/Water/Shaders/Private/WaterInfoMerge.usf
--- a/Engine/Plugins/Experimental/Water/Shaders/Private/WaterInfoMerge.usf
+++ b/Engine/Plugins/Experimental/Water/Shaders/Private/WaterInfoMerge.usf
@@ -33,6 +33,18 @@ float DecodeFloatRGB(float3 Enc)
 	return 1 - dot(Enc, Encoding);
 }
 
+float NormalizeHeight(float Height, float Min, float Max)
+{
+	float NormalizedHeight = 1.f;
+	// It's possible that the min and max are the same value in which case we might end up generating a div by zero error
+	// In this case it's okay to use the Normalized value of 1.0 which will just map to the Min value.
+	if (Min != Max)
+	{
+		NormalizedHeight = saturate((Height - Min) / (Max - Min));
+	}
+	return NormalizedHeight;
+}
+
 void Main(
 	in FScreenVertexOutput Input,
 	out float4 OutColor : SV_Target0
@@ -41,31 +53,34 @@ void Main(
 	const float4 SVPos = Input.Position;
 	uint2 TextureCoordinates = floor(SVPos.xy);
 	
-    const float GroundDeviceDepth = DecodeFloatRGB(DepthTexture.Load(int3(TextureCoordinates, 0)).rgb);
+	const float GroundDeviceDepth = DecodeFloatRGB(DepthTexture.Load(int3(TextureCoordinates, 0)).rgb);
 	const float GroundSceneDepth = ConvertFromDeviceZ(GroundDeviceDepth);
-    const FWaterColorPassResult WaterInfo = MakeWaterColorPassResult(ColorTexture.Load(int3(TextureCoordinates, 0)));
+	const FWaterColorPassResult WaterInfo = MakeWaterColorPassResult(ColorTexture.Load(int3(TextureCoordinates, 0)));
 	const float WaterSceneDepth = WaterInfo.SceneDepth;
-	float WaterZ = CaptureZ - WaterSceneDepth;
+	const float WaterZ = CaptureZ - WaterSceneDepth;
 
 	float GroundZ = CaptureZ - GroundSceneDepth;
 
-	// If this is a dilated region and it is above the ground we must push it down to prevent artifacts
-	if (WaterInfo.WaterBodyIndex == -1 && WaterZ > GroundZ)
-	{
-		WaterZ = GroundZ - 32.f;
-	}
-
 	const float WaterZMin = WaterHeightExtents.x;
 	const float WaterZMax = WaterHeightExtents.y;
 
+	float NormalizedWaterZ = NormalizeHeight(WaterZ, WaterZMin, WaterZMax);
+	
+	// If this is a dilated region (indicated by WaterBodyIndex == -1) and it is above the ground we must push it down to prevent artifacts
+	// A negative NormalizedWaterZ is a special value for the material and will map to a value slightly below the terrain.
+	if (WaterInfo.WaterBodyIndex == -1 && WaterZ > GroundZ)
+	{
+		NormalizedWaterZ = -1.f;
+	}
+	
 	// If the GroundDepth is further than the far plane of the projection, set the GroundZ to some really far away value
 	if (GroundDeviceDepth == 1.0)
 	{
 		GroundZ = GroundZMin;
 	}
+	
+	const float NormalizedGroundZ = NormalizeHeight(GroundZ, GroundZMin, WaterZMax);
 
-	float NormalizedWaterZ = saturate((WaterZ - WaterZMin) / (WaterZMax - WaterZMin));
-	float NormalizedGroundZ = saturate((GroundZ - GroundZMin) / (WaterZMax - GroundZMin));
 
 	OutColor = float4(WaterInfo.Velocity, NormalizedWaterZ, NormalizedGroundZ);
 }
diff --git a/Engine/Plugins/Experimental/Water/Source/Runtime/Private/WaterZoneActor.cpp b/Engine/Plugins/Experimental/Water/Source/Runtime/Private/WaterZoneActor.cpp
--- a/Engine/Plugins/Experimental/Water/Source/Runtime/Private/WaterZoneActor.cpp
+++ b/Engine/Plugins/Experimental/Water/Source/Runtime/Private/WaterZoneActor.cpp
@@ -266,7 +266,6 @@ bool AWaterZone::UpdateWaterInfoTexture()
 			return true;
 		}
 
-		checkf(WaterZMin != WaterZMax, TEXT("Water has a height extent of 0 which can cause divide by zero errors in the shaders! Ensure water bodies have proper bounding boxes."));
 		WaterHeightExtents = FVector2f(WaterZMin, WaterZMax);
 
 		// Only compute the ground min since we can use the water max z as the ground max z for more precision.
