Reverted change to water info merge which was causing single pixels in the water info texture to have incorrect depth values. A proper fix will depend on a better implementation for dilation

diff --git a/Engine/Plugins/Experimental/Water/Shaders/Private/WaterInfoMerge.usf b/Engine/Plugins/Experimental/Water/Shaders/Private/WaterInfoMerge.usf
--- a/Engine/Plugins/Experimental/Water/Shaders/Private/WaterInfoMerge.usf
+++ b/Engine/Plugins/Experimental/Water/Shaders/Private/WaterInfoMerge.usf
@@ -57,21 +57,21 @@ void Main(
 	const float GroundSceneDepth = ConvertFromDeviceZ(GroundDeviceDepth);
 	const FWaterColorPassResult WaterInfo = MakeWaterColorPassResult(ColorTexture.Load(int3(TextureCoordinates, 0)));
 	const float WaterSceneDepth = WaterInfo.SceneDepth;
-	const float WaterZ = CaptureZ - WaterSceneDepth;
+	float WaterZ = CaptureZ - WaterSceneDepth;
 
 	float GroundZ = CaptureZ - GroundSceneDepth;
 
 	const float WaterZMin = WaterHeightExtents.x;
 	const float WaterZMax = WaterHeightExtents.y;
-
-	float NormalizedWaterZ = NormalizeHeight(WaterZ, WaterZMin, WaterZMax);
 	
 	// If this is a dilated region (indicated by WaterBodyIndex == -1) and it is above the ground we must push it down to prevent artifacts
 	// A negative NormalizedWaterZ is a special value for the material and will map to a value slightly below the terrain.
 	if (WaterInfo.WaterBodyIndex == -1 && WaterZ > GroundZ)
 	{
-		NormalizedWaterZ = -1.f;
+		WaterZ = GroundZ - 32.f;
 	}
+
+	const float NormalizedWaterZ = NormalizeHeight(WaterZ, WaterZMin, WaterZMax);
 	
 	// If the GroundDepth is further than the far plane of the projection, set the GroundZ to some really far away value
 	if (GroundDeviceDepth == 1.0)
