Water info capture now captures from the highest point in the terrain rather than the highest water body.
Water info rendering didn't capture height data for the terrain above the highest water body in the map. This could cause mild dilation issues when observing water from far away.

diff --git a/Engine/Plugins/Experimental/Water/Source/Runtime/Private/WaterZoneActor.cpp b/Engine/Plugins/Experimental/Water/Source/Runtime/Private/WaterZoneActor.cpp
--- a/Engine/Plugins/Experimental/Water/Source/Runtime/Private/WaterZoneActor.cpp
+++ b/Engine/Plugins/Experimental/Water/Source/Runtime/Private/WaterZoneActor.cpp
@@ -268,14 +268,17 @@ bool AWaterZone::UpdateWaterInfoTexture()
 
 		WaterHeightExtents = FVector2f(WaterZMin, WaterZMax);
 
-		// Only compute the ground min since we can use the water max z as the ground max z for more precision.
+		// Terrain height is only computed from [GroundZMin, WaterZMax] however we will still compute the ground z max
+		// to prevent incorrect depth data if there are mountains which exceed the highest water body
 		GroundZMin = TNumericLimits<double>::Max();
+		double GroundZMax = TNumericLimits<double>::Lowest();
 
 		TArray<AActor*> GroundActors;
 		for (ALandscapeProxy* LandscapeProxy : TActorRange<ALandscapeProxy>(World))
 		{
 			const FBox LandscapeBox = LandscapeProxy->GetComponentsBoundingBox();
 			GroundZMin = FMath::Min(GroundZMin, LandscapeBox.Min.Z);
+			GroundZMax = FMath::Max(GroundZMax, LandscapeBox.Max.Z);
 			GroundActors.Add(LandscapeProxy);
 		}
 
@@ -286,7 +289,7 @@ bool AWaterZone::UpdateWaterInfoTexture()
 		Context.ZoneToRender = this;
 		Context.WaterBodies = WaterBodies;
 		Context.GroundActors = MoveTemp(GroundActors);
-		Context.CaptureZ = WaterZMax + CaptureZOffset;
+		Context.CaptureZ = FMath::Max(WaterZMax, GroundZMax) + CaptureZOffset;
 		Context.TextureRenderTarget = WaterInfoTexture;
 
 		if (UWaterSubsystem* WaterSubsystem = UWaterSubsystem::GetWaterSubsystem(World))
