* Fixed assert when iterating over components from a parallel (game) thread
* Made UWaterMeshComponent::ShouldRenderSelected editor-only to avoid doing this at runtime

diff --git a/Engine/Plugins/Experimental/Water/Source/Runtime/Private/WaterMeshComponent.cpp b/Engine/Plugins/Experimental/Water/Source/Runtime/Private/WaterMeshComponent.cpp
--- a/Engine/Plugins/Experimental/Water/Source/Runtime/Private/WaterMeshComponent.cpp
+++ b/Engine/Plugins/Experimental/Water/Source/Runtime/Private/WaterMeshComponent.cpp
@@ -109,21 +109,32 @@ void UWaterMeshComponent::SetMaterial(int32 ElementIndex, UMaterialInterface* Ma
 	UE_LOG(LogWater, Warning, TEXT("SetMaterial is not compatible with UWaterMeshComponent since all materials on this component are auto-populated from the Water Bodies contained within it."));
 }
 
+#if WITH_EDITOR
+
 bool UWaterMeshComponent::ShouldRenderSelected() const
 {
 	if (bSelectable)
 	{
-		for (TActorIterator<AWaterBody> ActorItr(GetWorld()); ActorItr; ++ActorItr)
+		bool bShouldRender = Super::ShouldRenderSelected();
+		if (!bShouldRender)
 		{
-			if (ActorItr->IsSelected())
+			UWaterSubsystem::ForEachWaterBodyComponent(GetWorld(), [&bShouldRender](UWaterBodyComponent* WaterBodyComponent)
 			{
-				return true;
-			}
+				check(WaterBodyComponent);
+				bShouldRender |= WaterBodyComponent->ShouldRenderSelected();
+
+				// Stop iterating over water body components by returning false as soon as one component says it should be "render selected" :
+				return !bShouldRender;
+			});
 		}
+
+		return bShouldRender;
 	}
 	return false;
 }
 
+#endif // WITH_EDITOR
+
 FMaterialRelevance UWaterMeshComponent::GetWaterMaterialRelevance(ERHIFeatureLevel::Type InFeatureLevel) const
 {
 	// Combine the material relevance for all materials.
diff --git a/Engine/Plugins/Experimental/Water/Source/Runtime/Public/WaterMeshComponent.h b/Engine/Plugins/Experimental/Water/Source/Runtime/Public/WaterMeshComponent.h
--- a/Engine/Plugins/Experimental/Water/Source/Runtime/Public/WaterMeshComponent.h
+++ b/Engine/Plugins/Experimental/Water/Source/Runtime/Public/WaterMeshComponent.h
@@ -36,7 +36,9 @@ public:
 	virtual FPrimitiveSceneProxy* CreateSceneProxy() override;
 	virtual void GetUsedMaterials(TArray<UMaterialInterface*>& OutMaterials, bool bGetDebugMaterials = false) const override;
 	virtual void SetMaterial(int32 ElementIndex, UMaterialInterface* Material) override;
+#if WITH_EDITOR
 	virtual bool ShouldRenderSelected() const override;
+#endif // WITH_EDITOR
 	//~ End UPrimitiveComponent Interface
 
 	void Update();
