WaterZoneActor.cpp: Initialize float with float max.

diff --git a/Engine/Plugins/Experimental/Water/Source/Runtime/Private/WaterZoneActor.cpp b/Engine/Plugins/Experimental/Water/Source/Runtime/Private/WaterZoneActor.cpp
--- a/Engine/Plugins/Experimental/Water/Source/Runtime/Private/WaterZoneActor.cpp
+++ b/Engine/Plugins/Experimental/Water/Source/Runtime/Private/WaterZoneActor.cpp
@@ -226,8 +226,8 @@ bool AWaterZone::UpdateWaterInfoTexture()
 {
 	if (UWorld* World = GetWorld(); World && FApp::CanEverRender())
 	{
-		float WaterZMin(TNumericLimits<double>::Max());
-		float WaterZMax(TNumericLimits<double>::Lowest());
+		float WaterZMin(TNumericLimits<float>::Max());
+		float WaterZMax(TNumericLimits<float>::Lowest());
 	
 		bool bHasIncompleteShaderMaps = false;
 		// #todo_water [roey]: we need to store which actors this zone is responsible for rendering once we implement support for multi zones.
@@ -268,10 +268,9 @@ bool AWaterZone::UpdateWaterInfoTexture()
 
 		WaterHeightExtents = FVector2f(WaterZMin, WaterZMax);
 
-		// Terrain height is only computed from [GroundZMin, WaterZMax] however we will still compute the ground z max
-		// to prevent incorrect depth data if there are mountains which exceed the highest water body
-		GroundZMin = TNumericLimits<double>::Max();
-		double GroundZMax = TNumericLimits<double>::Lowest();
+		// Only compute the ground min since we can use the water max z as the ground max z for more precision.
+		GroundZMin = TNumericLimits<float>::Max();
+		float GroundZMax = TNumericLimits<float>::Lowest();
 
 		TArray<AActor*> GroundActors;
 		for (ALandscapeProxy* LandscapeProxy : TActorRange<ALandscapeProxy>(World))
