Initialize async inputs for meatball sim

diff --git a/Engine/Plugins/Experimental/Water/Source/Runtime/Public/BuoyancyComponentSimulation.h b/Engine/Plugins/Experimental/Water/Source/Runtime/Public/BuoyancyComponentSimulation.h
--- a/Engine/Plugins/Experimental/Water/Source/Runtime/Public/BuoyancyComponentSimulation.h
+++ b/Engine/Plugins/Experimental/Water/Source/Runtime/Public/BuoyancyComponentSimulation.h
@@ -52,7 +52,7 @@ struct FBuoyancyComponentBaseAsyncInput : public FBuoyancyComponentAsyncInput
 {
 	TArray<UWaterBodyComponent*> WaterBodyComponents;
 	TArray<FSphericalPontoon> Pontoons;
-	float SmoothedWorldTimeSeconds;
+	float SmoothedWorldTimeSeconds = 0.f;
 
 	FBuoyancyComponentBaseAsyncInput()
 		: FBuoyancyComponentAsyncInput(EAsyncBuoyancyComponentDataType::AsyncBuoyancyBase)
