fix crash on adding spline points to oceans w/ non-zero shape dilation

diff --git a/Engine/Plugins/Experimental/Water/Source/Runtime/Private/WaterBodyLakeComponent.cpp b/Engine/Plugins/Experimental/Water/Source/Runtime/Private/WaterBodyLakeComponent.cpp
--- a/Engine/Plugins/Experimental/Water/Source/Runtime/Private/WaterBodyLakeComponent.cpp
+++ b/Engine/Plugins/Experimental/Water/Source/Runtime/Private/WaterBodyLakeComponent.cpp
@@ -78,7 +78,11 @@ void UWaterBodyLakeComponent::GenerateWaterBodyMesh()
 	FConstrainedDelaunay2d Triangulation;
 	Triangulation.FillRule = FConstrainedDelaunay2d::EFillRule::Positive;
 	Triangulation.Add(LakePoly);
-	Triangulation.Triangulate();
+	bool bTriangulationSuccess = Triangulation.Triangulate();
+	if (!bTriangulationSuccess)
+	{
+		UE_LOG(LogWater, Warning, TEXT("Failed correctly triangulate lake bounding curve; input may have self-intersections (%s"), *GetOwner()->GetActorNameOrLabel());
+	}
 
 	if (Triangulation.Triangles.Num() == 0)
 	{
@@ -111,11 +115,10 @@ void UWaterBodyLakeComponent::GenerateWaterBodyMesh()
 		FInsetMeshRegion Inset(&LakeMesh);
 		Inset.InsetDistance = -1 * ShapeDilation / 2.f;
 
-		for (const FIndex3i& Triangle : LakeMesh.GetTrianglesBuffer())
+		Inset.Triangles.Reserve(LakeMesh.TriangleCount());
+		for (int32 Idx : LakeMesh.TriangleIndicesItr())
 		{
-			Inset.Triangles.Add(Triangle.A);
-			Inset.Triangles.Add(Triangle.B);
-			Inset.Triangles.Add(Triangle.C);
+			Inset.Triangles.Add(Idx);
 		}
 		
 		if (Inset.Apply())
diff --git a/Engine/Plugins/Experimental/Water/Source/Runtime/Private/WaterBodyOceanComponent.cpp b/Engine/Plugins/Experimental/Water/Source/Runtime/Private/WaterBodyOceanComponent.cpp
--- a/Engine/Plugins/Experimental/Water/Source/Runtime/Private/WaterBodyOceanComponent.cpp
+++ b/Engine/Plugins/Experimental/Water/Source/Runtime/Private/WaterBodyOceanComponent.cpp
@@ -155,13 +155,20 @@ void UWaterBodyOceanComponent::GenerateWaterBodyMesh()
 
 	FVector OceanLocation = GetComponentLocation();
 	FPolygon2d OceanBoundingPolygon = FPolygon2d::MakeRectangle(FVector2d(OceanLocation.X, OceanLocation.Y), VisualExtents.X, VisualExtents.Y);
-	FGeneralPolygon2d FinalPoly(OceanBoundingPolygon);
-	FinalPoly.AddHole(Island);
 
 	FConstrainedDelaunay2d Triangulation;
 	Triangulation.FillRule = FConstrainedDelaunay2d::EFillRule::Positive;
-	Triangulation.Add(FinalPoly);
-	Triangulation.Triangulate();
+	Triangulation.Add(OceanBoundingPolygon);
+	if (!Island.IsClockwise())
+	{
+		Island.Reverse();
+	}
+	Triangulation.Add(Island);
+	bool bTriangulationSuccess = Triangulation.Triangulate();
+	if (!bTriangulationSuccess)
+	{
+		UE_LOG(LogWater, Warning, TEXT("Failed correctly triangulate ocean bounding curve; input may have self-intersections (%s"), *GetOwner()->GetActorNameOrLabel());
+	}
 
 	if (Triangulation.Triangles.Num() == 0)
 	{
@@ -196,11 +203,10 @@ void UWaterBodyOceanComponent::GenerateWaterBodyMesh()
 		FInsetMeshRegion Inset(&OceanMesh);
 		Inset.InsetDistance = -1 * ShapeDilation / 2.f;
 
-		for (FIndex3i Triangle : OceanMesh.GetTrianglesBuffer())
+		Inset.Triangles.Reserve(OceanMesh.TriangleCount());
+		for (int32 Idx : OceanMesh.TriangleIndicesItr())
 		{
-			Inset.Triangles.Add(Triangle.A);
-			Inset.Triangles.Add(Triangle.B);
-			Inset.Triangles.Add(Triangle.C);
+			Inset.Triangles.Add(Idx);
 		}
 		
 		if (Inset.Apply())
diff --git a/Engine/Plugins/Runtime/GeometryProcessing/Source/DynamicMesh/Private/EdgeLoop.cpp b/Engine/Plugins/Runtime/GeometryProcessing/Source/DynamicMesh/Private/EdgeLoop.cpp
--- a/Engine/Plugins/Runtime/GeometryProcessing/Source/DynamicMesh/Private/EdgeLoop.cpp
+++ b/Engine/Plugins/Runtime/GeometryProcessing/Source/DynamicMesh/Private/EdgeLoop.cpp
@@ -51,7 +51,7 @@ bool FEdgeLoop::InitializeFromVertices(const TArray<int>& VerticesIn, bool bAuto
 		Edges[i] = Mesh->FindEdge(a, b);
 		if (Edges[i] == FDynamicMesh3::InvalidID)
 		{
-			checkf(false, TEXT("EdgeLoop.FromVertices: invalid edge [%d,%d]"), a, b);
+			ensureMsgf(false, TEXT("EdgeLoop.FromVertices: invalid edge [%d,%d]"), a, b);
 			return false;
 		}
 	}
