CIS fix -- remove newly-added log warning on self-intersecting polygons

diff --git a/Engine/Plugins/Experimental/Water/Source/Runtime/Private/WaterBodyLakeComponent.cpp b/Engine/Plugins/Experimental/Water/Source/Runtime/Private/WaterBodyLakeComponent.cpp
--- a/Engine/Plugins/Experimental/Water/Source/Runtime/Private/WaterBodyLakeComponent.cpp
+++ b/Engine/Plugins/Experimental/Water/Source/Runtime/Private/WaterBodyLakeComponent.cpp
@@ -78,11 +78,7 @@ void UWaterBodyLakeComponent::GenerateWaterBodyMesh()
 	FConstrainedDelaunay2d Triangulation;
 	Triangulation.FillRule = FConstrainedDelaunay2d::EFillRule::Positive;
 	Triangulation.Add(LakePoly);
-	bool bTriangulationSuccess = Triangulation.Triangulate();
-	if (!bTriangulationSuccess)
-	{
-		UE_LOG(LogWater, Warning, TEXT("Failed correctly triangulate lake bounding curve; input may have self-intersections (%s"), *GetOwner()->GetActorNameOrLabel());
-	}
+	Triangulation.Triangulate();
 
 	if (Triangulation.Triangles.Num() == 0)
 	{
diff --git a/Engine/Plugins/Experimental/Water/Source/Runtime/Private/WaterBodyOceanComponent.cpp b/Engine/Plugins/Experimental/Water/Source/Runtime/Private/WaterBodyOceanComponent.cpp
--- a/Engine/Plugins/Experimental/Water/Source/Runtime/Private/WaterBodyOceanComponent.cpp
+++ b/Engine/Plugins/Experimental/Water/Source/Runtime/Private/WaterBodyOceanComponent.cpp
@@ -164,11 +164,7 @@ void UWaterBodyOceanComponent::GenerateWaterBodyMesh()
 		Island.Reverse();
 	}
 	Triangulation.Add(Island);
-	bool bTriangulationSuccess = Triangulation.Triangulate();
-	if (!bTriangulationSuccess)
-	{
-		UE_LOG(LogWater, Warning, TEXT("Failed correctly triangulate ocean bounding curve; input may have self-intersections (%s"), *GetOwner()->GetActorNameOrLabel());
-	}
+	Triangulation.Triangulate();
 
 	if (Triangulation.Triangles.Num() == 0)
 	{
