Updated Water MIDs to LWC

diff --git a/Engine/Plugins/Experimental/Water/Source/Runtime/Private/WaterBodyComponent.cpp b/Engine/Plugins/Experimental/Water/Source/Runtime/Private/WaterBodyComponent.cpp
--- a/Engine/Plugins/Experimental/Water/Source/Runtime/Private/WaterBodyComponent.cpp
+++ b/Engine/Plugins/Experimental/Water/Source/Runtime/Private/WaterBodyComponent.cpp
@@ -1283,12 +1283,12 @@ bool UWaterBodyComponent::SetDynamicParametersOnMID(UMaterialInstanceDynamic* In
 		const FVector2D ZoneExtent = WaterZone->GetZoneExtent();
 		const FVector2D WaterAreaLocation = FVector2D(WaterZone->GetActorLocation()) - (ZoneExtent / 2.f);
 
-		FLinearColor WaterArea;
-		WaterArea.R = WaterAreaLocation.X;
-		WaterArea.G = WaterAreaLocation.Y;
-		WaterArea.B = ZoneExtent.X;
-		WaterArea.A = ZoneExtent.Y;
-		InMID->SetVectorParameterValue(WaterAreaParamName, WaterArea);
+		FVector4 WaterArea;
+		WaterArea.X = WaterAreaLocation.X;
+		WaterArea.Y = WaterAreaLocation.Y;
+		WaterArea.Z = ZoneExtent.X;
+		WaterArea.W = ZoneExtent.Y;
+		InMID->SetDoubleVectorParameterValue(WaterAreaParamName, WaterArea);
 
 		const FVector2f WaterHeightExtents = WaterZone->GetWaterHeightExtents();
 		const float GroundZMin = WaterZone->GetGroundZMin();
