Possible fix for water crash

diff --git a/Engine/Plugins/Experimental/Water/Source/Runtime/Private/WaterInfoRendering.cpp b/Engine/Plugins/Experimental/Water/Source/Runtime/Private/WaterInfoRendering.cpp
--- a/Engine/Plugins/Experimental/Water/Source/Runtime/Private/WaterInfoRendering.cpp
+++ b/Engine/Plugins/Experimental/Water/Source/Runtime/Private/WaterInfoRendering.cpp
@@ -482,14 +482,17 @@ static FSceneRenderer* CreateWaterInfoDepthRenderer(
 	{
 		DepthView->ShowOnlyPrimitives.Emplace();
 		DepthView->ShowOnlyPrimitives->Reserve(Context.GroundActors.Num());
-		for (AActor* GroundActor : Context.GroundActors)
+		for (TWeakObjectPtr<AActor> GroundActor : Context.GroundActors)
 		{
-			TInlineComponentArray<UPrimitiveComponent*> PrimComps(GroundActor);
-			for (UPrimitiveComponent* PrimComp : PrimComps)
+			if (GroundActor.IsValid())
 			{
-				if (PrimComp)
+				TInlineComponentArray<UPrimitiveComponent*> PrimComps(GroundActor.Get());
+				for (UPrimitiveComponent* PrimComp : PrimComps)
 				{
-					DepthView->ShowOnlyPrimitives->Add(PrimComp->ComponentId);
+					if (PrimComp)
+					{
+						DepthView->ShowOnlyPrimitives->Add(PrimComp->ComponentId);
+					}
 				}
 			}
 		}
diff --git a/Engine/Plugins/Experimental/Water/Source/Runtime/Private/WaterInfoRendering.h b/Engine/Plugins/Experimental/Water/Source/Runtime/Private/WaterInfoRendering.h
--- a/Engine/Plugins/Experimental/Water/Source/Runtime/Private/WaterInfoRendering.h
+++ b/Engine/Plugins/Experimental/Water/Source/Runtime/Private/WaterInfoRendering.h
@@ -17,7 +17,7 @@ namespace UE::WaterInfo
 		AWaterZone* ZoneToRender = nullptr;
 		UTextureRenderTarget2D* TextureRenderTarget;
 		TArray<UWaterBodyComponent*> WaterBodies;
-		TArray<AActor*> GroundActors;
+		TArray<TWeakObjectPtr<AActor>> GroundActors;
 		float CaptureZ;
 	};
 	
diff --git a/Engine/Plugins/Experimental/Water/Source/Runtime/Private/WaterZoneActor.cpp b/Engine/Plugins/Experimental/Water/Source/Runtime/Private/WaterZoneActor.cpp
--- a/Engine/Plugins/Experimental/Water/Source/Runtime/Private/WaterZoneActor.cpp
+++ b/Engine/Plugins/Experimental/Water/Source/Runtime/Private/WaterZoneActor.cpp
@@ -272,7 +272,7 @@ bool AWaterZone::UpdateWaterInfoTexture()
 		GroundZMin = TNumericLimits<float>::Max();
 		float GroundZMax = TNumericLimits<float>::Lowest();
 
-		TArray<AActor*> GroundActors;
+		TArray<TWeakObjectPtr<AActor>> GroundActors;
 		for (ALandscapeProxy* LandscapeProxy : TActorRange<ALandscapeProxy>(World))
 		{
 			const FBox LandscapeBox = LandscapeProxy->GetComponentsBoundingBox();
