Initial support for multiple water zones.
- Water subsystem no longer recognizes a single global water zone
- Assumptions that a single water zone exists have been removed
- Water bodies associate themselves with the first overlapping water zone within their outered level (temporary)
- WaterZones allow iterating only the water bodies associated with that zone

diff --git a/Engine/Plugins/Experimental/Water/Source/Runtime/Private/WaterBodyComponent.cpp b/Engine/Plugins/Experimental/Water/Source/Runtime/Private/WaterBodyComponent.cpp
--- a/Engine/Plugins/Experimental/Water/Source/Runtime/Private/WaterBodyComponent.cpp
+++ b/Engine/Plugins/Experimental/Water/Source/Runtime/Private/WaterBodyComponent.cpp
@@ -825,12 +825,9 @@ void UWaterBodyComponent::UpdateComponentVisibility(bool bAllowWaterMeshRebuild)
 		// If the component is being or can be rendered by the water mesh, rebuild it in case its visibility has changed : 
 		if (bAllowWaterMeshRebuild && (GetWaterBodyType() != EWaterBodyType::Transition))
 		{
-			if (UWaterSubsystem* WaterSubsystem = UWaterSubsystem::GetWaterSubsystem(GetWorld())) 
+			if (AWaterZone* WaterZone = GetWaterZone())
 			{
-				if (AWaterZone* WaterZone = WaterSubsystem->GetWaterZoneActor())
-				{
-					WaterZone->MarkForRebuild(EWaterZoneRebuildFlags::All);
-				}
+				WaterZone->MarkForRebuild(EWaterZoneRebuildFlags::All);
 			}
 		}
 	}
@@ -902,24 +899,21 @@ TArray<TSharedRef<FTokenizedMessage>> UWaterBodyComponent::CheckWaterBodyStatus(
 	TArray<TSharedRef<FTokenizedMessage>> Result;
 	if (!IsTemplate())
 	{
-		if (const UWaterSubsystem* WaterSubsystem = UWaterSubsystem::GetWaterSubsystem(GetWorld()))
+		if (AffectsWaterMesh() && (GetWaterZone() == nullptr))
 		{
-			if (AffectsWaterMesh() && (WaterSubsystem->GetWaterZoneActor() == nullptr))
-			{
-				Result.Add(FTokenizedMessage::Create(EMessageSeverity::Error)
-					->AddToken(FUObjectToken::Create(this))
-					->AddToken(FTextToken::Create(FText::Format(
-						LOCTEXT("MapCheck_Message_MissingWaterZone", "Water body {0} requires a WaterZone actor to be rendered. Please add one to the map. "),
-						FText::FromString(GetWaterBodyActor()->GetActorLabel())))));
-				}
-			}
+			Result.Add(FTokenizedMessage::Create(EMessageSeverity::Error)
+				->AddToken(FUObjectToken::Create(this))
+				->AddToken(FTextToken::Create(FText::Format(
+					LOCTEXT("MapCheck_Message_MissingWaterZone", "Water body {0} requires a WaterZone actor to be rendered. Please add one to the map. "),
+					FText::FromString(GetWaterBodyActor()->GetActorLabel())))));
+		}
 
 		if (AffectsLandscape() && (FindLandscape() == nullptr))
-			{
+		{
 			Result.Add(FTokenizedMessage::Create(EMessageSeverity::Error)
 				->AddToken(FUObjectToken::Create(this))
 				->AddToken(FTextToken::Create(FText::Format(
-						LOCTEXT("MapCheck_Message_MissingLandscape", "Water body {0} requires a Landscape to be rendered. Please add one to the map. "),
+					LOCTEXT("MapCheck_Message_MissingLandscape", "Water body {0} requires a Landscape to be rendered. Please add one to the map. "),
 					FText::FromString(GetWaterBodyActor()->GetActorLabel())))));
 		}
 	}
@@ -1272,7 +1266,7 @@ bool UWaterBodyComponent::SetDynamicParametersOnMID(UMaterialInstanceDynamic* In
 	InMID->SetVectorParameterValue(FixedVelocityName, GetConstantVelocity());
 
 	// Use WaterZone actor of the same level
-	if (const AWaterZone* WaterZone = WaterSubsystem->GetWaterZoneActor(GetTypedOuter<ULevel>()))
+	if (const AWaterZone* WaterZone = GetWaterZone())
 	{
 		InMID->SetTextureParameterValue(WaterVelocityAndHeightName, WaterZone->WaterInfoTexture);
 
@@ -1423,13 +1417,36 @@ UWaterWavesBase* UWaterBodyComponent::GetWaterWaves() const
 
 AWaterZone* UWaterBodyComponent::GetWaterZone() const
 {
-	// #todo_water [roey]: Currently returns the global water zone. Eventually this function will return the specific water zone which encapsulates this water body. 
+	TRACE_CPUPROFILER_EVENT_SCOPE(UWaterBodyComponent::GetWaterZone);
 
-	if (UWaterSubsystem* WaterSubsystem = UWaterSubsystem::GetWaterSubsystem(GetWorld()))
+	TArray<AWaterZone*, TInlineAllocator<4>> ViableZones;
+	
+	if (const UWorld* World = GetWorld())
 	{
-		return WaterSubsystem->GetWaterZoneActor();
+		const ULevel* PreferredLevel = GetTypedOuter<ULevel>();
+		if (PreferredLevel)
+		{
+			for (AWaterZone* WaterZone : TActorRange<AWaterZone>(World, AWaterZone::StaticClass(), EActorIteratorFlags::SkipPendingKill))
+			{
+				// WaterZone->GetZoneExtents returns the full extent of the zone but BoxSphereBounds expects a half-extent.
+				const FBoxSphereBounds WaterZoneBounds(WaterZone->GetActorLocation(), FVector(WaterZone->GetZoneExtent() / 2.f, 8192.f), 0.f);
+				const FBoxSphereBounds ComponentBounds = CalcBounds(GetComponentTransform());
+
+				// Only consider WaterZones which this component overlaps but prefer choosing water zones
+				// which are part of the same outered level.
+				if (FBoxSphereBounds::BoxesIntersect(ComponentBounds, WaterZoneBounds))
+				{
+					if (WaterZone->GetTypedOuter<ULevel>() == PreferredLevel)
+					{
+						return WaterZone;
+					}
+
+					ViableZones.Add(WaterZone);
+				}
+			}
+		}
 	}
-	return nullptr;
+	return ViableZones.Num() > 0 ? ViableZones[0] : nullptr;
 }
 
 void UWaterBodyComponent::UpdateWaterBodyRenderData()
diff --git a/Engine/Plugins/Experimental/Water/Source/Runtime/Private/WaterBodyOceanComponent.cpp b/Engine/Plugins/Experimental/Water/Source/Runtime/Private/WaterBodyOceanComponent.cpp
--- a/Engine/Plugins/Experimental/Water/Source/Runtime/Private/WaterBodyOceanComponent.cpp
+++ b/Engine/Plugins/Experimental/Water/Source/Runtime/Private/WaterBodyOceanComponent.cpp
@@ -236,18 +236,9 @@ void UWaterBodyOceanComponent::PostLoad()
 
 	if (GetLinkerCustomVersion(FFortniteMainBranchObjectVersion::GUID) < FFortniteMainBranchObjectVersion::WaterZonesRefactor)
 	{
-		if (UWorld* World = GetWorld())
+		if (AWaterZone* WaterZone = GetWaterZone())
 		{
-			// #todo_water: this assumes only one water zone actor right now.  In the future we may need to associate a water mesh actor with a water body more directly
-			for (TObjectIterator<AWaterZone> It; It; ++It)
-			{
-				AWaterZone* WaterZone = *It;
-				if (WaterZone && (WaterZone->GetWorld() == World))
-				{
-					SetVisualExtents(WaterZone->GetZoneExtent());
-					return;
-				}
-			}
+			SetVisualExtents(WaterZone->GetZoneExtent());
 		}
 	}
 }
diff --git a/Engine/Plugins/Experimental/Water/Source/Runtime/Private/WaterMeshComponent.cpp b/Engine/Plugins/Experimental/Water/Source/Runtime/Private/WaterMeshComponent.cpp
--- a/Engine/Plugins/Experimental/Water/Source/Runtime/Private/WaterMeshComponent.cpp
+++ b/Engine/Plugins/Experimental/Water/Source/Runtime/Private/WaterMeshComponent.cpp
@@ -118,14 +118,17 @@ bool UWaterMeshComponent::ShouldRenderSelected() const
 		bool bShouldRender = Super::ShouldRenderSelected();
 		if (!bShouldRender)
 		{
-			UWaterSubsystem::ForEachWaterBodyComponent(GetWorld(), [&bShouldRender](UWaterBodyComponent* WaterBodyComponent)
+			if (AWaterZone* Owner = GetOwner<AWaterZone>())
 			{
-				check(WaterBodyComponent);
-				bShouldRender |= WaterBodyComponent->ShouldRenderSelected();
+				Owner->ForEachWaterBodyComponent([&bShouldRender](UWaterBodyComponent* WaterBodyComponent)
+				{
+					check(WaterBodyComponent);
+					bShouldRender |= WaterBodyComponent->ShouldRenderSelected();
 
-				// Stop iterating over water body components by returning false as soon as one component says it should be "render selected" :
-				return !bShouldRender;
-			});
+					// Stop iterating over water body components by returning false as soon as one component says it should be "render selected" :
+					return !bShouldRender;
+				});
+			}
 		}
 
 		return bShouldRender;
@@ -209,7 +212,9 @@ void UWaterMeshComponent::RebuildWaterMesh(float InTileSize, const FIntPoint& In
 	const bool bIsFlooded = OceanFlood > 0.0f;
 
 	// Go through all water body actors to figure out bounds and water tiles
-	UWaterSubsystem::ForEachWaterBodyComponent(GetWorld(), [this, WaterWorldBox, bIsFlooded, GlobalOceanHeight, OceanFlood, &FarMeshHeight](UWaterBodyComponent* WaterBodyComponent)
+	AWaterZone* OwningZone = GetOwner<AWaterZone>();
+	check(OwningZone);
+	OwningZone->ForEachWaterBodyComponent([this, WaterWorldBox, bIsFlooded, GlobalOceanHeight, OceanFlood, &FarMeshHeight](UWaterBodyComponent* WaterBodyComponent)
 	{
 		check(WaterBodyComponent);
 		AActor* Actor = WaterBodyComponent->GetOwner();
diff --git a/Engine/Plugins/Experimental/Water/Source/Runtime/Private/WaterSubsystem.cpp b/Engine/Plugins/Experimental/Water/Source/Runtime/Private/WaterSubsystem.cpp
--- a/Engine/Plugins/Experimental/Water/Source/Runtime/Private/WaterSubsystem.cpp
+++ b/Engine/Plugins/Experimental/Water/Source/Runtime/Private/WaterSubsystem.cpp
@@ -170,11 +170,11 @@ UWaterSubsystem* UWaterSubsystem::GetWaterSubsystem(const UWorld* InWorld)
 	return nullptr;
 }
 
-FWaterBodyManager* UWaterSubsystem::GetWaterBodyManager(UWorld* InWorld)
+FWaterBodyManager* UWaterSubsystem::GetWaterBodyManager(const UWorld* InWorld)
 {
 	if (UWaterSubsystem* Subsystem = GetWaterSubsystem(InWorld))
 	{
-		return &Subsystem->WaterBodyManager;
+		return &Subsystem->GetWaterBodyManagerInternal();
 	}
 
 	return nullptr;
@@ -184,7 +184,7 @@ void UWaterSubsystem::ForEachWaterBodyComponent(const UWorld* World, TFunctionRe
 {
 	if (UWaterSubsystem* Subsystem = GetWaterSubsystem(World))
 	{
-		Subsystem->WaterBodyManager.ForEachWaterBodyComponent(Predicate);
+		Subsystem->GetWaterBodyManagerInternal().ForEachWaterBodyComponent(Predicate);
 	}
 }
 
@@ -202,9 +202,12 @@ void UWaterSubsystem::Tick(float DeltaTime)
 	SetMPCTime(MPCTime, PrevWorldTimeSeconds);
 	PrevWorldTimeSeconds = MPCTime;
 
-	if (AWaterZone* LocalWaterZoneActor = GetWaterZoneActor())
+	for (AWaterZone* WaterZoneActor : TActorRange<AWaterZone>(GetWorld()))
 	{
-		LocalWaterZoneActor->Update();
+		if (WaterZoneActor)
+		{
+			WaterZoneActor->Update();
+		}
 	}
 
 	if (!bUnderWaterForAudio && CachedDepthUnderwater > 0.0f)
@@ -244,7 +247,7 @@ void UWaterSubsystem::Initialize(FSubsystemCollectionBase& Collection)
 	UWorld* World = GetWorld();
 	check(World != nullptr);
 
-	WaterBodyManager.Initialize(World);
+	GetWaterBodyManagerInternal().Initialize(World);
 	WaterViewExtension = FSceneViewExtensions::NewExtension<FWaterViewExtension>(World);
 
 	bUsingSmoothedTime = false;
@@ -298,7 +301,7 @@ void UWaterSubsystem::Deinitialize()
 	World->OnBeginPostProcessSettings.RemoveAll(this);
 	World->RemovePostProcessVolume(&UnderwaterPostProcessVolume);
 
-	WaterBodyManager.Deinitialize();
+	GetWaterBodyManagerInternal().Deinitialize();
 
 #if WITH_EDITOR
 	GetDefault<UWaterRuntimeSettings>()->OnSettingsChange.RemoveAll(this);
@@ -317,7 +320,7 @@ void UWaterSubsystem::ApplyRuntimeSettings(const UWaterRuntimeSettings* Settings
 #if WITH_EDITOR
 	// Update sprites since we may have changed the sprite Z offset setting.
 
-	ForEachWaterBodyComponent(World, [](UWaterBodyComponent* Component)
+	GetWaterBodyManagerInternal().ForEachWaterBodyComponent([](UWaterBodyComponent* Component)
 	{
 		Component->UpdateWaterSpriteComponent();
 		return true;
@@ -335,6 +338,14 @@ void UWaterSubsystem::ApplyRuntimeSettings(const UWaterRuntimeSettings* Settings
 #endif // WITH_EDITOR
 }
 
+FWaterBodyManager& UWaterSubsystem::GetWaterBodyManagerInternal()
+{
+	// #todo_water [roey]: Remove these pragmas when we can move WaterBodyManager to private.
+	PRAGMA_DISABLE_DEPRECATION_WARNINGS
+	return WaterBodyManager;
+	PRAGMA_ENABLE_DEPRECATION_WARNINGS
+}
+
 void UWaterSubsystem::AddWaterCollisionProfile()
 {
 	// Make sure WaterCollisionProfileName is added to Engine's collision profiles
@@ -474,7 +485,7 @@ void UWaterSubsystem::SetOceanFloodHeight(float InFloodHeight)
 				OceanBodyComponent->SetHeightOffset(InFloodHeight);
 			}
 
-			WaterBodyManager.ForEachWaterBodyComponent([this](UWaterBodyComponent* WaterBodyComponent)
+			GetWaterBodyManagerInternal().ForEachWaterBodyComponent([this](UWaterBodyComponent* WaterBodyComponent)
 			{
 				WaterBodyComponent->UpdateMaterialInstances();
 				return true;
@@ -483,29 +494,6 @@ void UWaterSubsystem::SetOceanFloodHeight(float InFloodHeight)
 	}
 }
 
-AWaterZone* UWaterSubsystem::GetWaterZoneActor(ULevel* InPreferredOuterLevel) const
-{
-	if (UWorld* World = GetWorld())
-	{
-		// First, try to find the AWaterZone outered to the provided level
-		if (InPreferredOuterLevel)
-		{
-			for (TActorIterator<AWaterZone> It(World); It; ++It)
-			{
-				if (It && (It->GetTypedOuter<ULevel>() == InPreferredOuterLevel))
-				{
-					return *It;
-				}
-			}
-		}
-		// #todo_water: this assumes only one water zone actor right now.  In the future we may need to associate a water mesh actor with a water body more directly
-		TActorIterator<AWaterZone> It(World);
-		return It ? *It : nullptr;
-	}
-
-	return nullptr;
-}
-
 float UWaterSubsystem::GetOceanBaseHeight() const
 {
 	if (OceanBodyComponent.IsValid())
@@ -540,7 +528,7 @@ void UWaterSubsystem::NotifyWaterScalabilityChangedInternal(IConsoleVariable* CV
 void UWaterSubsystem::NotifyWaterVisibilityChangedInternal(IConsoleVariable* CVar)
 {
 	// Water body visibility depends on various CVars. All need to update the visibility in water body components : 
-	WaterBodyManager.ForEachWaterBodyComponent([](UWaterBodyComponent* WaterBodyComponent)
+	GetWaterBodyManagerInternal().ForEachWaterBodyComponent([](UWaterBodyComponent* WaterBodyComponent)
 	{
 		WaterBodyComponent->UpdateComponentVisibility(/* bAllowWaterMeshRebuild = */true);
 		return true;
@@ -606,8 +594,8 @@ void UWaterSubsystem::ComputeUnderwaterPostProcess(FVector ViewLocation, FSceneV
 
 	TArray<FHitResult> Hits;
 	TArray<FWaterBodyPostProcessQuery, TInlineAllocator<4>> WaterBodyQueriesToProcess;
-	const AWaterZone* LocalWaterZoneActor = GetWaterZoneActor();
-	if ((LocalWaterZoneActor != nullptr) && World->SweepMultiByChannel(Hits, ViewLocation, ViewLocation + FVector(0, 0, TraceDistance), FQuat::Identity, UnderwaterTraceChannel, FCollisionShape::MakeSphere(TraceDistance), TraceSimple))
+	const bool bWorldHasWater  = GetWaterBodyManagerInternal().HasAnyWaterBodies();
+	if (bWorldHasWater && World->SweepMultiByChannel(Hits, ViewLocation, ViewLocation + FVector(0, 0, TraceDistance), FQuat::Identity, UnderwaterTraceChannel, FCollisionShape::MakeSphere(TraceDistance), TraceSimple))
 	{
 		if (Hits.Num() > 1)
 		{
diff --git a/Engine/Plugins/Experimental/Water/Source/Runtime/Private/WaterZoneActor.cpp b/Engine/Plugins/Experimental/Water/Source/Runtime/Private/WaterZoneActor.cpp
--- a/Engine/Plugins/Experimental/Water/Source/Runtime/Private/WaterZoneActor.cpp
+++ b/Engine/Plugins/Experimental/Water/Source/Runtime/Private/WaterZoneActor.cpp
@@ -92,13 +92,9 @@ void AWaterZone::PostLoad()
 	Super::PostLoad();
 
 #if WITH_EDITORONLY_DATA
-	if (GetLinkerCustomVersion(FFortniteMainBranchObjectVersion::GUID) < FFortniteMainBranchObjectVersion::WaterZonesRefactor)
-	{
-		// Divide by two since the ZoneExtent represents the radius of the box while the result of the Extent * Size gives the total size
-		const FVector2D ExtentInTiles = WaterMesh->GetExtentInTiles();
-		ZoneExtent = FVector2D(ExtentInTiles * WaterMesh->GetTileSize());
-		OnExtentChanged();
-	}
+	const FVector2D ExtentInTiles = WaterMesh->GetExtentInTiles();
+	ZoneExtent = FVector2D(ExtentInTiles * WaterMesh->GetTileSize());
+	OnExtentChanged();
 #endif // WITH_EDITORONLY_DATA
 }
 
@@ -114,6 +110,18 @@ void AWaterZone::MarkForRebuild(EWaterZoneRebuildFlags Flags)
 	}
 }
 
+void AWaterZone::ForEachWaterBodyComponent(TFunctionRef<bool(UWaterBodyComponent*)> Predicate)
+{
+	UWaterSubsystem::ForEachWaterBodyComponent(GetWorld(), [this, Predicate](UWaterBodyComponent* Component)
+		{
+			if (Component->GetWaterZone() == this)
+			{
+				return Predicate(Component);
+			}
+			return true;
+		});
+}
+
 void AWaterZone::Update()
 {
 	if (bNeedsWaterInfoRebuild || (ForceUpdateWaterInfoNextFrames != 0))
@@ -230,10 +238,10 @@ bool AWaterZone::UpdateWaterInfoTexture()
 		float WaterZMax(TNumericLimits<float>::Lowest());
 	
 		bool bHasIncompleteShaderMaps = false;
-		// #todo_water [roey]: we need to store which actors this zone is responsible for rendering once we implement support for multi zones.
-		// For now whenever we update the water info texture we will just collect all water bodies and pass those to the renderer.
+		// #todo_water [roey]: we should try caching this list to avoid potentially iterating over a lot of water bodies which may not belong to this zone specifically.
+		// For now whenever we update the water info texture we will collect all water bodies within the zone and pass those to the renderer each time this function is called.
 		TArray<UWaterBodyComponent*> WaterBodies;
-		UWaterSubsystem::ForEachWaterBodyComponent(World, [World, &WaterBodies, &WaterZMax, &WaterZMin, &bHasIncompleteShaderMaps](UWaterBodyComponent* Component)
+		ForEachWaterBodyComponent([World, &WaterBodies, &WaterZMax, &WaterZMin, &bHasIncompleteShaderMaps](UWaterBodyComponent* Component)
 		{
 			if (UMaterialInterface* WaterInfoMaterial = Component->GetWaterInfoMaterialInstance())
 			{
diff --git a/Engine/Plugins/Experimental/Water/Source/Runtime/Public/WaterBodyManager.h b/Engine/Plugins/Experimental/Water/Source/Runtime/Public/WaterBodyManager.h
--- a/Engine/Plugins/Experimental/Water/Source/Runtime/Public/WaterBodyManager.h
+++ b/Engine/Plugins/Experimental/Water/Source/Runtime/Public/WaterBodyManager.h
@@ -32,6 +32,8 @@ public:
 	/** Execute a predicate function on each valid water body. Predicate should return false for early exit. */
 	void ForEachWaterBodyComponent (TFunctionRef<bool(UWaterBodyComponent*)> Pred) const;
 
+	bool HasAnyWaterBodies() const { return WaterBodyComponents.Num() > 0; }
+
 private:
 	/** List of components registered to this manager. May contain nullptr indices (indicated by the UnusedWaterBodyIndices array). */
 	TArray<UWaterBodyComponent*> WaterBodyComponents;
diff --git a/Engine/Plugins/Experimental/Water/Source/Runtime/Public/WaterSubsystem.h b/Engine/Plugins/Experimental/Water/Source/Runtime/Public/WaterSubsystem.h
--- a/Engine/Plugins/Experimental/Water/Source/Runtime/Public/WaterSubsystem.h
+++ b/Engine/Plugins/Experimental/Water/Source/Runtime/Public/WaterSubsystem.h
@@ -97,14 +97,16 @@ public:
 	static UWaterSubsystem* GetWaterSubsystem(const UWorld* InWorld);
 
 	/** Static helper function to get a waterbody manager from a world, returns nullptr if world or manager don't exist */
-	static FWaterBodyManager* GetWaterBodyManager(UWorld* InWorld);
+	static FWaterBodyManager* GetWaterBodyManager(const UWorld* InWorld);
 
 	/** Execute a predicate function on each valid water body. Predicate should return false for early exit. */
 	static void ForEachWaterBodyComponent(const UWorld* World, TFunctionRef<bool(UWaterBodyComponent*)> Predicate);
 
+	UE_DEPRECATED(5.1, "This will become a private member. Prefer calling GetWaterBodyManager instead")
 	FWaterBodyManager WaterBodyManager;
 
-	AWaterZone* GetWaterZoneActor(ULevel* InPreferredOuterLevel = nullptr) const;
+	UE_DEPRECATED(5.1, "There may be multiple water zones per level. Prefer calling the GetWaterZone on a per-water body basis or iterating over all zones.")
+	AWaterZone* GetWaterZoneActor(ULevel* InPreferredOuterLevel = nullptr) const { return nullptr; }
 
 	ABuoyancyManager* GetBuoyancyManager() const { return BuoyancyManager; }
 
@@ -201,6 +203,8 @@ private:
 	void AddWaterCollisionProfile();
 	void ApplyRuntimeSettings(const UWaterRuntimeSettings* Settings, EPropertyChangeType::Type ChangeType);
 
+	FWaterBodyManager& GetWaterBodyManagerInternal();
+
 #if !(UE_BUILD_SHIPPING || UE_BUILD_TEST)
 	void ShowOnScreenDebugInfo(const FVector& InViewLocation, const FUnderwaterPostProcessDebugInfo& InDebugInfo);
 #endif // !(UE_BUILD_SHIPPING || UE_BUILD_TEST)
diff --git a/Engine/Plugins/Experimental/Water/Source/Runtime/Public/WaterZoneActor.h b/Engine/Plugins/Experimental/Water/Source/Runtime/Public/WaterZoneActor.h
--- a/Engine/Plugins/Experimental/Water/Source/Runtime/Public/WaterZoneActor.h
+++ b/Engine/Plugins/Experimental/Water/Source/Runtime/Public/WaterZoneActor.h
@@ -11,6 +11,7 @@
 class UWaterMeshComponent;
 class UBoxComponent;
 class AWaterBody;
+class UWaterBodyComponent;
 
 enum class EWaterZoneRebuildFlags
 {
@@ -31,6 +32,9 @@ public:
 
 	void MarkForRebuild(EWaterZoneRebuildFlags Flags);
 	void Update();
+		
+	/** Execute a predicate function on each valid water body within the water zone. Predicate should return false for early exit. */
+	void ForEachWaterBodyComponent(TFunctionRef<bool(UWaterBodyComponent*)> Predicate);
 
 	FVector2D GetZoneExtent() const { return ZoneExtent; }
 	void SetZoneExtent(FVector2D NewExtents);
