Deprecated UWaterSubsystem::ForEachWaterBodyComponent in favor of FWaterBodyManager::ForEachWaterBodyComponent.

diff --git a/Engine/Plugins/Experimental/Water/Source/Runtime/Private/EnvQueryTest_InsideWaterBody.cpp b/Engine/Plugins/Experimental/Water/Source/Runtime/Private/EnvQueryTest_InsideWaterBody.cpp
--- a/Engine/Plugins/Experimental/Water/Source/Runtime/Private/EnvQueryTest_InsideWaterBody.cpp
+++ b/Engine/Plugins/Experimental/Water/Source/Runtime/Private/EnvQueryTest_InsideWaterBody.cpp
@@ -26,7 +26,7 @@ void UEnvQueryTest_InsideWaterBody::RunTest(FEnvQueryInstance& QueryInstance) co
 		const FVector ItemLocation = GetItemLocation(QueryInstance, It.GetIndex());
 
 		bool bInside = false;
-		UWaterSubsystem::ForEachWaterBodyComponent(GetWorld(), [this, ItemLocation, &bInside](UWaterBodyComponent* WaterBodyComponent)
+		FWaterBodyManager::ForEachWaterBodyComponent(GetWorld(), [this, ItemLocation, &bInside](UWaterBodyComponent* WaterBodyComponent)
 		{
 			EWaterBodyQueryFlags QueryFlags = EWaterBodyQueryFlags::ComputeImmersionDepth;
 			if (bIncludeWaves)
diff --git a/Engine/Plugins/Experimental/Water/Source/Runtime/Private/GerstnerWaterWaveViewExtension.cpp b/Engine/Plugins/Experimental/Water/Source/Runtime/Private/GerstnerWaterWaveViewExtension.cpp
--- a/Engine/Plugins/Experimental/Water/Source/Runtime/Private/GerstnerWaterWaveViewExtension.cpp
+++ b/Engine/Plugins/Experimental/Water/Source/Runtime/Private/GerstnerWaterWaveViewExtension.cpp
@@ -47,7 +47,7 @@ void FGerstnerWaterWaveViewExtension::SetupViewFamily(FSceneViewFamily& InViewFa
 		const TWeakObjectPtr<UWorld> WorldPtr = GetWorld();
 		check(WorldPtr.IsValid())
 
-		UWaterSubsystem::ForEachWaterBodyComponent(WorldPtr.Get(), [&WaterIndirectionBuffer, &WaterDataBuffer](UWaterBodyComponent* WaterBodyComponent)
+		FWaterBodyManager::ForEachWaterBodyComponent(WorldPtr.Get(), [&WaterIndirectionBuffer, &WaterDataBuffer](UWaterBodyComponent* WaterBodyComponent)
 		{
 			// Some max value
 			constexpr int32 MaxWavesPerWaterBody = 4096;
diff --git a/Engine/Plugins/Experimental/Water/Source/Runtime/Private/WaterBodyExclusionVolume.cpp b/Engine/Plugins/Experimental/Water/Source/Runtime/Private/WaterBodyExclusionVolume.cpp
--- a/Engine/Plugins/Experimental/Water/Source/Runtime/Private/WaterBodyExclusionVolume.cpp
+++ b/Engine/Plugins/Experimental/Water/Source/Runtime/Private/WaterBodyExclusionVolume.cpp
@@ -38,7 +38,7 @@ void AWaterBodyExclusionVolume::UpdateOverlappingWaterBodies()
 	TLazyObjectPtr<AWaterBodyExclusionVolume> LazyThis(this);
 
 	// Fixup overlapping bodies (iterating on actors on post-load will fail, but this is fine as this exclusion volume should not yet be referenced by an existing water body upon loading) : 
-	UWaterSubsystem::ForEachWaterBodyComponent(GetWorld(), [LazyThis, &ExistingOverlappingBodies](UWaterBodyComponent* WaterBodyComponent)
+	FWaterBodyManager::ForEachWaterBodyComponent(GetWorld(), [LazyThis, &ExistingOverlappingBodies](UWaterBodyComponent* WaterBodyComponent)
 	{
 		if (WaterBodyComponent->ContainsExclusionVolume(LazyThis))
 		{
diff --git a/Engine/Plugins/Experimental/Water/Source/Runtime/Private/WaterBodyIslandActor.cpp b/Engine/Plugins/Experimental/Water/Source/Runtime/Private/WaterBodyIslandActor.cpp
--- a/Engine/Plugins/Experimental/Water/Source/Runtime/Private/WaterBodyIslandActor.cpp
+++ b/Engine/Plugins/Experimental/Water/Source/Runtime/Private/WaterBodyIslandActor.cpp
@@ -184,7 +184,7 @@ void AWaterBodyIsland::UpdateOverlappingWaterBodyComponents()
 	TLazyObjectPtr<AWaterBodyIsland> LazyThis(this);
 
 	// Fixup overlapping bodies
-	UWaterSubsystem::ForEachWaterBodyComponent(GetWorld(), [LazyThis, &ExistingOverlappingBodies](UWaterBodyComponent* WaterBodyComponent)
+	FWaterBodyManager::ForEachWaterBodyComponent(GetWorld(), [LazyThis, &ExistingOverlappingBodies](UWaterBodyComponent* WaterBodyComponent)
 	{
 		if (WaterBodyComponent->ContainsIsland(LazyThis))
 		{
diff --git a/Engine/Plugins/Experimental/Water/Source/Runtime/Private/WaterBodyManager.cpp b/Engine/Plugins/Experimental/Water/Source/Runtime/Private/WaterBodyManager.cpp
--- a/Engine/Plugins/Experimental/Water/Source/Runtime/Private/WaterBodyManager.cpp
+++ b/Engine/Plugins/Experimental/Water/Source/Runtime/Private/WaterBodyManager.cpp
@@ -2,6 +2,7 @@
 
 #include "WaterBodyManager.h"
 #include "WaterBodyActor.h"
+#include "WaterSubsystem.h"
 #include "GerstnerWaterWaveViewExtension.h"
 
 void FWaterBodyManager::Initialize(UWorld* World)
@@ -86,4 +87,12 @@ void FWaterBodyManager::ForEachWaterBodyComponent(TFunctionRef<bool(UWaterBodyCo
 			}
 		}
 	}
+}
+
+void FWaterBodyManager::ForEachWaterBodyComponent(const UWorld* World, TFunctionRef<bool(UWaterBodyComponent*)> Pred)
+{
+	if (FWaterBodyManager* Manager = UWaterSubsystem::GetWaterBodyManager(World))
+	{
+		Manager->ForEachWaterBodyComponent(Pred);
+	}
 }
\ No newline at end of file
diff --git a/Engine/Plugins/Experimental/Water/Source/Runtime/Private/WaterSubsystem.cpp b/Engine/Plugins/Experimental/Water/Source/Runtime/Private/WaterSubsystem.cpp
--- a/Engine/Plugins/Experimental/Water/Source/Runtime/Private/WaterSubsystem.cpp
+++ b/Engine/Plugins/Experimental/Water/Source/Runtime/Private/WaterSubsystem.cpp
@@ -180,14 +180,6 @@ FWaterBodyManager* UWaterSubsystem::GetWaterBodyManager(const UWorld* InWorld)
 	return nullptr;
 }
 
-void UWaterSubsystem::ForEachWaterBodyComponent(const UWorld* World, TFunctionRef<bool(UWaterBodyComponent*)> Predicate)
-{
-	if (UWaterSubsystem* Subsystem = GetWaterSubsystem(World))
-	{
-		Subsystem->GetWaterBodyManagerInternal().ForEachWaterBodyComponent(Predicate);
-	}
-}
-
 void UWaterSubsystem::Tick(float DeltaTime)
 {
 	Super::Tick(DeltaTime);
diff --git a/Engine/Plugins/Experimental/Water/Source/Runtime/Private/WaterZoneActor.cpp b/Engine/Plugins/Experimental/Water/Source/Runtime/Private/WaterZoneActor.cpp
--- a/Engine/Plugins/Experimental/Water/Source/Runtime/Private/WaterZoneActor.cpp
+++ b/Engine/Plugins/Experimental/Water/Source/Runtime/Private/WaterZoneActor.cpp
@@ -112,7 +112,7 @@ void AWaterZone::MarkForRebuild(EWaterZoneRebuildFlags Flags)
 
 void AWaterZone::ForEachWaterBodyComponent(TFunctionRef<bool(UWaterBodyComponent*)> Predicate)
 {
-	UWaterSubsystem::ForEachWaterBodyComponent(GetWorld(), [this, Predicate](UWaterBodyComponent* Component)
+	FWaterBodyManager::ForEachWaterBodyComponent(GetWorld(), [this, Predicate](UWaterBodyComponent* Component)
 		{
 			if (Component->GetWaterZone() == this)
 			{
diff --git a/Engine/Plugins/Experimental/Water/Source/Runtime/Public/WaterBodyManager.h b/Engine/Plugins/Experimental/Water/Source/Runtime/Public/WaterBodyManager.h
--- a/Engine/Plugins/Experimental/Water/Source/Runtime/Public/WaterBodyManager.h
+++ b/Engine/Plugins/Experimental/Water/Source/Runtime/Public/WaterBodyManager.h
@@ -32,6 +32,9 @@ public:
 	/** Execute a predicate function on each valid water body. Predicate should return false for early exit. */
 	void ForEachWaterBodyComponent (TFunctionRef<bool(UWaterBodyComponent*)> Pred) const;
 
+	/** Execute a predicate function on each valid water body. Predicate should return false for early exit. */
+	static void ForEachWaterBodyComponent (const UWorld* World, TFunctionRef<bool(UWaterBodyComponent*)> Pred);
+
 	bool HasAnyWaterBodies() const { return WaterBodyComponents.Num() > 0; }
 
 private:
diff --git a/Engine/Plugins/Experimental/Water/Source/Runtime/Public/WaterSubsystem.h b/Engine/Plugins/Experimental/Water/Source/Runtime/Public/WaterSubsystem.h
--- a/Engine/Plugins/Experimental/Water/Source/Runtime/Public/WaterSubsystem.h
+++ b/Engine/Plugins/Experimental/Water/Source/Runtime/Public/WaterSubsystem.h
@@ -99,8 +99,8 @@ public:
 	/** Static helper function to get a waterbody manager from a world, returns nullptr if world or manager don't exist */
 	static FWaterBodyManager* GetWaterBodyManager(const UWorld* InWorld);
 
-	/** Execute a predicate function on each valid water body. Predicate should return false for early exit. */
-	static void ForEachWaterBodyComponent(const UWorld* World, TFunctionRef<bool(UWaterBodyComponent*)> Predicate);
+	UE_DEPRECATED(5.1, "Please use FWaterBodyManager::ForEachWaterBodyComponent instead.")
+	static void ForEachWaterBodyComponent(const UWorld* World, TFunctionRef<bool(UWaterBodyComponent*)> Predicate) {}
 
 	UE_DEPRECATED(5.1, "This will become a private member. Prefer calling GetWaterBodyManager instead")
 	FWaterBodyManager WaterBodyManager;
