Fixed: Water surfaces are not aligned between seams across the map on Android
Use Landscape LOD1 for water texture capture. Inline bottom 5 LOD for landscape mesh, similar to texture based landscape which always has 64x64 mip available

diff --git a/Engine/Plugins/Experimental/Water/Source/Runtime/Private/WaterInfoRendering.cpp b/Engine/Plugins/Experimental/Water/Source/Runtime/Private/WaterInfoRendering.cpp
--- a/Engine/Plugins/Experimental/Water/Source/Runtime/Private/WaterInfoRendering.cpp
+++ b/Engine/Plugins/Experimental/Water/Source/Runtime/Private/WaterInfoRendering.cpp
@@ -452,6 +452,7 @@ static FSceneRenderer* CreateWaterInfoDepthRenderer(
 		.SetRealtimeUpdate(false)
 		.SetResolveScene(false));
 	DepthViewFamily.SceneCaptureSource = ESceneCaptureSource::SCS_DeviceDepth;
+	DepthViewFamily.LandscapeLODOverride = Context.LandscapeLODOverride;
 
 	// Setup the view family
 	FSceneViewInitOptions DepthViewInitOptions;
@@ -531,6 +532,7 @@ static FSceneRenderer* CreateWaterInfoColorRenderer(
 		.SetRealtimeUpdate(false)
 		.SetResolveScene(false));
 	ColorViewFamily.SceneCaptureSource = ESceneCaptureSource::SCS_SceneColorSceneDepth;
+	ColorViewFamily.LandscapeLODOverride = Context.LandscapeLODOverride;
 
 	FSceneViewInitOptions ColorViewInitOptions;
 	ColorViewInitOptions.SetViewRectangle(FIntRect(0, 0, RenderTargetSize.X, RenderTargetSize.Y));
diff --git a/Engine/Plugins/Experimental/Water/Source/Runtime/Private/WaterInfoRendering.h b/Engine/Plugins/Experimental/Water/Source/Runtime/Private/WaterInfoRendering.h
--- a/Engine/Plugins/Experimental/Water/Source/Runtime/Private/WaterInfoRendering.h
+++ b/Engine/Plugins/Experimental/Water/Source/Runtime/Private/WaterInfoRendering.h
@@ -19,6 +19,7 @@ namespace UE::WaterInfo
 		TArray<UWaterBodyComponent*> WaterBodies;
 		TArray<TWeakObjectPtr<AActor>> GroundActors;
 		float CaptureZ;
+		int32 LandscapeLODOverride;
 	};
 	
 void UpdateWaterInfoRendering(
diff --git a/Engine/Plugins/Experimental/Water/Source/Runtime/Private/WaterZoneActor.cpp b/Engine/Plugins/Experimental/Water/Source/Runtime/Private/WaterZoneActor.cpp
--- a/Engine/Plugins/Experimental/Water/Source/Runtime/Private/WaterZoneActor.cpp
+++ b/Engine/Plugins/Experimental/Water/Source/Runtime/Private/WaterZoneActor.cpp
@@ -280,6 +280,7 @@ bool AWaterZone::UpdateWaterInfoTexture()
 		GroundZMin = TNumericLimits<float>::Max();
 		float GroundZMax = TNumericLimits<float>::Lowest();
 
+		int32 LandscapeLODOverride = 0;
 		TArray<TWeakObjectPtr<AActor>> GroundActors;
 		for (ALandscapeProxy* LandscapeProxy : TActorRange<ALandscapeProxy>(World))
 		{
@@ -287,6 +288,10 @@ bool AWaterZone::UpdateWaterInfoTexture()
 			GroundZMin = FMath::Min(GroundZMin, LandscapeBox.Min.Z);
 			GroundZMax = FMath::Max(GroundZMax, LandscapeBox.Max.Z);
 			GroundActors.Add(LandscapeProxy);
+
+			// Target mip 64x64 for a capture
+			int32 LOD64x64 = FMath::CeilLogTwo(LandscapeProxy->SubsectionSizeQuads + 1) - 6;
+			LandscapeLODOverride = FMath::Max(LandscapeLODOverride, LOD64x64);
 		}
 
 		const ETextureRenderTargetFormat Format = bHalfPrecisionTexture ? ETextureRenderTargetFormat::RTF_RGBA16f : RTF_RGBA32f;
@@ -298,6 +303,7 @@ bool AWaterZone::UpdateWaterInfoTexture()
 		Context.GroundActors = MoveTemp(GroundActors);
 		Context.CaptureZ = FMath::Max(WaterZMax, GroundZMax) + CaptureZOffset;
 		Context.TextureRenderTarget = WaterInfoTexture;
+		Context.LandscapeLODOverride = LandscapeLODOverride;
 
 		if (UWaterSubsystem* WaterSubsystem = UWaterSubsystem::GetWaterSubsystem(World))
 		{
diff --git a/Engine/Source/Runtime/Landscape/Private/LandscapeEdit.cpp b/Engine/Source/Runtime/Landscape/Private/LandscapeEdit.cpp
--- a/Engine/Source/Runtime/Landscape/Private/LandscapeEdit.cpp
+++ b/Engine/Source/Runtime/Landscape/Private/LandscapeEdit.cpp
@@ -7013,7 +7013,8 @@ void ULandscapeComponent::GeneratePlatformVertexData(const ITargetPlatform* Targ
 	const bool bStreamLandscapeMeshLODs = TargetPlatform
 		&& TargetPlatform->SupportsFeature(ETargetPlatformFeatures::LandscapeMeshLODStreaming)
 		&& NumHoleLods <= FMath::Min(MaxLOD, MaxLODClamp);
-	const int32 NumStreamingLODs = bStreamLandscapeMeshLODs ? FMath::Min(MaxLOD, MaxLODClamp) : 0;
+	int32 NumStreamingLODs = bStreamLandscapeMeshLODs ? FMath::Min(MaxLOD, MaxLODClamp) : 0;
+	NumStreamingLODs = FMath::Max(0, NumStreamingLODs - 5); // Always inline bottom 5 LODs
 	TArray<int32> StreamingLODVertStartOffsets;
 	StreamingLODVertStartOffsets.AddUninitialized(NumStreamingLODs);
 
diff --git a/Engine/Source/Runtime/Landscape/Private/LandscapeRender.cpp b/Engine/Source/Runtime/Landscape/Private/LandscapeRender.cpp
--- a/Engine/Source/Runtime/Landscape/Private/LandscapeRender.cpp
+++ b/Engine/Source/Runtime/Landscape/Private/LandscapeRender.cpp
@@ -618,10 +618,8 @@ static int32 GetViewLodOverride(FSceneView const& View)
 {
 	// Apply r.ForceLOD override
 	int32 LodOverride = GetCVarForceLOD_AnyThread();
-#if WITH_EDITOR
 	// Apply editor landscape lod override
 	LodOverride = View.Family->LandscapeLODOverride >= 0 ? View.Family->LandscapeLODOverride : LodOverride;
-#endif
 	// Use lod 0 if lodding is disabled
 	LodOverride = View.Family->EngineShowFlags.LOD == 0 ? 0 : LodOverride;
 	return LodOverride;
