Select the lowest required landscape LOD when rendering the water info texture to prevent massive overdraw of landscape quads

diff --git a/Engine/Plugins/Experimental/Water/Source/Runtime/Private/WaterInfoRendering.cpp b/Engine/Plugins/Experimental/Water/Source/Runtime/Private/WaterInfoRendering.cpp
--- a/Engine/Plugins/Experimental/Water/Source/Runtime/Private/WaterInfoRendering.cpp
+++ b/Engine/Plugins/Experimental/Water/Source/Runtime/Private/WaterInfoRendering.cpp
@@ -17,6 +17,7 @@
 #include "WaterZoneActor.h"
 #include "Rendering/NaniteStreamingManager.h"
 #include "SceneViewExtension.h"
+#include "LandscapeRender.h"
 
 static int32 RenderCaptureNextWaterInfoDraws = 0;
 static FAutoConsoleVariableRef CVarRenderCaptureNextWaterInfoDraws(
@@ -34,6 +35,7 @@ struct FUpdateWaterInfoParams
 	FRenderTarget* RenderTarget;
 	FTexture* OutputTexture;
 
+	FVector2D WaterZoneExtents;
 	FVector2f WaterHeightExtents;
 	float GroundZMin;
 	float CaptureZ;
@@ -291,10 +293,66 @@ static void SetWaterBodiesWithinWaterInfoPass(FSceneRenderer* SceneRenderer, boo
 	}
 }
 
+static void SetOptimalLandscapeLODOverrides(const FUpdateWaterInfoParams& Params)
+{
+	// In order to prevent overdrawing the landscape components, we compute the lowest-detailed LOD level which satisfies the pixel coverage of the Water Info texture
+	// and force it on all landscape components. This override is set different per Landscape actor in case there are multiple under the same water zone.
+	//
+	// Ex: If the WaterInfoTexture only has 1 pixel per 100 units, and the highest landscape LOD has 1 vertex per 20 units, we don't need to use the maximum landscape LOD
+	// and can force a lower level of detail (in this case LOD2) while still satisfying the resolution of the water info texture.
+
+	const double MinWaterInfoTextureExtent = (double)FMath::Min(Params.OutputTexture->GetSizeX(), Params.OutputTexture->GetSizeY());
+	const double MaxWaterZoneExtent = FMath::Max(Params.WaterZoneExtents.X, Params.WaterZoneExtents.Y);
+	const double WaterInfoUnitsPerPixel =  MaxWaterZoneExtent / MinWaterInfoTextureExtent;
+
+	for (TPair<FLandscapeNeighborInfo::FLandscapeKey, FLandscapeRenderSystem*>& Pair : LandscapeRenderSystems)
+	{
+		int32 OptimalLODLevel = -1;
+		// All components within the same landscape (and thus its render system) should have the same number of quads and the same extent.
+		// therefore we can simply find the first component and compute its optimal LOD level.
+		for (FLandscapeComponentSceneProxy* LandscapeComponentProxy : Pair.Value->SceneProxies)
+		{
+			if (LandscapeComponentProxy != nullptr)
+			{
+				// LandscapeComponent Max Extend represents the half-extent of the landscape component. Multiply by two to get the actual size.
+				const double LandscapeComponentFullExtent = 2.0 * LandscapeComponentProxy->GetComponentMaxExtend();
+				const double LandscapeComponentUnitsPerVertex = LandscapeComponentFullExtent / (double)(LandscapeComponentProxy->GetComponentSizeQuads() + 1);
+
+				// Derived from:
+				// LandscapeComponentUnitsPerVertex * 2 ^ (LODLevel) <= WaterInfoUnitsPerPixel
+				OptimalLODLevel = FMath::FloorLog2(WaterInfoUnitsPerPixel / LandscapeComponentUnitsPerVertex);
+
+				break;
+			}
+		}
+
+		// There should always be at least one valid component proxy and the optimal LOD level should never be negative.
+		check(OptimalLODLevel >= 0);
+
+		for (FLandscapeRenderSystem::LODSettingsComponent& LODSettings : Pair.Value->SectionLODSettings)
+		{
+			LODSettings.ForcedLOD = OptimalLODLevel;
+		}
+	}
+}
+
+static void ResetLandscapeLODOverrides()
+{
+	for (TPair<FLandscapeNeighborInfo::FLandscapeKey, FLandscapeRenderSystem*>& Pair : LandscapeRenderSystems)
+	{
+		for (FLandscapeRenderSystem::LODSettingsComponent& LODSettings : Pair.Value->SectionLODSettings)
+		{
+			LODSettings.ForcedLOD = -1;
+		}
+	}
+}
+
 static void UpdateWaterInfoRendering_RenderThread(
 	FRHICommandListImmediate& RHICmdList,
 	const FUpdateWaterInfoParams& Params)
 {
+	SetOptimalLandscapeLODOverrides(Params);
+
 	FMaterialRenderProxy::UpdateDeferredCachedUniformExpressions();
 
 	FRenderTarget* RenderTarget = Params.RenderTarget;
@@ -415,6 +473,8 @@ static void UpdateWaterInfoRendering_RenderThread(
 	}
 
 	RHICmdList.Transition(FRHITransitionInfo(Params.OutputTexture->TextureRHI, ERHIAccess::RTV, ERHIAccess::SRVMask));
+
+	ResetLandscapeLODOverrides();
 }
 
 static FEngineShowFlags GetWaterInfoBaseShowFlags()
@@ -452,7 +512,6 @@ static FSceneRenderer* CreateWaterInfoDepthRenderer(
 		.SetRealtimeUpdate(false)
 		.SetResolveScene(false));
 	DepthViewFamily.SceneCaptureSource = ESceneCaptureSource::SCS_DeviceDepth;
-	DepthViewFamily.LandscapeLODOverride = Context.LandscapeLODOverride;
 
 	// Setup the view family
 	FSceneViewInitOptions DepthViewInitOptions;
@@ -532,7 +591,6 @@ static FSceneRenderer* CreateWaterInfoColorRenderer(
 		.SetRealtimeUpdate(false)
 		.SetResolveScene(false));
 	ColorViewFamily.SceneCaptureSource = ESceneCaptureSource::SCS_SceneColorSceneDepth;
-	ColorViewFamily.LandscapeLODOverride = Context.LandscapeLODOverride;
 
 	FSceneViewInitOptions ColorViewInitOptions;
 	ColorViewInitOptions.SetViewRectangle(FIntRect(0, 0, RenderTargetSize.X, RenderTargetSize.Y));
@@ -640,6 +698,7 @@ void UpdateWaterInfoRendering(
 	Params.WaterHeightExtents = Context.ZoneToRender->GetWaterHeightExtents();
 	Params.GroundZMin = Context.ZoneToRender->GetGroundZMin();
 	Params.VelocityBlurRadius = Context.ZoneToRender->GetVelocityBlurRadius();
+	Params.WaterZoneExtents = Context.ZoneToRender->GetZoneExtent();
 
 	ENQUEUE_RENDER_COMMAND(WaterInfoCommand)(
 	[Params, ZoneName = Context.ZoneToRender->GetActorNameOrLabel()](FRHICommandListImmediate& RHICmdList)
diff --git a/Engine/Plugins/Experimental/Water/Source/Runtime/Private/WaterInfoRendering.h b/Engine/Plugins/Experimental/Water/Source/Runtime/Private/WaterInfoRendering.h
--- a/Engine/Plugins/Experimental/Water/Source/Runtime/Private/WaterInfoRendering.h
+++ b/Engine/Plugins/Experimental/Water/Source/Runtime/Private/WaterInfoRendering.h
@@ -19,7 +19,6 @@ namespace UE::WaterInfo
 		TArray<UWaterBodyComponent*> WaterBodies;
 		TArray<TWeakObjectPtr<AActor>> GroundActors;
 		float CaptureZ;
-		int32 LandscapeLODOverride;
 	};
 	
 void UpdateWaterInfoRendering(
diff --git a/Engine/Plugins/Experimental/Water/Source/Runtime/Private/WaterZoneActor.cpp b/Engine/Plugins/Experimental/Water/Source/Runtime/Private/WaterZoneActor.cpp
--- a/Engine/Plugins/Experimental/Water/Source/Runtime/Private/WaterZoneActor.cpp
+++ b/Engine/Plugins/Experimental/Water/Source/Runtime/Private/WaterZoneActor.cpp
@@ -280,18 +280,18 @@ bool AWaterZone::UpdateWaterInfoTexture()
 		GroundZMin = TNumericLimits<float>::Max();
 		float GroundZMax = TNumericLimits<float>::Lowest();
 
-		int32 LandscapeLODOverride = 0;
 		TArray<TWeakObjectPtr<AActor>> GroundActors;
 		for (ALandscapeProxy* LandscapeProxy : TActorRange<ALandscapeProxy>(World))
 		{
+			if (LandscapeProxy->LandscapeComponents.Num() == 0)
+			{
+				continue; // Skip empty landscape proxies
+			}
+
 			const FBox LandscapeBox = LandscapeProxy->GetComponentsBoundingBox();
 			GroundZMin = FMath::Min(GroundZMin, LandscapeBox.Min.Z);
 			GroundZMax = FMath::Max(GroundZMax, LandscapeBox.Max.Z);
 			GroundActors.Add(LandscapeProxy);
-
-			// Target mip 64x64 for a capture
-			int32 LOD64x64 = FMath::CeilLogTwo(LandscapeProxy->SubsectionSizeQuads + 1) - 6;
-			LandscapeLODOverride = FMath::Max(LandscapeLODOverride, LOD64x64);
 		}
 
 		const ETextureRenderTargetFormat Format = bHalfPrecisionTexture ? ETextureRenderTargetFormat::RTF_RGBA16f : RTF_RGBA32f;
@@ -303,7 +303,6 @@ bool AWaterZone::UpdateWaterInfoTexture()
 		Context.GroundActors = MoveTemp(GroundActors);
 		Context.CaptureZ = FMath::Max(WaterZMax, GroundZMax) + CaptureZOffset;
 		Context.TextureRenderTarget = WaterInfoTexture;
-		Context.LandscapeLODOverride = LandscapeLODOverride;
 
 		if (UWaterSubsystem* WaterSubsystem = UWaterSubsystem::GetWaterSubsystem(World))
 		{
diff --git a/Engine/Source/Runtime/Landscape/Public/LandscapeRender.h b/Engine/Source/Runtime/Landscape/Public/LandscapeRender.h
--- a/Engine/Source/Runtime/Landscape/Public/LandscapeRender.h
+++ b/Engine/Source/Runtime/Landscape/Public/LandscapeRender.h
@@ -627,6 +627,9 @@ public:
 	// Also being reused by GPULightmass currently to save mem
 	static TMap<uint32, FLandscapeSharedBuffers*> SharedBuffersMap;
 
+	float GetComponentMaxExtend() const { return ComponentMaxExtend; }
+	int32 GetComponentSizeQuads() const { return ComponentSizeQuads; }
+
 protected:
 	int8						MaxLOD;		// Maximum LOD level, user override possible
 	int8						NumWeightmapLayerAllocations;
