* Fixed water brush letting the user change the list of affected weightmap layers while it's actually automatically populated by the list of layer names specified in the water bodies
* Fixed weightmap changes on water bodies not triggering the appropriate water brush updates

diff --git a/Engine/Plugins/Experimental/Water/Source/Editor/Private/WaterLandscapeBrush.cpp b/Engine/Plugins/Experimental/Water/Source/Editor/Private/WaterLandscapeBrush.cpp
--- a/Engine/Plugins/Experimental/Water/Source/Editor/Private/WaterLandscapeBrush.cpp
+++ b/Engine/Plugins/Experimental/Water/Source/Editor/Private/WaterLandscapeBrush.cpp
@@ -604,6 +604,22 @@ void AWaterLandscapeBrush::UpdateActorIcon()
 	}
 }
 
+bool AWaterLandscapeBrush::CanEditChange(const FProperty* InProperty) const
+{
+	if (!Super::CanEditChange(InProperty))
+	{
+		return false;
+	}
+
+	// Weightmap layers are automatically populated by the list of IWaterBrushActorInterface affecting this brush :
+	if (InProperty->GetFName() == GET_MEMBER_NAME_CHECKED(AWaterLandscapeBrush, AffectedWeightmapLayers))
+	{
+		return false;
+	}
+
+	return true;
+}
+
 #endif // WITH_EDITOR
 
 #undef LOCTEXT_NAMESPACE
diff --git a/Engine/Plugins/Experimental/Water/Source/Editor/Private/WaterLandscapeBrush.h b/Engine/Plugins/Experimental/Water/Source/Editor/Private/WaterLandscapeBrush.h
--- a/Engine/Plugins/Experimental/Water/Source/Editor/Private/WaterLandscapeBrush.h
+++ b/Engine/Plugins/Experimental/Water/Source/Editor/Private/WaterLandscapeBrush.h
@@ -96,6 +96,8 @@ public:
 	virtual void CheckForErrors() override;
 	
 	void UpdateActorIcon();
+
+	virtual bool CanEditChange(const FProperty* InProperty) const override;
 #endif // WITH_EDITOR
 
 #if WITH_EDITORONLY_DATA
diff --git a/Engine/Plugins/Experimental/Water/Source/Runtime/Private/WaterBodyComponent.cpp b/Engine/Plugins/Experimental/Water/Source/Runtime/Private/WaterBodyComponent.cpp
--- a/Engine/Plugins/Experimental/Water/Source/Runtime/Private/WaterBodyComponent.cpp
+++ b/Engine/Plugins/Experimental/Water/Source/Runtime/Private/WaterBodyComponent.cpp
@@ -866,7 +866,7 @@ void UWaterBodyComponent::PostEditImport()
 void UWaterBodyComponent::OnPostEditChangeProperty(FPropertyChangedEvent& PropertyChangedEvent, bool& bShapeOrPositionChanged, bool& bWeightmapSettingsChanged)
 {
 	const FName PropertyName = PropertyChangedEvent.GetPropertyName();
-	if (PropertyName == GET_MEMBER_NAME_CHECKED(UWaterBodyComponent, LayerWeightmapSettings))
+	if (PropertyChangedEvent.MemberProperty && PropertyChangedEvent.MemberProperty->GetFName() == GET_MEMBER_NAME_CHECKED(UWaterBodyComponent, LayerWeightmapSettings))
 	{
 		bWeightmapSettingsChanged = true;
 	}
