Consolidated the Waterbody navigation and collision settings onto the UPrimitiveComponent interfaces and deprecated the duplicate water functions.

diff --git a/Engine/Plugins/Experimental/Water/Source/Runtime/Private/WaterBodyActor.cpp b/Engine/Plugins/Experimental/Water/Source/Runtime/Private/WaterBodyActor.cpp
--- a/Engine/Plugins/Experimental/Water/Source/Runtime/Private/WaterBodyActor.cpp
+++ b/Engine/Plugins/Experimental/Water/Source/Runtime/Private/WaterBodyActor.cpp
@@ -448,7 +448,7 @@ void AWaterBody::DeprecateData()
 		WaterBodyComponent->PhysicalMaterial = PhysicalMaterial_DEPRECATED;
 		WaterBodyComponent->TargetWaveMaskDepth = TargetWaveMaskDepth_DEPRECATED;
 		WaterBodyComponent->MaxWaveHeightOffset = MaxWaveHeightOffset_DEPRECATED;
-		WaterBodyComponent->bFillCollisionUnderWaterBodiesForNavmesh = bFillCollisionUnderWaterBodiesForNavmesh_DEPRECATED;
+		WaterBodyComponent->bFillCollisionUnderWaterBodiesForNavmesh_DEPRECATED = bFillCollisionUnderWaterBodiesForNavmesh_DEPRECATED;
 		WaterBodyComponent->UnderwaterPostProcessSettings = UnderwaterPostProcessSettings_DEPRECATED;
 		WaterBodyComponent->CurveSettings = CurveSettings_DEPRECATED;
 		WaterBodyComponent->SetWaterMaterial(WaterMaterial_DEPRECATED);
@@ -456,17 +456,20 @@ void AWaterBody::DeprecateData()
 		WaterBodyComponent->WaterHeightmapSettings = WaterHeightmapSettings_DEPRECATED;
 		WaterBodyComponent->LayerWeightmapSettings = LayerWeightmapSettings_DEPRECATED;
 		WaterBodyComponent->bAffectsLandscape = bAffectsLandscape_DEPRECATED;
-		WaterBodyComponent->bGenerateCollisions = bGenerateCollisions_DEPRECATED;
+		WaterBodyComponent->bGenerateCollisions_DEPRECATED = bGenerateCollisions_DEPRECATED;
 		WaterBodyComponent->WaterMeshOverride = WaterMeshOverride_DEPRECATED;
 		WaterBodyComponent->OverlapMaterialPriority = OverlapMaterialPriority_DEPRECATED;
-		WaterBodyComponent->CollisionProfileName = CollisionProfileName_DEPRECATED;
+		WaterBodyComponent->CollisionProfileName_DEPRECATED = CollisionProfileName_DEPRECATED;
 		WaterBodyComponent->WaterMID = WaterMID_DEPRECATED;
 		WaterBodyComponent->UnderwaterPostProcessMID = UnderwaterPostProcessMID_DEPRECATED;
 		WaterBodyComponent->Islands = Islands_DEPRECATED;
 		WaterBodyComponent->ExclusionVolumes = ExclusionVolumes_DEPRECATED;
-		WaterBodyComponent->bCanAffectNavigation = bCanAffectNavigation_DEPRECATED;
+		WaterBodyComponent->bCanAffectNavigation_DEPRECATED = bCanAffectNavigation_DEPRECATED;
 		WaterBodyComponent->WaterNavAreaClass = WaterNavAreaClass_DEPRECATED;
 		WaterBodyComponent->ShapeDilation = ShapeDilation_DEPRECATED;
+
+		// Some deprecated data on the actor were also later deprecated on the component (e.g. bFillCollisionUnderWaterBodiesForNavmesh_DEPRECATED), make sure it runs its own deprecation operations on it : 
+		WaterBodyComponent->DeprecateData();
 	}
 #endif // WITH_EDITORONLY_DATA
 }
diff --git a/Engine/Plugins/Experimental/Water/Source/Runtime/Private/WaterBodyComponent.cpp b/Engine/Plugins/Experimental/Water/Source/Runtime/Private/WaterBodyComponent.cpp
--- a/Engine/Plugins/Experimental/Water/Source/Runtime/Private/WaterBodyComponent.cpp
+++ b/Engine/Plugins/Experimental/Water/Source/Runtime/Private/WaterBodyComponent.cpp
@@ -2,6 +2,7 @@
 
 
 #include "WaterBodyComponent.h"
+#include "UObject/FortniteMainBranchObjectVersion.h"
 #include "UObject/UObjectIterator.h"
 #include "EngineUtils.h"
 #include "Landscape.h"
@@ -79,17 +80,24 @@ UWaterBodyComponent::UWaterBodyComponent(const FObjectInitializer& ObjectInitial
 {
 	bAffectsLandscape = true;
 
-	CollisionProfileName = GetDefault<UWaterRuntimeSettings>()->GetDefaultWaterCollisionProfileName();
+	SetCollisionProfileName(GetDefault<UWaterRuntimeSettings>()->GetDefaultWaterCollisionProfileName());
 
 	WaterMID = nullptr;
 	WaterInfoMID = nullptr;
 
 	TargetWaveMaskDepth = 2048.f;
 
-	bCanAffectNavigation = false;
-	bFillCollisionUnderWaterBodiesForNavmesh = false;
+	bFillCollisionUnderneathForNavmesh = false;
+	bCanEverAffectNavigation = false;
 
 	WaterInfoMaterial = GetDefault<UWaterRuntimeSettings>()->GetDefaultWaterInfoMaterial();
+
+#if WITH_EDITORONLY_DATA
+	// Maintain the old default values for deprecated members so delta serialization is still correct when we deprecate them in PostLoad
+	bCanAffectNavigation_DEPRECATED = false;
+	bFillCollisionUnderWaterBodiesForNavmesh_DEPRECATED = false;
+	CollisionProfileName_DEPRECATED = GetDefault<UWaterRuntimeSettings>()->GetDefaultWaterCollisionProfileName();
+#endif // WITH_EDITORONLY_DATA
 }
 
 void UWaterBodyComponent::OnVisibilityChanged()
@@ -1018,10 +1026,9 @@ void UWaterBodyComponent::RegisterOnChangeWaterSplineMetadata(UWaterSplineMetada
 
 void UWaterBodyComponent::GetNavigationData(struct FNavigationRelevantData& Data) const
 {
-	if (CanAffectNavigation())
+	if (IsNavigationRelevant())
 	{
-		const TSubclassOf<UNavAreaBase> UseAreaClass = GetNavAreaClass();
-		TArray<UPrimitiveComponent*> LocalCollisionComponents = GetCollisionComponents();
+		TArray<UPrimitiveComponent*> LocalCollisionComponents = GetCollisionComponents(/*bInOnlyEnabledComponents = */true);
 		for (int32 CompIdx = 0; CompIdx < LocalCollisionComponents.Num(); CompIdx++)
 		{
 			UPrimitiveComponent* PrimComp = LocalCollisionComponents[CompIdx];
@@ -1056,41 +1063,26 @@ FBox UWaterBodyComponent::GetNavigationBounds() const
 
 bool UWaterBodyComponent::IsNavigationRelevant() const
 {
-	return CanAffectNavigation() && (GetCollisionComponents().Num() > 0);
+	return Super::IsNavigationRelevant() && (GetCollisionComponents().Num() > 0);
 }
 
-void UWaterBodyComponent::ApplyNavigationSettings()
+void UWaterBodyComponent::ApplyCollisionSettings()
 {
-	AActor* Owner = GetOwner();
-	if (Owner)
+	// Transfer the collision settings of the water body component to all of its child collision components 
+	TArray<UPrimitiveComponent*> CollisionComponents = GetCollisionComponents(/*bInOnlyEnabledComponents = */false);
+	for (UPrimitiveComponent* CollisionComponent : CollisionComponents)
 	{
-		const bool bCanAffectNav = CanAffectNavigation();
-
-		// Make sure the engine's bCanEverAffectionNavigation matches the component's one.
-		SetCanEverAffectNavigation(bCanAffectNav);
-
-		// @todo_water: change this
-		TInlineComponentArray<UActorComponent*> Components;
-		Owner->GetComponents(Components);
+		CopySharedCollisionSettingsToComponent(CollisionComponent);
+	}
+}
 
-		const TArray<UPrimitiveComponent*> LocalCollisionComponents = GetCollisionComponents();
-		for (UActorComponent* ActorComp : Components)
-		{
-			UPrimitiveComponent* PrimComp = Cast<UPrimitiveComponent>(ActorComp);
-			if (PrimComp == this)
-			{
-				continue;
-			}
-			if (!PrimComp || LocalCollisionComponents.Find(PrimComp) == INDEX_NONE)
-			{
-				ActorComp->SetCanEverAffectNavigation(false);
-			}
-			else
-			{
-				PrimComp->SetCustomNavigableGeometry(bCanAffectNav ? EHasCustomNavigableGeometry::EvenIfNotCollidable : EHasCustomNavigableGeometry::No);
-				PrimComp->SetCanEverAffectNavigation(bCanAffectNav);
-			}
-		}
+void UWaterBodyComponent::ApplyNavigationSettings()
+{
+	// Transfer the navigation settings of the water body component to all of its child collision components 
+	const TArray<UPrimitiveComponent*> CollisionComponents = GetCollisionComponents(/*bInOnlyEnabledComponents = */false);
+	for (UPrimitiveComponent* CollisionComponent : CollisionComponents)
+	{
+		CopySharedNavigationSettingsToComponent(CollisionComponent);
 	}
 }
 
@@ -1150,6 +1142,8 @@ void UWaterBodyComponent::UpdateAll(bool bShapeOrPositionChanged)
 		// Finally, generate the body once again, this time with the updated list of exclusion volumes
 		UpdateWaterBody(/*bWithExclusionVolumes*/true);
 
+		ApplyCollisionSettings();
+
 		ApplyNavigationSettings();
 
 		if (bShapeOrPositionChanged)
@@ -1205,6 +1199,7 @@ void UWaterBodyComponent::Serialize(FArchive& Ar)
 	Super::Serialize(Ar);
 
 	Ar.UsingCustomVersion(FWaterCustomVersion::GUID);
+	Ar.UsingCustomVersion(FFortniteMainBranchObjectVersion::GUID);
 }
 
 void UWaterBodyComponent::PostLoad()
@@ -1217,11 +1212,34 @@ void UWaterBodyComponent::PostLoad()
 	{
 		WaterMeshOverride = nullptr;
 	}
-#endif
+#endif // WITH_EDITORONLY_DATA
+
+	DeprecateData();
 
 #if WITH_EDITOR
 	RegisterOnUpdateWavesData(GetWaterWaves(), /* bRegister = */true);
-#endif
+#endif // WITH_EDITOR
+}
+
+void UWaterBodyComponent::DeprecateData()
+{
+#if WITH_EDITORONLY_DATA
+	if (GetLinkerCustomVersion(FFortniteMainBranchObjectVersion::GUID) < FFortniteMainBranchObjectVersion::WaterBodyComponentCollisionSettingsRefactor)
+	{
+		// Deprecate the old collision / navigation data and update it on all sub-components :
+		SetCollisionProfileName(CollisionProfileName_DEPRECATED);
+		SetGenerateOverlapEvents(bGenerateCollisions_DEPRECATED);
+		// Transfer info to sub-components :
+		ApplyCollisionSettings();
+
+		bool bCanAffectNav = bGenerateCollisions_DEPRECATED && bCanAffectNavigation_DEPRECATED;
+		SetCustomNavigableGeometry(bCanAffectNav ? EHasCustomNavigableGeometry::EvenIfNotCollidable : EHasCustomNavigableGeometry::No);
+		SetCanEverAffectNavigation(bCanAffectNav);
+		bFillCollisionUnderneathForNavmesh = bFillCollisionUnderWaterBodiesForNavmesh_DEPRECATED;
+		// Transfer info to sub-components :
+		ApplyNavigationSettings();
+	}
+#endif // WITH_EDITORONLY_DATA
 }
 
 void UWaterBodyComponent::OnComponentDestroyed(bool bDestroyingHierarchy)
@@ -1249,6 +1267,27 @@ bool UWaterBodyComponent::MoveComponentImpl(const FVector& Delta, const FQuat& N
 	return Super::MoveComponentImpl(Delta, CorrectedRotation, bSweep, Hit, MoveFlags, Teleport);
 }
 
+void UWaterBodyComponent::OnComponentCollisionSettingsChanged(bool bUpdateOverlaps)
+{
+	if (IsRegistered() && !IsTemplate())			// not for CDOs
+	{
+		Super::OnComponentCollisionSettingsChanged(bUpdateOverlaps);
+
+		// Transfer all settings leading to OnComponentCollisionSettingsChanged to be called to the sub-components handling collisions:
+		ApplyCollisionSettings();
+	}
+}
+
+void UWaterBodyComponent::OnGenerateOverlapEventsChanged()
+{
+	if (IsRegistered() && !IsTemplate())			// not for CDOs
+	{
+		Super::OnGenerateOverlapEventsChanged();
+
+		ApplyCollisionSettings();
+	}
+}
+
 bool UWaterBodyComponent::SetDynamicParametersOnMID(UMaterialInstanceDynamic* InMID)
 {
 	UWaterSubsystem* WaterSubsystem = UWaterSubsystem::GetWaterSubsystem(GetWorld());
@@ -1385,6 +1424,22 @@ float UWaterBodyComponent::GetMaxWaveHeight() const
 	return (HasWaves() ? GetWaterWaves()->GetMaxWaveHeight() : 0.0f) + MaxWaveHeightOffset;
 }
 
+void UWaterBodyComponent::CopySharedCollisionSettingsToComponent(UPrimitiveComponent * InComponent)
+{
+	InComponent->SetCollisionEnabled(GetCollisionEnabled());
+	InComponent->SetNotifyRigidBodyCollision(BodyInstance.bNotifyRigidBodyCollision);
+	InComponent->SetCollisionResponseToChannels(BodyInstance.GetResponseToChannels());
+	InComponent->SetCollisionProfileName(GetCollisionProfileName(), /*bUpdateOverlaps=*/ true);
+	InComponent->SetGenerateOverlapEvents(GetGenerateOverlapEvents());
+}
+
+void UWaterBodyComponent::CopySharedNavigationSettingsToComponent(UPrimitiveComponent * InComponent)
+{
+	InComponent->SetCanEverAffectNavigation(CanEverAffectNavigation());
+	InComponent->SetCustomNavigableGeometry(HasCustomNavigableGeometry());
+	InComponent->bFillCollisionUnderneathForNavmesh = GetCollisionEnabled() != ECollisionEnabled::NoCollision && bFillCollisionUnderneathForNavmesh;
+}
+
 float UWaterBodyComponent::GetWaveHeightAtPosition(const FVector& InPosition, float InWaterDepth, float InTime, FVector& OutNormal) const
 {
 	check(HasWaves());
diff --git a/Engine/Plugins/Experimental/Water/Source/Runtime/Private/WaterBodyCustomComponent.cpp b/Engine/Plugins/Experimental/Water/Source/Runtime/Private/WaterBodyCustomComponent.cpp
--- a/Engine/Plugins/Experimental/Water/Source/Runtime/Private/WaterBodyCustomComponent.cpp
+++ b/Engine/Plugins/Experimental/Water/Source/Runtime/Private/WaterBodyCustomComponent.cpp
@@ -25,10 +25,10 @@ UWaterBodyCustomComponent::UWaterBodyCustomComponent(const FObjectInitializer& O
 	check(!IsHeightOffsetSupported());
 }
 
-TArray<UPrimitiveComponent*> UWaterBodyCustomComponent::GetCollisionComponents() const
+TArray<UPrimitiveComponent*> UWaterBodyCustomComponent::GetCollisionComponents(bool bInOnlyEnabledComponents) const
 {
 	TArray<UPrimitiveComponent*> Result;
-	if ((MeshComp != nullptr) && (MeshComp->GetCollisionEnabled() != ECollisionEnabled::NoCollision))
+	if ((MeshComp != nullptr) && (!bInOnlyEnabledComponents || (MeshComp->GetCollisionEnabled() != ECollisionEnabled::NoCollision)))
 	{
 		Result.Add(MeshComp);
 	}
@@ -70,13 +70,6 @@ void UWaterBodyCustomComponent::OnUpdateBody(bool bWithExclusionVolumes)
 		MeshComp = NewObject<UStaticMeshComponent>(OwnerActor, TEXT("CustomMeshComponent"), RF_Transactional);
 		MeshComp->SetNetAddressable(); // it's deterministically named so it's addressable over network (needed for collision)
 		MeshComp->SetupAttachment(this);
-		MeshComp->SetCollisionProfileName(GetCollisionProfileName());
-		// In the case of custom meshes, the static mesh component acts as both collision and visual component so we simply disable collision on it: 
-		MeshComp->SetGenerateOverlapEvents(bGenerateCollisions);
-		if (!bGenerateCollisions)
-		{
-			MeshComp->SetCollisionEnabled(ECollisionEnabled::NoCollision);
-		}
 		MeshComp->RegisterComponent();
 	}
 
@@ -86,11 +79,7 @@ void UWaterBodyCustomComponent::OnUpdateBody(bool bWithExclusionVolumes)
 	// Make no assumptions for custom meshes.  Add all components with collision to the list of collision components
 	for (UPrimitiveComponent* Comp : PrimitiveComponents)
 	{
-		if (bGenerateCollisions && (Comp->GetCollisionEnabled() != ECollisionEnabled::NoCollision))
-		{
-			// Use value of bFillCollisionUnderWaterBodiesForNavmesh for all components with collisions.
-			Comp->bFillCollisionUnderneathForNavmesh = bFillCollisionUnderWaterBodiesForNavmesh;
-		}
+		CopySharedNavigationSettingsToComponent(Comp);
 
 		Comp->SetMobility(Mobility);
 	}
@@ -99,6 +88,8 @@ void UWaterBodyCustomComponent::OnUpdateBody(bool bWithExclusionVolumes)
 	MeshComp->SetStaticMesh(GetWaterMeshOverride());
 	MeshComp->SetMaterial(0, WaterMID);
 	MeshComp->SetCastShadow(false);
+	CopySharedCollisionSettingsToComponent(MeshComp);
+	CopySharedNavigationSettingsToComponent(MeshComp);
 	MeshComp->MarkRenderStateDirty();
 }
 
diff --git a/Engine/Plugins/Experimental/Water/Source/Runtime/Private/WaterBodyLakeComponent.cpp b/Engine/Plugins/Experimental/Water/Source/Runtime/Private/WaterBodyLakeComponent.cpp
--- a/Engine/Plugins/Experimental/Water/Source/Runtime/Private/WaterBodyLakeComponent.cpp
+++ b/Engine/Plugins/Experimental/Water/Source/Runtime/Private/WaterBodyLakeComponent.cpp
@@ -29,10 +29,10 @@ UWaterBodyLakeComponent::UWaterBodyLakeComponent(const FObjectInitializer& Objec
 	check(!IsHeightOffsetSupported());
 }
 
-TArray<UPrimitiveComponent*> UWaterBodyLakeComponent::GetCollisionComponents() const
+TArray<UPrimitiveComponent*> UWaterBodyLakeComponent::GetCollisionComponents(bool bInOnlyEnabledComponents) const
 {
 	TArray<UPrimitiveComponent*> Result;
-	if (LakeCollision != nullptr)
+	if ((LakeCollision != nullptr) && (!bInOnlyEnabledComponents || (LakeCollision->GetCollisionEnabled() != ECollisionEnabled::NoCollision)))
 	{
 		Result.Add(LakeCollision);
 	}
@@ -184,7 +184,7 @@ void UWaterBodyLakeComponent::OnUpdateBody(bool bWithExclusionVolumes)
 		LakeMeshComp->RegisterComponent();
 	}
 
-	if (bGenerateCollisions)
+	if (GetCollisionEnabled() != ECollisionEnabled::NoCollision)
 	{
 		if (!LakeCollision)
 		{
@@ -236,11 +236,10 @@ void UWaterBodyLakeComponent::OnUpdateBody(bool bWithExclusionVolumes)
 
 		if (LakeCollision)
 		{
-			check(bGenerateCollisions);
-			LakeCollision->bFillCollisionUnderneathForNavmesh = bFillCollisionUnderWaterBodiesForNavmesh;
+			check(GetCollisionEnabled () != ECollisionEnabled::NoCollision);
 			LakeCollision->SetMobility(Mobility);
-			LakeCollision->SetCollisionProfileName(GetCollisionProfileName());
-			LakeCollision->SetGenerateOverlapEvents(true);
+			CopySharedCollisionSettingsToComponent(LakeCollision);
+			CopySharedNavigationSettingsToComponent(LakeCollision);
 
 			const float Depth = GetChannelDepth() / 2;
 			FVector LakeCollisionExtent = FVector(SplineExtent.X, SplineExtent.Y, 0.f) / GetComponentScale();
diff --git a/Engine/Plugins/Experimental/Water/Source/Runtime/Private/WaterBodyOceanComponent.cpp b/Engine/Plugins/Experimental/Water/Source/Runtime/Private/WaterBodyOceanComponent.cpp
--- a/Engine/Plugins/Experimental/Water/Source/Runtime/Private/WaterBodyOceanComponent.cpp
+++ b/Engine/Plugins/Experimental/Water/Source/Runtime/Private/WaterBodyOceanComponent.cpp
@@ -16,6 +16,7 @@
 #include "ConstrainedDelaunay2.h"
 #include "Operations/InsetMeshRegion.h"
 #include "DynamicMesh/DynamicMeshChangeTracker.h"
+#include "Algo/Transform.h"
 
 #if WITH_EDITOR
 #include "WaterIconHelper.h"
@@ -35,20 +36,19 @@ UWaterBodyOceanComponent::UWaterBodyOceanComponent(const FObjectInitializer& Obj
 	check(IsHeightOffsetSupported());
 }
 
-TArray<UPrimitiveComponent*> UWaterBodyOceanComponent::GetCollisionComponents() const
+TArray<UPrimitiveComponent*> UWaterBodyOceanComponent::GetCollisionComponents(bool bInOnlyEnabledComponents) const
 {
 	TArray<UPrimitiveComponent*> Result;
 	Result.Reserve(CollisionBoxes.Num() + CollisionHullSets.Num());
 
-	for (UOceanBoxCollisionComponent* Comp : CollisionBoxes)
-	{
-		Result.Add(Comp);
-	}
+	Algo::TransformIf(CollisionBoxes, Result,
+		[bInOnlyEnabledComponents](UOceanBoxCollisionComponent* Comp) { return ((Comp != nullptr) && (!bInOnlyEnabledComponents || (Comp->GetCollisionEnabled() != ECollisionEnabled::NoCollision))); },
+		[](UOceanBoxCollisionComponent* Comp) { return Comp; });
+
+	Algo::TransformIf(CollisionHullSets, Result,
+		[bInOnlyEnabledComponents](UOceanCollisionComponent* Comp) { return ((Comp != nullptr) && (!bInOnlyEnabledComponents || (Comp->GetCollisionEnabled() != ECollisionEnabled::NoCollision))); },
+		[](UOceanCollisionComponent* Comp) { return Comp; });
 
-	for (UOceanCollisionComponent* Comp : CollisionHullSets)
-	{
-		Result.Add(Comp);
-	}
 	return Result;
 }
 
@@ -255,7 +255,7 @@ void UWaterBodyOceanComponent::OnUpdateBody(bool bWithExclusionVolumes)
 	AActor* OwnerActor = GetOwner();
 	check(OwnerActor);
 
-	if (bGenerateCollisions)
+	if (GetCollisionEnabled() != ECollisionEnabled::NoCollision)
 	{
 		FVector OceanCollisionExtents = GetCollisionExtents();
 		OceanCollisionExtents.Z += CollisionHeightOffset / 2;
@@ -316,9 +316,8 @@ void UWaterBodyOceanComponent::OnUpdateBody(bool bWithExclusionVolumes)
 			BoxComponent->bDrawOnlyIfSelected = true;
 			BoxComponent->SetRelativeLocation(FVector::ZeroVector);
 
-			BoxComponent->bFillCollisionUnderneathForNavmesh = bFillCollisionUnderWaterBodiesForNavmesh;
-			BoxComponent->SetCollisionProfileName(GetCollisionProfileName());
-			BoxComponent->SetGenerateOverlapEvents(true);
+			CopySharedCollisionSettingsToComponent(BoxComponent);
+			CopySharedNavigationSettingsToComponent(BoxComponent);
 
 			FVector RelativePosition = Box.Origin;			// boxes are calculated in space of actor
 			BoxComponent->SetRelativeLocation(RelativePosition);
@@ -353,9 +352,8 @@ void UWaterBodyOceanComponent::OnUpdateBody(bool bWithExclusionVolumes)
 			CollisionComponent->SetNetAddressable(); // it's deterministically named so it's addressable over network (needed for collision)
 			CollisionComponent->SetRelativeLocation(FVector::ZeroVector);
 
-			CollisionComponent->bFillCollisionUnderneathForNavmesh = bFillCollisionUnderWaterBodiesForNavmesh;
-			CollisionComponent->SetCollisionProfileName(GetCollisionProfileName());
-			CollisionComponent->SetGenerateOverlapEvents(true);
+			CopySharedCollisionSettingsToComponent(CollisionComponent);
+			CopySharedNavigationSettingsToComponent(CollisionComponent);
 
 			CollisionComponent->InitializeFromConvexElements(ConvexSet);
 		}
diff --git a/Engine/Plugins/Experimental/Water/Source/Runtime/Private/WaterBodyRiverComponent.cpp b/Engine/Plugins/Experimental/Water/Source/Runtime/Private/WaterBodyRiverComponent.cpp
--- a/Engine/Plugins/Experimental/Water/Source/Runtime/Private/WaterBodyRiverComponent.cpp
+++ b/Engine/Plugins/Experimental/Water/Source/Runtime/Private/WaterBodyRiverComponent.cpp
@@ -34,11 +34,15 @@ UWaterBodyRiverComponent::UWaterBodyRiverComponent(const FObjectInitializer& Obj
 }
 
 
-TArray<UPrimitiveComponent*> UWaterBodyRiverComponent::GetCollisionComponents() const
+TArray<UPrimitiveComponent*> UWaterBodyRiverComponent::GetCollisionComponents(bool bInOnlyEnabledComponents) const
 {
 	TArray<UPrimitiveComponent*> Result;
 	Result.Reserve(SplineMeshComponents.Num());
-	Algo::TransformIf(SplineMeshComponents, Result, [](USplineMeshComponent* SplineComp) { return ((SplineComp != nullptr) && (SplineComp->GetCollisionEnabled() != ECollisionEnabled::NoCollision)); }, [](USplineMeshComponent* SplineComp) { return SplineComp; });
+
+	Algo::TransformIf(SplineMeshComponents, Result, 
+		[bInOnlyEnabledComponents](USplineMeshComponent* SplineComp) { return ((SplineComp != nullptr) && (!bInOnlyEnabledComponents || (SplineComp->GetCollisionEnabled() != ECollisionEnabled::NoCollision))); }, 
+		[](USplineMeshComponent* SplineComp) { return SplineComp; });
+
 	return Result;
 }
 
@@ -46,7 +50,11 @@ TArray<UPrimitiveComponent*> UWaterBodyRiverComponent::GetStandardRenderableComp
 {
 	TArray<UPrimitiveComponent*> Result;
 	Result.Reserve(SplineMeshComponents.Num());
-	Algo::TransformIf(SplineMeshComponents, Result, [](USplineMeshComponent* SplineComp) { return (SplineComp != nullptr); }, [](USplineMeshComponent* SplineComp) { return SplineComp; });
+
+	Algo::TransformIf(SplineMeshComponents, Result, 
+		[](USplineMeshComponent* SplineComp) { return (SplineComp != nullptr); }, 
+		[](USplineMeshComponent* SplineComp) { return SplineComp; });
+
 	return Result;
 }
 
@@ -465,16 +473,10 @@ void UWaterBodyRiverComponent::UpdateSplineMesh(USplineMeshComponent* MeshComp,
 		check(StaticMesh != nullptr);
 		MeshComp->SetStaticMesh(StaticMesh);
 		MeshComp->SetMaterial(0, GetWaterMaterialInstance());
-		MeshComp->SetCollisionProfileName(GetCollisionProfileName());
-		// In the case of rivers, the USplineMeshComponent acts as both collision and visual component so we simply disable collision on them : 
-		MeshComp->SetGenerateOverlapEvents(bGenerateCollisions);
-		if (!bGenerateCollisions)
-		{
-			MeshComp->SetCollisionEnabled(ECollisionEnabled::NoCollision);
-		}
-		MeshComp->SetCastShadow(false);
 
-		MeshComp->bFillCollisionUnderneathForNavmesh = bGenerateCollisions && bFillCollisionUnderWaterBodiesForNavmesh;
+		CopySharedCollisionSettingsToComponent(MeshComp);
+		CopySharedNavigationSettingsToComponent(MeshComp);
+		MeshComp->SetCastShadow(false);
 
 		const bool bUpdateMesh = false;
 		const FVector StartScale = WaterSpline->GetScaleAtSplinePoint(StartSplinePointIndex);
diff --git a/Engine/Plugins/Experimental/Water/Source/Runtime/Public/WaterBodyActor.h b/Engine/Plugins/Experimental/Water/Source/Runtime/Public/WaterBodyActor.h
--- a/Engine/Plugins/Experimental/Water/Source/Runtime/Public/WaterBodyActor.h
+++ b/Engine/Plugins/Experimental/Water/Source/Runtime/Public/WaterBodyActor.h
@@ -264,9 +264,6 @@ protected:
 	UE_DEPRECATED(4.27, "Moved to WaterBodyComponent")
 	virtual bool IsHeightOffsetSupported() const final { return WaterBodyComponent->IsHeightOffsetSupported(); }
 
-	UE_DEPRECATED(4.27, "Moved to WaterBodyComponent")
-	virtual bool CanAffectNavigation() const final { return WaterBodyComponent->CanAffectNavigation(); }
-
 	UE_DEPRECATED(4.27, "Moved to WaterBodyComponent")
 	virtual void BeginUpdateWaterBody() final {}
 
diff --git a/Engine/Plugins/Experimental/Water/Source/Runtime/Public/WaterBodyComponent.h b/Engine/Plugins/Experimental/Water/Source/Runtime/Public/WaterBodyComponent.h
--- a/Engine/Plugins/Experimental/Water/Source/Runtime/Public/WaterBodyComponent.h
+++ b/Engine/Plugins/Experimental/Water/Source/Runtime/Public/WaterBodyComponent.h
@@ -112,7 +112,7 @@ public:
 	
 	/** Returns body's collision components */
 	UFUNCTION(BlueprintCallable, Category = Collision)
-	virtual TArray<UPrimitiveComponent*> GetCollisionComponents() const { return TArray<UPrimitiveComponent*>(); }
+	virtual TArray<UPrimitiveComponent*> GetCollisionComponents(bool bInOnlyEnabledComponents = true) const { return TArray<UPrimitiveComponent*>(); }
 
 	/** Retrieves the list of primitive components that this water body uses when not being rendered by the water mesh (e.g. the static mesh component used when WaterMeshOverride is specified) */
 	UFUNCTION(BlueprintCallable, Category = Rendering)
@@ -160,9 +160,6 @@ public:
 	/** Returns the unique id of this water body for accessing data in GPU buffers */
 	int32 GetWaterBodyIndex() const { return WaterBodyIndex; }
 	
-	/** Returns collision profile name */
-	FName GetCollisionProfileName() const { return CollisionProfileName; }
-
 	/** Returns water mesh override */
 	UStaticMesh* GetWaterMeshOverride() const { return WaterMeshOverride; }
 
@@ -326,6 +323,9 @@ public:
 
 	virtual FPrimitiveSceneProxy* CreateSceneProxy() override;
 	virtual void GetUsedMaterials(TArray<UMaterialInterface*>& OutMaterials, bool bGetDebugMaterials) const override;
+	
+	/** Set the navigation area class */
+	void SetNavAreaClass(TSubclassOf<UNavAreaBase> NewWaterNavAreaClass) { WaterNavAreaClass = NewWaterNavAreaClass; }
 protected:
 	//~ Begin USceneComponent Interface.
 	virtual void OnVisibilityChanged();
@@ -341,8 +341,8 @@ protected:
 	/** Returns whether the body support a height offset */
 	virtual bool IsHeightOffsetSupported() const;
 
-	/** Returns whether the body affects navigation */
-	virtual bool CanAffectNavigation() const { return bGenerateCollisions && bCanAffectNavigation; }
+	UE_DEPRECATED(5.1, "Please use CanEverAffectNavigation() instead.")
+	virtual bool CanAffectNavigation() const { return false; }
 
 	/** Called every time UpdateAll is called on WaterBody (prior to UpdateWaterBody) */
 	virtual void BeginUpdateWaterBody();
@@ -358,6 +358,9 @@ protected:
 	/** Returns navigation area class */
 	TSubclassOf<UNavAreaBase> GetNavAreaClass() const { return WaterNavAreaClass; }
 
+	void CopySharedCollisionSettingsToComponent(UPrimitiveComponent* InComponent);
+	void CopySharedNavigationSettingsToComponent(UPrimitiveComponent* InComponent);
+
 protected:
 	/** Computes the raw wave perturbation of the water height/normal */
 	virtual float GetWaveHeightAtPosition(const FVector& InPosition, float InWaterDepth, float InTime, FVector& OutNormal) const;
@@ -393,13 +396,17 @@ protected:
 	void CreateOrUpdateWaterInfoMID();
 	void PrepareCurrentPostProcessSettings();
 	void ApplyNavigationSettings();
+	void ApplyCollisionSettings();
 	void RequestGPUWaveDataUpdate();
 	EObjectFlags GetTransientMIDFlags() const; 
+	void DeprecateData();
 
 	virtual void Serialize(FArchive& Ar) override;
 	virtual void PostLoad() override;
 	virtual void OnComponentDestroyed(bool bDestroyingHierarchy) override;
 	virtual bool MoveComponentImpl(const FVector& Delta, const FQuat& NewRotation, bool bSweep, FHitResult* Hit = NULL, EMoveComponentFlags MoveFlags = MOVECOMP_NoFlags, ETeleportType Teleport = ETeleportType::None) override;
+	virtual void OnComponentCollisionSettingsChanged(bool bUpdateOverlaps) override;
+	virtual void OnGenerateOverlapEventsChanged() override;
 
 #if WITH_EDITOR
 	virtual void PreEditUndo() override;
@@ -448,10 +455,6 @@ public:
 	UPROPERTY(EditAnywhere, Category = Wave)
 	float MaxWaveHeightOffset = 0.0f;
 
-	/** Prevent navmesh generation under the water geometry */
-	UPROPERTY(EditAnywhere, Category = Navigation, meta = (EditCondition = "bGenerateCollisions"))
-	bool bFillCollisionUnderWaterBodiesForNavmesh;
-
 	/** Post process settings to apply when the camera goes underwater (only available when bGenerateCollisions is true because collisions are needed to detect if it's under water).
 	Note: Underwater post process material is setup using UnderwaterPostProcessMaterial. */
 	UPROPERTY(Category = Rendering, EditAnywhere, BlueprintReadWrite, meta = (EditCondition = "bGenerateCollisions && UnderwaterPostProcessMaterial != nullptr", DisplayAfter = "UnderwaterPostProcessMaterial"))
@@ -489,10 +492,6 @@ public:
 	UPROPERTY(Category = Terrain, EditAnywhere, BlueprintReadWrite)
 	bool bAffectsLandscape;
 
-	/** If true, one or more collision components associated with this water will be generated. Otherwise, this water body will only affect visuals. */
-	UPROPERTY(Category = Collision, EditAnywhere, BlueprintReadOnly)
-	bool bGenerateCollisions = true;
-
 protected:
 
 	/** Unique Id for accessing (wave, ... ) data in GPU buffers */
@@ -506,9 +505,6 @@ protected:
 	UPROPERTY(Category = Rendering, EditAnywhere, BlueprintReadOnly, meta = (ClampMin = "-8192", ClampMax = "8191"))
 	int32 OverlapMaterialPriority = 0;
 
-	UPROPERTY(Category = Collision, EditAnywhere, BlueprintReadOnly, meta = (EditCondition = "bGenerateCollisions"))
-	FName CollisionProfileName;
-
 	UPROPERTY(Transient)
 	UWaterSplineMetadata* WaterSplineMetadata;
 
@@ -534,9 +530,6 @@ protected:
 	UPROPERTY(Transient)
 	FPostProcessSettings CurrentPostProcessSettings;
 
-	UPROPERTY(EditAnywhere, BlueprintReadWrite, Category = Navigation, meta = (EditCondition = "bGenerateCollisions"))
-	bool bCanAffectNavigation;
-
 	// The navigation area class that will be generated on nav mesh
 	UPROPERTY(EditAnywhere, BlueprintReadWrite, Category = Navigation, meta = (EditCondition = "bCanAffectNavigation && bGenerateCollisions"))
 	TSubclassOf<UNavAreaBase> WaterNavAreaClass;
@@ -549,7 +542,19 @@ protected:
 	TArray<uint32> WaterBodyMeshIndices;
 
 #if WITH_EDITORONLY_DATA
+	UPROPERTY()
+	bool bFillCollisionUnderWaterBodiesForNavmesh_DEPRECATED;
+
+	UPROPERTY()
+	FName CollisionProfileName_DEPRECATED;
+
+	UPROPERTY()
+	bool bGenerateCollisions_DEPRECATED = true;
+
+	UPROPERTY()
+	bool bCanAffectNavigation_DEPRECATED;
+
 	UPROPERTY()
 	bool bOverrideWaterMesh_DEPRECATED;
-#endif
+#endif // WITH_EDITORONLY_DATA
 };
diff --git a/Engine/Plugins/Experimental/Water/Source/Runtime/Public/WaterBodyCustomComponent.h b/Engine/Plugins/Experimental/Water/Source/Runtime/Public/WaterBodyCustomComponent.h
--- a/Engine/Plugins/Experimental/Water/Source/Runtime/Public/WaterBodyCustomComponent.h
+++ b/Engine/Plugins/Experimental/Water/Source/Runtime/Public/WaterBodyCustomComponent.h
@@ -18,7 +18,7 @@ class WATER_API UWaterBodyCustomComponent : public UWaterBodyComponent
 public:
 	/** AWaterBody Interface */
 	virtual EWaterBodyType GetWaterBodyType() const override { return EWaterBodyType::Transition; }
-	virtual TArray<UPrimitiveComponent*> GetCollisionComponents() const override;
+	virtual TArray<UPrimitiveComponent*> GetCollisionComponents(bool bInOnlyEnabledComponents = true) const override;
 	virtual TArray<UPrimitiveComponent*> GetStandardRenderableComponents() const override;
 
 protected:
diff --git a/Engine/Plugins/Experimental/Water/Source/Runtime/Public/WaterBodyLakeComponent.h b/Engine/Plugins/Experimental/Water/Source/Runtime/Public/WaterBodyLakeComponent.h
--- a/Engine/Plugins/Experimental/Water/Source/Runtime/Public/WaterBodyLakeComponent.h
+++ b/Engine/Plugins/Experimental/Water/Source/Runtime/Public/WaterBodyLakeComponent.h
@@ -21,7 +21,7 @@ class WATER_API UWaterBodyLakeComponent : public UWaterBodyComponent
 public:
 	/** UWaterBodyComponent Interface */
 	virtual EWaterBodyType GetWaterBodyType() const override { return EWaterBodyType::Lake; }
-	virtual TArray<UPrimitiveComponent*> GetCollisionComponents() const override;
+	virtual TArray<UPrimitiveComponent*> GetCollisionComponents(bool bInOnlyEnabledComponents = true) const override;
 	virtual TArray<UPrimitiveComponent*> GetStandardRenderableComponents() const override;
 	virtual void GenerateWaterBodyMesh() override;
 
diff --git a/Engine/Plugins/Experimental/Water/Source/Runtime/Public/WaterBodyOceanComponent.h b/Engine/Plugins/Experimental/Water/Source/Runtime/Public/WaterBodyOceanComponent.h
--- a/Engine/Plugins/Experimental/Water/Source/Runtime/Public/WaterBodyOceanComponent.h
+++ b/Engine/Plugins/Experimental/Water/Source/Runtime/Public/WaterBodyOceanComponent.h
@@ -20,7 +20,7 @@ class WATER_API UWaterBodyOceanComponent : public UWaterBodyComponent
 public:
 	/** UWaterBodyComponent Interface */
 	virtual EWaterBodyType GetWaterBodyType() const override { return EWaterBodyType::Ocean; }
-	virtual TArray<UPrimitiveComponent*> GetCollisionComponents() const override;
+	virtual TArray<UPrimitiveComponent*> GetCollisionComponents(bool bInOnlyEnabledComponents = true) const override;
 	virtual FVector GetCollisionExtents() const override { return CollisionExtents; }
 	virtual void SetHeightOffset(float InHeightOffset) override;
 	virtual float GetHeightOffset() const override { return HeightOffset; }
diff --git a/Engine/Plugins/Experimental/Water/Source/Runtime/Public/WaterBodyRiverComponent.h b/Engine/Plugins/Experimental/Water/Source/Runtime/Public/WaterBodyRiverComponent.h
--- a/Engine/Plugins/Experimental/Water/Source/Runtime/Public/WaterBodyRiverComponent.h
+++ b/Engine/Plugins/Experimental/Water/Source/Runtime/Public/WaterBodyRiverComponent.h
@@ -19,7 +19,7 @@ class WATER_API UWaterBodyRiverComponent : public UWaterBodyComponent
 public:
 	/** UWaterBodyComponent Interface */
 	virtual EWaterBodyType GetWaterBodyType() const override { return EWaterBodyType::River; }
-	virtual TArray<UPrimitiveComponent*> GetCollisionComponents() const override;
+	virtual TArray<UPrimitiveComponent*> GetCollisionComponents(bool bInOnlyEnabledComponents = true) const override;
 	virtual TArray<UPrimitiveComponent*> GetStandardRenderableComponents() const override;
 	virtual UMaterialInstanceDynamic* GetRiverToLakeTransitionMaterialInstance() override;
 	virtual UMaterialInstanceDynamic* GetRiverToOceanTransitionMaterialInstance() override;
diff --git a/Engine/Source/Runtime/Core/Public/UObject/FortniteMainBranchObjectVersion.h b/Engine/Source/Runtime/Core/Public/UObject/FortniteMainBranchObjectVersion.h
--- a/Engine/Source/Runtime/Core/Public/UObject/FortniteMainBranchObjectVersion.h
+++ b/Engine/Source/Runtime/Core/Public/UObject/FortniteMainBranchObjectVersion.h
@@ -217,6 +217,9 @@ struct CORE_API FFortniteMainBranchObjectVersion
 		// Remapped bEvaluateWorldPositionOffset to bEvaluateWorldPositionOffsetInRayTracing
 		RemappedEvaluateWorldPositionOffsetInRayTracing,
 
+		// Water body collision settings are now those of the base UPrimitiveComponent, rather than duplicated in UWaterBodyComponent
+		WaterBodyComponentCollisionSettingsRefactor,
+
 		// -----<new versions can be added above this line>-------------------------------------------------
 		VersionPlusOne,
 		LatestVersion = VersionPlusOne - 1
diff --git a/Engine/Source/Runtime/Engine/Classes/Components/PrimitiveComponent.h b/Engine/Source/Runtime/Engine/Classes/Components/PrimitiveComponent.h
--- a/Engine/Source/Runtime/Engine/Classes/Components/PrimitiveComponent.h
+++ b/Engine/Source/Runtime/Engine/Classes/Components/PrimitiveComponent.h
@@ -2478,6 +2478,9 @@ protected:
 	/** Called when the BodyInstance ResponseToChannels, CollisionEnabled or bNotifyRigidBodyCollision changes, in case subclasses want to use that information. */
 	virtual void OnComponentCollisionSettingsChanged(bool bUpdateOverlaps=true);
 
+	/** Called when bGenerateOverlapEvents changes, in case subclasses want to use that information. */
+	virtual void OnGenerateOverlapEventsChanged();
+
 	/** Ends all current component overlaps. Generally used when destroying this component or when it can no longer generate overlaps. */
 	void ClearComponentOverlaps(bool bDoNotifies, bool bSkipNotifySelf);
 
diff --git a/Engine/Source/Runtime/Engine/Private/Components/PrimitiveComponent.cpp b/Engine/Source/Runtime/Engine/Private/Components/PrimitiveComponent.cpp
--- a/Engine/Source/Runtime/Engine/Private/Components/PrimitiveComponent.cpp
+++ b/Engine/Source/Runtime/Engine/Private/Components/PrimitiveComponent.cpp
@@ -3651,10 +3651,16 @@ void UPrimitiveComponent::SetGenerateOverlapEvents(bool bInGenerateOverlapEvents
 	if (bGenerateOverlapEvents != bInGenerateOverlapEvents)
 	{
 		bGenerateOverlapEvents = bInGenerateOverlapEvents;
-		ClearSkipUpdateOverlaps();
+
+		OnGenerateOverlapEventsChanged();
 	}
 }
 
+void UPrimitiveComponent::OnGenerateOverlapEventsChanged()
+{
+	ClearSkipUpdateOverlaps();
+}
+
 void UPrimitiveComponent::SetLightingChannels(bool bChannel0, bool bChannel1, bool bChannel2)
 {
 	if (bChannel0 != LightingChannels.bChannel0 ||
