Remove deprecated water functions from <4.27

diff --git a/Engine/Plugins/Experimental/Water/Source/Editor/Private/WaterLandscapeBrush.cpp b/Engine/Plugins/Experimental/Water/Source/Editor/Private/WaterLandscapeBrush.cpp
--- a/Engine/Plugins/Experimental/Water/Source/Editor/Private/WaterLandscapeBrush.cpp
+++ b/Engine/Plugins/Experimental/Water/Source/Editor/Private/WaterLandscapeBrush.cpp
@@ -414,11 +414,6 @@ void AWaterLandscapeBrush::BlueprintWaterBodyChanged_Implementation(AActor* Acto
 	BlueprintWaterBodyChanged_Native(Actor);
 }
 
-void AWaterLandscapeBrush::SetWaterBodyCache(AWaterBody* WaterBody, UObject* InCache)
-{
-	SetActorCache(WaterBody, InCache);
-}
-
 void AWaterLandscapeBrush::SetActorCache(AActor* InActor, UObject* InCache)
 {
 	if (!InCache)
@@ -431,11 +426,6 @@ void AWaterLandscapeBrush::SetActorCache(AActor* InActor, UObject* InCache)
 }
 
 
-UObject* AWaterLandscapeBrush::GetWaterBodyCache(AWaterBody* WaterBody, TSubclassOf<UObject> CacheClass) const
-{
-	return GetActorCache(WaterBody, CacheClass);
-}
-
 UObject* AWaterLandscapeBrush::GetActorCache(AActor* InActor, TSubclassOf<UObject> CacheClass) const
 {
 	UObject* const* ValuePtr = Cache.Find(TWeakObjectPtr<AActor>(InActor));
@@ -446,11 +436,6 @@ UObject* AWaterLandscapeBrush::GetActorCache(AActor* InActor, TSubclassOf<UObjec
 	return nullptr;
 }
 
-void AWaterLandscapeBrush::ClearWaterBodyCache(AWaterBody* WaterBody)
-{
-	ClearActorCache(WaterBody);
-}
-
 void AWaterLandscapeBrush::ClearActorCache(AActor* InActor)
 {
 	Cache.Remove(TWeakObjectPtr<AActor>(InActor));
diff --git a/Engine/Plugins/Experimental/Water/Source/Editor/Private/WaterLandscapeBrush.h b/Engine/Plugins/Experimental/Water/Source/Editor/Private/WaterLandscapeBrush.h
--- a/Engine/Plugins/Experimental/Water/Source/Editor/Private/WaterLandscapeBrush.h
+++ b/Engine/Plugins/Experimental/Water/Source/Editor/Private/WaterLandscapeBrush.h
@@ -46,14 +46,17 @@ public:
 	void BlueprintWaterBodyChanged(AActor* Actor);
 	virtual void BlueprintWaterBodyChanged_Native(AActor* Actor) {}
 
+	UE_DEPRECATED(4.27, "Use SetActorCache instead")
 	UFUNCTION(BlueprintCallable, Category = "Cache", meta = (DeprecatedFunction, DeprecationMessage="Use SetActorCache instead"))
-	void SetWaterBodyCache(AWaterBody* WaterBody, UObject* Cache);
+	void SetWaterBodyCache(AWaterBody* WaterBody, UObject* InCache) {}
 	
+	UE_DEPRECATED(4.27, "Use GetActorCache instead")
 	UFUNCTION(BlueprintCallable, Category = "Cache", meta = (DeterminesOutputType = "CacheClass", DeprecatedFunction, DeprecationMessage = "Use GetActorCache instead"))
-	UObject* GetWaterBodyCache(AWaterBody* WaterBody, TSubclassOf<UObject> CacheClass) const;
+	UObject* GetWaterBodyCache(AWaterBody* WaterBody, TSubclassOf<UObject> CacheClass) const { return nullptr; }
 
+	UE_DEPRECATED(4.27, "Use ClearActorCache instead")
 	UFUNCTION(BlueprintCallable, Category = "Cache", meta = (DeprecatedFunction, DeprecationMessage = "Use ClearActorCache instead"))
-	void ClearWaterBodyCache(AWaterBody* WaterBody);
+	void ClearWaterBodyCache(AWaterBody* WaterBody) {}
 
 	UFUNCTION(BlueprintCallable, Category = "Cache")
 	void SetActorCache(AActor* InActor, UObject* InCache);
diff --git a/Engine/Plugins/Experimental/Water/Source/Runtime/Private/NiagaraWaterFunctionLibrary.cpp b/Engine/Plugins/Experimental/Water/Source/Runtime/Private/NiagaraWaterFunctionLibrary.cpp
--- a/Engine/Plugins/Experimental/Water/Source/Runtime/Private/NiagaraWaterFunctionLibrary.cpp
+++ b/Engine/Plugins/Experimental/Water/Source/Runtime/Private/NiagaraWaterFunctionLibrary.cpp
@@ -6,17 +6,6 @@
 #include "NiagaraDataInterfaceWater.h"
 #include "WaterModule.h"
 
-void UNiagaraWaterFunctionLibrary::SetWaterBody(UNiagaraComponent* NiagaraSystem, const FString& OverrideName, AWaterBody* WaterBody)
-{
-	if (!WaterBody)
-	{
-		UE_LOG(LogWater, Warning, TEXT("WaterBody in \"Set Water Body Component\" is NULL, OverrideName \"%s\" and NiagaraSystem \"%s\", skipping."), *OverrideName, *NiagaraSystem->GetOwner()->GetName());
-		return;
-	}
-
-	SetWaterBodyComponent(NiagaraSystem, OverrideName, WaterBody->GetWaterBodyComponent());
-}
-
 void UNiagaraWaterFunctionLibrary::SetWaterBodyComponent(UNiagaraComponent* NiagaraSystem, const FString& OverrideName, UWaterBodyComponent* WaterBodyComponent)
 {
 	if (!NiagaraSystem)
diff --git a/Engine/Plugins/Experimental/Water/Source/Runtime/Private/WaterBodyCustomActor.cpp b/Engine/Plugins/Experimental/Water/Source/Runtime/Private/WaterBodyCustomActor.cpp
--- a/Engine/Plugins/Experimental/Water/Source/Runtime/Private/WaterBodyCustomActor.cpp
+++ b/Engine/Plugins/Experimental/Water/Source/Runtime/Private/WaterBodyCustomActor.cpp
@@ -37,13 +37,6 @@ void AWaterBodyCustom::PostLoad()
 #endif // WITH_EDITORONLY_DATA
 }
 
-#if WITH_EDITOR
-bool AWaterBodyCustom::IsIconVisible() const
-{
-	return (WaterBodyComponent->GetWaterMeshOverride() == nullptr);
-}
-#endif // WITH_EDITOR
-
 // ----------------------------------------------------------------------------------
 
 UDEPRECATED_CustomMeshGenerator::UDEPRECATED_CustomMeshGenerator(const FObjectInitializer& ObjectInitializer)
diff --git a/Engine/Plugins/Experimental/Water/Source/Runtime/Private/WaterBodyCustomComponent.cpp b/Engine/Plugins/Experimental/Water/Source/Runtime/Private/WaterBodyCustomComponent.cpp
--- a/Engine/Plugins/Experimental/Water/Source/Runtime/Private/WaterBodyCustomComponent.cpp
+++ b/Engine/Plugins/Experimental/Water/Source/Runtime/Private/WaterBodyCustomComponent.cpp
@@ -131,10 +131,16 @@ TArray<TSharedRef<FTokenizedMessage>> UWaterBodyCustomComponent::CheckWaterBodyS
 
 	return StatusMessages;
 }
+
 const TCHAR* UWaterBodyCustomComponent::GetWaterSpriteTextureName() const
 {
 	return TEXT("/Water/Icons/WaterBodyCustomSprite");
 }
+
+bool UWaterBodyCustomComponent::IsIconVisible() const
+{
+	return (GetWaterMeshOverride() == nullptr);
+}
 #endif // WITH_EDITOR
 
 #undef LOCTEXT_NAMESPACE
diff --git a/Engine/Plugins/Experimental/Water/Source/Runtime/Public/NiagaraWaterFunctionLibrary.h b/Engine/Plugins/Experimental/Water/Source/Runtime/Public/NiagaraWaterFunctionLibrary.h
--- a/Engine/Plugins/Experimental/Water/Source/Runtime/Public/NiagaraWaterFunctionLibrary.h
+++ b/Engine/Plugins/Experimental/Water/Source/Runtime/Public/NiagaraWaterFunctionLibrary.h
@@ -18,11 +18,11 @@ class WATER_API UNiagaraWaterFunctionLibrary : public UBlueprintFunctionLibrary
 	GENERATED_BODY()
 
 	/** Sets the water body on the Niagra water data interface on a Niagara particle system*/
+	UE_DEPRECATED(4.27, "Use SetWaterBodyComponent instead.")
 	UFUNCTION(BlueprintCallable, Category = Water, meta = (DisplayName = "Set Water Body"), meta = (DeprecatedFunction, DeprecationMessage="Use SetWaterBodyComponent instead"))
-	static void SetWaterBody(UNiagaraComponent* NiagaraSystem, const FString& OverrideName, AWaterBody* WaterBody);
+	static void SetWaterBody(UNiagaraComponent* NiagaraSystem, const FString& OverrideName, AWaterBody* WaterBody) {}
 
 	/** Sets the water body component on the Niagra water data interface on a Niagara particle system*/
 	UFUNCTION(BlueprintCallable, Category = Water)
 	static void SetWaterBodyComponent(UNiagaraComponent* NiagaraSystem, const FString& OverrideName, UWaterBodyComponent* WaterBodyComponent);
-
 };
\ No newline at end of file
diff --git a/Engine/Plugins/Experimental/Water/Source/Runtime/Public/WaterBodyActor.h b/Engine/Plugins/Experimental/Water/Source/Runtime/Public/WaterBodyActor.h
--- a/Engine/Plugins/Experimental/Water/Source/Runtime/Public/WaterBodyActor.h
+++ b/Engine/Plugins/Experimental/Water/Source/Runtime/Public/WaterBodyActor.h
@@ -136,158 +136,50 @@ protected:
 
 public:
 	UE_DEPRECATED(4.27, "Moved to WaterBodyComponent")
-	virtual bool HasWaves() const final { return WaterBodyComponent->HasWaves(); }
-
-	UE_DEPRECATED(4.27, "Moved to WaterBodyComponent")
-	virtual bool IsWaveSupported() const final { return WaterBodyComponent->IsWaveSupported(); }
-
-	UE_DEPRECATED(4.27, "Moved to WaterBodyComponent")
-	virtual TArray<UPrimitiveComponent*> GetCollisionComponents() const final { return WaterBodyComponent->GetCollisionComponents(); }
-
-	UE_DEPRECATED(4.27, "Moved to WaterBodyComponent")
-	virtual FVector GetCollisionExtents() const final { return WaterBodyComponent->GetCollisionExtents(); }
-
-	UE_DEPRECATED(4.27, "Moved to WaterBodyComponent")
-	virtual void SetHeightOffset(float InHeightOffset) final { WaterBodyComponent->SetHeightOffset(InHeightOffset); }
-
-	UE_DEPRECATED(4.27, "Moved to WaterBodyComponent")
-	virtual float GetHeightOffset() const final { return WaterBodyComponent->GetHeightOffset(); }
-
-	UE_DEPRECATED(4.27, "Moved to WaterBodyComponent")
-	void SetWaterMeshOverride(UStaticMesh* InMesh) { WaterBodyComponent->SetWaterMeshOverride(InMesh); }
-
 	UFUNCTION(BlueprintCallable, Category = Rendering, meta = (DeprecatedFunction))
 	virtual UMaterialInstanceDynamic* GetRiverToLakeTransitionMaterialInstance() final { return WaterBodyComponent->GetRiverToLakeTransitionMaterialInstance(); }
 
+	UE_DEPRECATED(4.27, "Moved to WaterBodyComponent")
 	UFUNCTION(BlueprintCallable, Category = Rendering, meta = (DeprecatedFunction))
 	virtual UMaterialInstanceDynamic* GetRiverToOceanTransitionMaterialInstance() final { return WaterBodyComponent->GetRiverToOceanTransitionMaterialInstance(); }
 
+	UE_DEPRECATED(4.27, "Moved to WaterBodyComponent")
 	UFUNCTION(BlueprintCallable, Category = Rendering, meta = (DeprecatedFunction))
 	void SetWaterMaterial(UMaterialInterface* InMaterial) { WaterBodyComponent->SetWaterMaterial(InMaterial); }
 
+	UE_DEPRECATED(4.27, "Moved to WaterBodyComponent")
 	UFUNCTION(BlueprintCallable, Category = Rendering, meta = (DeprecatedFunction))
 	UMaterialInstanceDynamic* GetWaterMaterialInstance() { return WaterBodyComponent->GetWaterMaterialInstance(); }
 
 	UE_DEPRECATED(4.27, "Moved to WaterBodyComponent")
-	void SetUnderwaterPostProcessMaterial(UMaterialInterface* InMaterial) { WaterBodyComponent->SetUnderwaterPostProcessMaterial(InMaterial); }
-
-	UE_DEPRECATED(4.27, "Moved to WaterBodyComponent")
-	bool ShouldGenerateWaterMeshTile() const { return WaterBodyComponent->ShouldGenerateWaterMeshTile(); }
-
-	UE_DEPRECATED(4.27, "Moved to WaterBodyComponent")
-	FVector GetWaterNavCollisionOffset() const { return WaterBodyComponent->GetWaterNavCollisionOffset(); }
-
-	UE_DEPRECATED(4.27, "Moved to WaterBodyComponent")
-	int32 GetOverlapMaterialPriority() const { return WaterBodyComponent->OverlapMaterialPriority; }
-
-	UE_DEPRECATED(4.27, "Moved to WaterBodyComponent")
-	float GetChannelDepth() const { return WaterBodyComponent->CurveSettings.ChannelDepth; }
-
-	UE_DEPRECATED(4.27, "Moved to WaterBodyComponent")
-	void AddExclusionVolume(AWaterBodyExclusionVolume* InExclusionVolume) { WaterBodyComponent->AddExclusionVolume(InExclusionVolume); }
-
-	UE_DEPRECATED(4.27, "Moved to WaterBodyComponent")
-	void RemoveExclusionVolume(AWaterBodyExclusionVolume* InExclusionVolume) { WaterBodyComponent->RemoveExclusionVolume(InExclusionVolume); }
-
-	UE_DEPRECATED(4.27, "Moved to WaterBodyComponent")
-	FPostProcessVolumeProperties GetPostProcessProperties() const { return WaterBodyComponent->GetPostProcessProperties(); }
-
-	UE_DEPRECATED(4.27, "Moved to WaterBodyComponent")
-	virtual FWaterBodyQueryResult QueryWaterInfoClosestToWorldLocation(const FVector& InWorldLocation, EWaterBodyQueryFlags InQueryFlags, const TOptional<float>& InSplineInputKey = TOptional<float>()) const
-	{
-		return WaterBodyComponent->QueryWaterInfoClosestToWorldLocation(InWorldLocation, InQueryFlags, InSplineInputKey);
-	}
-
-	UE_DEPRECATED(4.27, "Moved to WaterBodyComponent")
-	float FindInputKeyClosestToWorldLocation(const FVector& WorldLocation) const { return WaterBodyComponent->FindInputKeyClosestToWorldLocation(WorldLocation); }
-
 	UFUNCTION(BlueprintCallable, Category = WaterBody, meta = (DeprecatedFunction))
 	virtual float GetWaterVelocityAtSplineInputKey(float InKey) const { return WaterBodyComponent->GetWaterVelocityAtSplineInputKey(InKey); }
-	
+
+	UE_DEPRECATED(4.27, "Moved to WaterBodyComponent")
 	UFUNCTION(BlueprintCallable, Category = WaterBody, meta = (DeprecatedFunction))
 	virtual FVector GetWaterVelocityVectorAtSplineInputKey(float InKey) const { return WaterBodyComponent->GetWaterVelocityVectorAtSplineInputKey(InKey); }
-	
+
+	UE_DEPRECATED(4.27, "Moved to WaterBodyComponent")
 	UFUNCTION(BlueprintCallable, Category = WaterBody, meta = (DeprecatedFunction))
 	virtual float GetAudioIntensityAtSplineInputKey(float InKey) const { return WaterBodyComponent->GetAudioIntensityAtSplineInputKey(InKey); }
 
+	UE_DEPRECATED(4.27, "Moved to WaterBodyComponent")
 	UFUNCTION(BlueprintCallable, Category = Water, meta = (DeprecatedFunction))
 	TArray<AWaterBodyIsland*> GetIslands() const { return WaterBodyComponent->GetIslands(); }
 
 	UE_DEPRECATED(4.27, "Moved to WaterBodyComponent")
-	bool ContainsIsland(TLazyObjectPtr<AWaterBodyIsland> Island) const { return WaterBodyComponent->ContainsIsland(Island); }
-
 	UFUNCTION(BlueprintCallable, Category = Water, meta = (DeprecatedFunction))
 	TArray<AWaterBodyExclusionVolume*> GetExclusionVolumes() const { return WaterBodyComponent->GetExclusionVolumes(); }
 
 	UE_DEPRECATED(4.27, "Moved to WaterBodyComponent")
-	bool ContainsExclusionVolume(TLazyObjectPtr<AWaterBodyExclusionVolume> InExclusionVolume) const { return WaterBodyComponent->ContainsExclusionVolume(InExclusionVolume); }
-
 	UFUNCTION(BlueprintCallable, Category=Water, meta = (DeprecatedFunction))
 	void OnWaterBodyChanged(bool bShapeOrPositionChanged, bool bWeightmapSettingsChanged = false) { return WaterBodyComponent->OnWaterBodyChanged(bShapeOrPositionChanged, bWeightmapSettingsChanged); }
 
-	UE_DEPRECATED(4.27, "Moved to WaterBodyComponent")
-	bool GetWaveInfoAtPosition(const FVector& InPosition, float InWaterDepth, bool bInSimpleWaves, FWaveInfo& InOutWaveInfo) const
-	{
-		return WaterBodyComponent->GetWaveInfoAtPosition(InPosition, InWaterDepth, bInSimpleWaves, InOutWaveInfo);
-	}
-
-	UE_DEPRECATED(4.27, "Moved to WaterBodyComponent")
-	virtual void SetDynamicParametersOnMID(UMaterialInstanceDynamic* InMID) final { WaterBodyComponent->SetDynamicParametersOnMID(InMID); }
-
-	UE_DEPRECATED(4.27, "Moved to WaterBodyComponent")
-	bool IsWorldLocationInExclusionVolume(const FVector& InWorldLocation) const { return WaterBodyComponent->IsWorldLocationInExclusionVolume(InWorldLocation); }
-
-	UE_DEPRECATED(4.27, "Moved to WaterBodyComponent")
-	virtual void UpdateMaterialInstances() final {}
-
-	UE_DEPRECATED(4.27, "Moved to WaterBodyComponent")
-	virtual float GetWaveReferenceTime() const final { return WaterBodyComponent->GetWaveReferenceTime(); }
-
-	UE_DEPRECATED(4.27, "Moved to WaterBodyComponent")
-	virtual ALandscapeProxy* FindLandscape() const final { return WaterBodyComponent->FindLandscape(); }
-
 #if WITH_EDITOR
 	UE_DEPRECATED(5.1, "Moved to WaterBodyComponent")
 	virtual bool IsIconVisible() const { return WaterBodyComponent->IsIconVisible(); }
 #endif // WITH_EDITOR
-
 protected:
-	UE_DEPRECATED(4.27, "Moved to WaterBodyComponent")
-	virtual bool IsBodyDynamic() const { return WaterBodyComponent->IsBodyDynamic(); }
-
-	UE_DEPRECATED(4.27, "Moved to WaterBodyComponent")
-	virtual bool IsFlatSurface() const { return WaterBodyComponent->IsFlatSurface(); }
-
-	UE_DEPRECATED(4.27, "Moved to WaterBodyComponent")
-	virtual bool IsWaterSplineClosedLoop() const final { return WaterBodyComponent->IsWaterSplineClosedLoop(); }
-
-	UE_DEPRECATED(4.27, "Moved to WaterBodyComponent")
-	virtual bool IsHeightOffsetSupported() const final { return WaterBodyComponent->IsHeightOffsetSupported(); }
-
-	UE_DEPRECATED(4.27, "Moved to WaterBodyComponent")
-	virtual void BeginUpdateWaterBody() final {}
-
-	UE_DEPRECATED(4.27, "Moved to WaterBodyComponent")
-	virtual void UpdateWaterBody(bool bWithExclusionVolumes) final { WaterBodyComponent->UpdateWaterBody(bWithExclusionVolumes); }
-
-	UE_DEPRECATED(4.27, "Moved to WaterBodyComponent")
-	virtual float GetConstantSurfaceZ() const final { return WaterBodyComponent->GetConstantSurfaceZ(); }
-
-	UE_DEPRECATED(4.27, "Moved to WaterBodyComponent")
-	virtual float GetConstantDepth() const final { return WaterBodyComponent->GetConstantDepth(); };
-
-	UE_DEPRECATED(4.27, "Moved to WaterBodyComponent")
-	virtual void GetSurfaceMinMaxZ(float& OutMinZ, float& OutMaxZ) const final { return WaterBodyComponent->GetSurfaceMinMaxZ(OutMinZ, OutMaxZ); }
-
-	UE_DEPRECATED(4.27, "Moved to WaterBodyComponent")
-	virtual float GetWaveHeightAtPosition(const FVector& InPosition, float InWaterDepth, float InTime, FVector& OutNormal) const final { return WaterBodyComponent->GetWaveHeightAtPosition(InPosition, InWaterDepth, InTime, OutNormal); }
-
-	UE_DEPRECATED(4.27, "Moved to WaterBodyComponent")
-	virtual float GetSimpleWaveHeightAtPosition(const FVector& InPosition, float InWaterDepth, float InTime) const final { return WaterBodyComponent->GetSimpleWaveHeightAtPosition(InPosition, InWaterDepth, InTime); }
-
-	UE_DEPRECATED(4.27, "Moved to WaterBodyComponent")
-	virtual float GetWaveAttenuationFactor(const FVector& InPosition, float InWaterDepth) const final { return WaterBodyComponent->GetWaveAttenuationFactor(InPosition, InWaterDepth); }
-
 	friend class UWaterBodyComponent;
 #if WITH_EDITORONLY_DATA
 	UPROPERTY()
@@ -361,5 +253,5 @@ protected:
 
 	UPROPERTY()
 	float ShapeDilation_DEPRECATED = 4096.0f;
-#endif
+#endif // WITH_EDITORONLY_DATA
 };
diff --git a/Engine/Plugins/Experimental/Water/Source/Runtime/Public/WaterBodyCustomActor.h b/Engine/Plugins/Experimental/Water/Source/Runtime/Public/WaterBodyCustomActor.h
--- a/Engine/Plugins/Experimental/Water/Source/Runtime/Public/WaterBodyCustomActor.h
+++ b/Engine/Plugins/Experimental/Water/Source/Runtime/Public/WaterBodyCustomActor.h
@@ -30,10 +30,6 @@ class WATER_API AWaterBodyCustom : public AWaterBody
 protected:
 	virtual void PostLoad() override;
 
-#if WITH_EDITOR
-	virtual bool IsIconVisible() const override;
-#endif
-
 #if WITH_EDITORONLY_DATA
 	UPROPERTY(NonPIEDuplicateTransient)
 	UDEPRECATED_CustomMeshGenerator* CustomGenerator_DEPRECATED;
diff --git a/Engine/Plugins/Experimental/Water/Source/Runtime/Public/WaterBodyCustomComponent.h b/Engine/Plugins/Experimental/Water/Source/Runtime/Public/WaterBodyCustomComponent.h
--- a/Engine/Plugins/Experimental/Water/Source/Runtime/Public/WaterBodyCustomComponent.h
+++ b/Engine/Plugins/Experimental/Water/Source/Runtime/Public/WaterBodyCustomComponent.h
@@ -33,7 +33,7 @@ protected:
 
 	virtual const TCHAR* GetWaterSpriteTextureName() const override;
 
-	virtual bool IsIconVisible() const override { return false; }
+	virtual bool IsIconVisible() const override;
 #endif // WITH_EDITOR
 
 protected:
diff --git a/Engine/Plugins/Experimental/Water/Source/Runtime/Public/WaterIconHelper.h b/Engine/Plugins/Experimental/Water/Source/Runtime/Public/WaterIconHelper.h
--- a/Engine/Plugins/Experimental/Water/Source/Runtime/Public/WaterIconHelper.h
+++ b/Engine/Plugins/Experimental/Water/Source/Runtime/Public/WaterIconHelper.h
@@ -19,13 +19,6 @@ struct WATER_API FWaterIconHelper
 	{
 		return EnsureSpriteComponentCreated_Internal(Actor, ActorType::StaticClass(), InIconTextureName);
 	}
-	
-	template <typename ActorType>
-	UE_DEPRECATED(4.27, "void EnsureSpriteComponentCreated(Actor, InIconTextureName, InDisplayName) is deprecated. Use EnsureSpriteComponentCreated(Actor, InIconTextureName) instead.")
-	static UBillboardComponent* EnsureSpriteComponentCreated(ActorType* Actor, const TCHAR* InIconTextureName, const FText& InDisplayName)
-	{
-		return EnsureSpriteComponentCreated_Internal(Actor, ActorType::StaticClass(), InIconTextureName);
-	}
 
 	/** Updates the texture/scale/position of the actor's billboard component, if any */
 	static void UpdateSpriteComponent(AActor* Actor, UTexture2D* InTexture);
