Reimplemented the system which allows Water to select the optimal landscape LOD to prevent overdraw.

diff --git a/Engine/Plugins/Experimental/Water/Source/Runtime/Private/WaterInfoRendering.cpp b/Engine/Plugins/Experimental/Water/Source/Runtime/Private/WaterInfoRendering.cpp
--- a/Engine/Plugins/Experimental/Water/Source/Runtime/Private/WaterInfoRendering.cpp
+++ b/Engine/Plugins/Experimental/Water/Source/Runtime/Private/WaterInfoRendering.cpp
@@ -17,6 +17,8 @@
 #include "WaterZoneActor.h"
 #include "Rendering/NaniteStreamingManager.h"
 #include "SceneViewExtension.h"
+#include "LandscapeRender.h"
+#include "LandscapeModule.h"
 
 static int32 RenderCaptureNextWaterInfoDraws = 0;
 static FAutoConsoleVariableRef CVarRenderCaptureNextWaterInfoDraws(
@@ -34,6 +36,7 @@ struct FUpdateWaterInfoParams
 	FRenderTarget* RenderTarget;
 	FTexture* OutputTexture;
 
+	FVector2D WaterZoneExtents;
 	FVector2f WaterHeightExtents;
 	float GroundZMin;
 	float CaptureZ;
@@ -291,10 +294,104 @@ static void SetWaterBodiesWithinWaterInfoPass(FSceneRenderer* SceneRenderer, boo
 	}
 }
 
+/**
+ * Sets up the landscape LOD override for water info rendering and restores it when the instance goes out of scope.
+ */
+struct FScopedLandscapeLODOverride
+{
+public:
+	FScopedLandscapeLODOverride(const FUpdateWaterInfoParams& InParams)	
+		: Params(InParams)
+	{
+		if (ILandscapeModule* LandscapeModule = FModuleManager::GetModulePtr<ILandscapeModule>("Landscape"))
+		{
+			if (TSharedPtr<FLandscapeSceneViewExtension, ESPMode::ThreadSafe> LandscapeViewExtension = LandscapeModule->GetLandscapeSceneViewExtension())
+			{
+				LandscapeRenderSystems = &LandscapeViewExtension->GetLandscapeRenderSystems();
+			}
+		}
+
+		SetOptimalLandscapeLODOverrides();
+	}
+
+	~FScopedLandscapeLODOverride()
+	{
+		RestoreLandscapeLODOverrides();
+	}
+
+private:
+	FScopedLandscapeLODOverride(const FScopedLandscapeLODOverride&) = delete;
+	FScopedLandscapeLODOverride& operator=(const FScopedLandscapeLODOverride&) = delete;
+
+	void SetOptimalLandscapeLODOverrides()
+	{
+		if (!LandscapeRenderSystems) return;
+		
+		// In order to prevent overdrawing the landscape components, we compute the lowest-detailed LOD level which satisfies the pixel coverage of the Water Info texture
+		// and force it on all landscape components. This override is set different per Landscape actor in case there are multiple under the same water zone.
+		//
+		// Ex: If the WaterInfoTexture only has 1 pixel per 100 units, and the highest landscape LOD has 1 vertex per 20 units, we don't need to use the maximum landscape LOD
+		// and can force a lower level of detail (in this case LOD2) while still satisfying the resolution of the water info texture.
+
+		const double MinWaterInfoTextureExtent = (double)FMath::Min(Params.OutputTexture->GetSizeX(), Params.OutputTexture->GetSizeY());
+		const double MaxWaterZoneExtent = FMath::Max(Params.WaterZoneExtents.X, Params.WaterZoneExtents.Y);
+		const double WaterInfoUnitsPerPixel =  MaxWaterZoneExtent / MinWaterInfoTextureExtent;
+
+		for (const TPair<uint32, FLandscapeRenderSystem*>& Pair : *LandscapeRenderSystems)
+		{
+			FLandscapeRenderSystem* LandscapeRenderSystem = Pair.Value;
+			int32 OptimalLODLevel = -1;
+
+			// All components within the same landscape (and thus its render system) should have the same number of quads and the same extent.
+			// therefore we can simply find the first component and compute its optimal LOD level.
+			for (FLandscapeSectionInfo* LandscapeSectionInfo : LandscapeRenderSystem->SectionInfos)
+			{
+				if (LandscapeSectionInfo != nullptr)
+				{
+					const double LandscapeComponentUnitsPerVertex = LandscapeSectionInfo->ComputeSectionResolution();
+					if (LandscapeComponentUnitsPerVertex <= 0.f)
+					{
+						continue;
+					}
+
+					// Derived from:
+					// LandscapeComponentUnitsPerVertex * 2 ^ (LODLevel) <= WaterInfoUnitsPerPixel
+					OptimalLODLevel = FMath::FloorLog2(WaterInfoUnitsPerPixel / LandscapeComponentUnitsPerVertex);
+
+					break;
+				}
+			}
+
+			// There should always be at least one valid component proxy and the optimal LOD level should never be negative.
+			check(OptimalLODLevel >= 0);
+
+			LandscapeLODOverridesToRestore.Add(LandscapeRenderSystem, LandscapeRenderSystem->ForcedLODOverride);
+			LandscapeRenderSystem->ForcedLODOverride = OptimalLODLevel;
+		}
+	}
+
+	void RestoreLandscapeLODOverrides()
+	{
+		if (!LandscapeRenderSystems) return;
+		
+		for (const TPair<uint32, FLandscapeRenderSystem*>& Pair : *LandscapeRenderSystems)
+		{
+			FLandscapeRenderSystem* LandscapeRenderSystem = Pair.Value;
+			LandscapeRenderSystem->ForcedLODOverride = LandscapeLODOverridesToRestore.FindChecked(LandscapeRenderSystem);
+		}
+	}
+
+	const FUpdateWaterInfoParams& Params;
+	TMap<uint32, FLandscapeRenderSystem*> const* LandscapeRenderSystems = nullptr;
+	TMap<FLandscapeRenderSystem*, int8> LandscapeLODOverridesToRestore;
+};
+
 static void UpdateWaterInfoRendering_RenderThread(
 	FRHICommandListImmediate& RHICmdList,
 	const FUpdateWaterInfoParams& Params)
 {
+	FScopedLandscapeLODOverride ScopedLandscapeLODOverride(Params);
+
 	FMaterialRenderProxy::UpdateDeferredCachedUniformExpressions();
 
 	FRenderTarget* RenderTarget = Params.RenderTarget;
@@ -452,7 +549,6 @@ static FSceneRenderer* CreateWaterInfoDepthRenderer(
 		.SetRealtimeUpdate(false)
 		.SetResolveScene(false));
 	DepthViewFamily.SceneCaptureSource = ESceneCaptureSource::SCS_DeviceDepth;
-	DepthViewFamily.LandscapeLODOverride = Context.LandscapeLODOverride;
 
 	// Setup the view family
 	FSceneViewInitOptions DepthViewInitOptions;
@@ -532,7 +628,6 @@ static FSceneRenderer* CreateWaterInfoColorRenderer(
 		.SetRealtimeUpdate(false)
 		.SetResolveScene(false));
 	ColorViewFamily.SceneCaptureSource = ESceneCaptureSource::SCS_SceneColorSceneDepth;
-	ColorViewFamily.LandscapeLODOverride = Context.LandscapeLODOverride;
 
 	FSceneViewInitOptions ColorViewInitOptions;
 	ColorViewInitOptions.SetViewRectangle(FIntRect(0, 0, RenderTargetSize.X, RenderTargetSize.Y));
@@ -640,6 +735,7 @@ void UpdateWaterInfoRendering(
 	Params.WaterHeightExtents = Context.ZoneToRender->GetWaterHeightExtents();
 	Params.GroundZMin = Context.ZoneToRender->GetGroundZMin();
 	Params.VelocityBlurRadius = Context.ZoneToRender->GetVelocityBlurRadius();
+	Params.WaterZoneExtents = Context.ZoneToRender->GetZoneExtent();
 
 	ENQUEUE_RENDER_COMMAND(WaterInfoCommand)(
 	[Params, ZoneName = Context.ZoneToRender->GetActorNameOrLabel()](FRHICommandListImmediate& RHICmdList)
diff --git a/Engine/Source/Runtime/Landscape/Private/LandscapeModule.cpp b/Engine/Source/Runtime/Landscape/Private/LandscapeModule.cpp
--- a/Engine/Source/Runtime/Landscape/Private/LandscapeModule.cpp
+++ b/Engine/Source/Runtime/Landscape/Private/LandscapeModule.cpp
@@ -23,13 +23,14 @@
 // Register the custom version with core
 FCustomVersionRegistration GRegisterLandscapeCustomVersion(FLandscapeCustomVersion::GUID, FLandscapeCustomVersion::LatestVersion, TEXT("Landscape"));
 
-class FLandscapeModule : public IModuleInterface
+class FLandscapeModule : public ILandscapeModule
 {
 public:
 	/** IModuleInterface implementation */
 	void StartupModule() override;
 	void ShutdownModule() override;
 
+	virtual TSharedPtr<FLandscapeSceneViewExtension, ESPMode::ThreadSafe> GetLandscapeSceneViewExtension() const override { return SceneViewExtension; }
 private:
 	void OnPostEngineInit();
 	void OnEnginePreExit();
diff --git a/Engine/Source/Runtime/Landscape/Private/LandscapeRender.cpp b/Engine/Source/Runtime/Landscape/Private/LandscapeRender.cpp
--- a/Engine/Source/Runtime/Landscape/Private/LandscapeRender.cpp
+++ b/Engine/Source/Runtime/Landscape/Private/LandscapeRender.cpp
@@ -669,6 +669,7 @@ FLandscapeRenderSystem::FLandscapeRenderSystem()
 	: Min(MAX_int32, MAX_int32)
 	, Size(EForceInit::ForceInitToZero)
 	, ReferenceCount(0)
+	, ForcedLODOverride(-1)
 {
 	SectionLODBiases.SetAllowCPUAccess(true);
 
@@ -1014,6 +1015,12 @@ void FLandscapeSceneViewExtension::EndFrame_RenderThread()
 	LandscapeIndirectionBuffer = FBufferRHIRef();
 }
 
+const TMap<uint32, FLandscapeRenderSystem*>& FLandscapeSceneViewExtension::GetLandscapeRenderSystems() const
+{
+	return LandscapeRenderSystems;
+}
+
+
 //
 // FLandscapeVisibilityHelper
 //
@@ -4026,7 +4033,10 @@ float FLandscapeComponentSceneProxy::ComputeLODForView(const FSceneView& InView)
 
 	float LODScale = ViewLODDistanceFactor * CVarStaticMeshLODDistanceScale.GetValueOnRenderThread();
 
-	int32 ForcedLODLevel = LODSettings.ForcedLOD;
+	FLandscapeRenderSystem* LandscapeRenderSystem = LandscapeRenderSystems.FindChecked(LandscapeKey);
+
+	// Prefer the RenderSystem's ForcedLODOverride if set over any per-component LOD override
+	int32 ForcedLODLevel = LandscapeRenderSystem->ForcedLODOverride >= 0 ? LandscapeRenderSystem->ForcedLODOverride : LODSettings.ForcedLOD;
 	ForcedLODLevel = ViewLODOverride >= 0 ? ViewLODOverride : ForcedLODLevel;
 	const int32 DrawCollisionLODOverride = GetDrawCollisionLodOverride(ViewEngineShowFlagCollisionPawn, ViewEngineShowFlagCollisionVisibility, LODSettings.DrawCollisionPawnLOD, LODSettings.DrawCollisionVisibilityLOD);
 	ForcedLODLevel = DrawCollisionLODOverride >= 0 ? DrawCollisionLODOverride : ForcedLODLevel;
@@ -4069,6 +4079,13 @@ float FLandscapeComponentSceneProxy::ComputeLODBias() const
 	return ComputedLODBias;
 }
 
+double FLandscapeComponentSceneProxy::ComputeSectionResolution() const
+{
+	// LandscapeComponent Max Extend represents the half-extent of the landscape component. Multiply by two to get the actual size.
+	const double ComponentFullExtent = 2.0 * ComponentMaxExtend;
+	return ComponentFullExtent / (double)(ComponentSizeVerts);
+}
+
 //
 // FLandscapeSectionInfo
 //
diff --git a/Engine/Source/Runtime/Landscape/Public/LandscapeModule.h b/Engine/Source/Runtime/Landscape/Public/LandscapeModule.h
--- a/Engine/Source/Runtime/Landscape/Public/LandscapeModule.h
+++ b/Engine/Source/Runtime/Landscape/Public/LandscapeModule.h
@@ -5,10 +5,13 @@
 #include "CoreMinimal.h"
 #include "Modules/ModuleInterface.h"
 
+class FLandscapeSceneViewExtension;
+
 /**
 * Landscape module interface
 */
 class ILandscapeModule : public IModuleInterface
 {
 public:
+	virtual TSharedPtr<FLandscapeSceneViewExtension, ESPMode::ThreadSafe> GetLandscapeSceneViewExtension() const = 0;
 };
diff --git a/Engine/Source/Runtime/Landscape/Public/LandscapeRender.h b/Engine/Source/Runtime/Landscape/Public/LandscapeRender.h
--- a/Engine/Source/Runtime/Landscape/Public/LandscapeRender.h
+++ b/Engine/Source/Runtime/Landscape/Public/LandscapeRender.h
@@ -361,10 +361,13 @@ public:
 	virtual float ComputeLODBias() const = 0;
 	virtual int32 GetSectionPriority() const { return INDEX_NONE; }
 
+	/** Computes the worldspace units per vertex of the landscape section. */
+	virtual double ComputeSectionResolution() const { return -1.0; }
+
 public:
-	uint32		LandscapeKey;
-	FIntPoint				ComponentBase;
-	bool					bRegistered;
+	uint32 LandscapeKey;
+	FIntPoint ComponentBase;
+	bool bRegistered;
 };
 
 struct FLandscapeRenderSystem
@@ -423,6 +426,9 @@ struct FLandscapeRenderSystem
 
 	TMap<FViewKey, TResourceArray<float>> CachedSectionLODValues;
 
+	/** Forced LOD level which overrides the ForcedLOD level of all the sections under this LandscapeRenderSystem. */
+	int8 ForcedLODOverride;
+
 	FLandscapeRenderSystem();
 	~FLandscapeRenderSystem();
 
@@ -487,6 +493,7 @@ public:
 	virtual void PreRenderViewFamily_RenderThread(FRHICommandListImmediate& RHICmdList, FSceneViewFamily& InViewFamily) override;
 	virtual void PreRenderView_RenderThread(FRHICommandListImmediate& RHICmdList, FSceneView& InView) override;
 
+	LANDSCAPE_API const TMap<uint32, FLandscapeRenderSystem*>& GetLandscapeRenderSystems() const;
 private:
 	FBufferRHIRef LandscapeLODDataBuffer;
 	FBufferRHIRef LandscapeIndirectionBuffer;
@@ -809,9 +816,10 @@ public:
 	virtual bool IsRayTracingRelevant() const override { return true; }
 #endif
 
-	// FLandscapeSceneInfo interface
+	// FLandscapeSectionInfo interface
 	virtual float ComputeLODForView(const FSceneView& InView) const override;
 	virtual float ComputeLODBias() const override;
+	virtual double ComputeSectionResolution() const override;
 };
 
 class FLandscapeDebugMaterialRenderProxy : public FMaterialRenderProxy
