Fix crash on editor shutdown

diff --git a/Engine/Plugins/Experimental/Water/Source/Runtime/Private/GerstnerWaterWaveViewExtension.cpp b/Engine/Plugins/Experimental/Water/Source/Runtime/Private/GerstnerWaterWaveViewExtension.cpp
--- a/Engine/Plugins/Experimental/Water/Source/Runtime/Private/GerstnerWaterWaveViewExtension.cpp
+++ b/Engine/Plugins/Experimental/Water/Source/Runtime/Private/GerstnerWaterWaveViewExtension.cpp
@@ -31,7 +31,8 @@ void FGerstnerWaterWaveViewExtension::Deinitialize()
 		[WaveGPUData=WaveGPUData](FRHICommandListImmediate& RHICmdList){}
 	);
 
-	if (UGerstnerWaterWaveSubsystem* GerstnerWaterWaveSubsystem = GEngine->GetEngineSubsystem<UGerstnerWaterWaveSubsystem>())
+	// It's possible for GEngine to be null when UWaterSubsystem deinitializes
+	if (UGerstnerWaterWaveSubsystem* GerstnerWaterWaveSubsystem = GEngine ? GEngine->GetEngineSubsystem<UGerstnerWaterWaveSubsystem>() : nullptr)
 	{
 		GerstnerWaterWaveSubsystem->Unregister(this);
 	}
