Fixed water bodies not writing velocities by replacing the identity primitive uniform buffer with the actual one that has the proper primitive flag set.

#rnx

diff --git a/Engine/Plugins/Experimental/Water/Source/Runtime/Private/WaterMeshSceneProxy.cpp b/Engine/Plugins/Experimental/Water/Source/Runtime/Private/WaterMeshSceneProxy.cpp
--- a/Engine/Plugins/Experimental/Water/Source/Runtime/Private/WaterMeshSceneProxy.cpp
+++ b/Engine/Plugins/Experimental/Water/Source/Runtime/Private/WaterMeshSceneProxy.cpp
@@ -369,10 +369,12 @@ void FWaterMeshSceneProxy::GetDynamicMeshElements(const TArray<const FSceneView*
 							BatchElement.MinVertexIndex = 0;
 							BatchElement.MaxVertexIndex = WaterVertexFactories[DensityIndex]->VertexBuffer->GetVertexCount() - 1;
 
-							// Don't use primitive buffer
 							BatchElement.IndexBuffer = WaterVertexFactories[DensityIndex]->IndexBuffer;
 							BatchElement.PrimitiveIdMode = PrimID_ForceZero;
-							BatchElement.PrimitiveUniformBufferResource = &GIdentityPrimitiveUniformBuffer;
+
+							// We need the uniform buffer of this primitive because it stores the proper value for the bOutputVelocity flag.
+							// The identity primitive uniform buffer simply stores false for this flag which leads to missing motion vectors.
+							BatchElement.PrimitiveUniformBuffer = GetUniformBuffer(); 
 						}
 
 						{
@@ -630,7 +632,7 @@ FPrimitiveViewRelevance FWaterMeshSceneProxy::GetViewRelevance(const FSceneView*
 	Result.bRenderCustomDepth = ShouldRenderCustomDepth();
 	Result.bTranslucentSelfShadow = bCastVolumetricTranslucentShadow;
 	MaterialRelevance.SetPrimitiveViewRelevance(Result);
-	Result.bVelocityRelevance = false;
+	Result.bVelocityRelevance = DrawsVelocity() && Result.bOpaque && Result.bRenderInMainPass;
 	return Result;
 }
 
