More explicit logging for ensure in UBuoyancyComponent::ComputePontoonCoefficients

diff --git a/Engine/Plugins/Experimental/Water/Source/Runtime/Private/BuoyancyComponent.cpp b/Engine/Plugins/Experimental/Water/Source/Runtime/Private/BuoyancyComponent.cpp
--- a/Engine/Plugins/Experimental/Water/Source/Runtime/Private/BuoyancyComponent.cpp
+++ b/Engine/Plugins/Experimental/Water/Source/Runtime/Private/BuoyancyComponent.cpp
@@ -394,8 +394,9 @@ void UBuoyancyComponent::ComputePontoonCoefficients()
 		{
 			const FVector& LocalCOM = BodyInstance->GetMassSpaceLocal().GetLocation();
 			//Distribute a mass of 1 to each pontoon so that we get a scaling factor based on position relative to CoM
-			bool ComputeSuccess = FSimpleSuspensionHelpers::ComputeSprungMasses(LocalPontoonLocations, LocalCOM, 1.f, PontoonCoefficients);
-			ensureMsgf(ComputeSuccess, TEXT("Failed to compute sprung masses for: %s"), *GetOwner()->GetName());
+			FString ErrMsg;
+			bool ComputeSuccess = FSimpleSuspensionHelpers::ComputeSprungMasses(LocalPontoonLocations, LocalCOM, 1.f, PontoonCoefficients, &ErrMsg);
+			ensureMsgf(ComputeSuccess, TEXT("Failed to compute %d sprung masses for: %s\nErrMsg: \"%s\""), LocalPontoonLocations.Num(), *GetOwner()->GetName(), *ErrMsg);
 		}
 	}
 
