Fixe non-unity compilation

diff --git a/Engine/Plugins/Experimental/Water/Source/Runtime/Private/GerstnerWaterWaveViewExtension.cpp b/Engine/Plugins/Experimental/Water/Source/Runtime/Private/GerstnerWaterWaveViewExtension.cpp
--- a/Engine/Plugins/Experimental/Water/Source/Runtime/Private/GerstnerWaterWaveViewExtension.cpp
+++ b/Engine/Plugins/Experimental/Water/Source/Runtime/Private/GerstnerWaterWaveViewExtension.cpp
@@ -4,6 +4,7 @@
 #include "GerstnerWaterWaveSubsystem.h"
 #include "WaterBodyActor.h"
 #include "GerstnerWaterWaves.h"
+#include "Engine/Engine.h"
 #include "WaterBodyManager.h"
 #include "WaterSubsystem.h"
 
@@ -33,10 +34,10 @@ void FGerstnerWaterWaveViewExtension::Deinitialize()
 
 	// It's possible for GEngine to be null when UWaterSubsystem deinitializes
 	if (UGerstnerWaterWaveSubsystem* GerstnerWaterWaveSubsystem = GEngine ? GEngine->GetEngineSubsystem<UGerstnerWaterWaveSubsystem>() : nullptr)
-	{
-		GerstnerWaterWaveSubsystem->Unregister(this);
+		{
+			GerstnerWaterWaveSubsystem->Unregister(this);
+		}
 	}
-}
 
 void FGerstnerWaterWaveViewExtension::SetupViewFamily(FSceneViewFamily& InViewFamily)
 {
@@ -50,7 +51,7 @@ void FGerstnerWaterWaveViewExtension::SetupViewFamily(FSceneViewFamily& InViewFa
 
 		FWaterBodyManager::ForEachWaterBodyComponent(WorldPtr.Get(), [&WaterIndirectionBuffer, &WaterDataBuffer](UWaterBodyComponent* WaterBodyComponent)
 		{
-			// Some max value
+		// Some max value
 			constexpr int32 MaxWavesPerWaterBody = 4096;
 			constexpr int32 NumFloat4PerWave = 2;
 
@@ -89,9 +90,9 @@ void FGerstnerWaterWaveViewExtension::SetupViewFamily(FSceneViewFamily& InViewFa
 
 						WaterDataBuffer[WaveIndex] = FVector4f(Wave.Direction.X, Wave.Direction.Y, Wave.WaveLength, Wave.Amplitude);
 						WaterDataBuffer[WaveIndex + 1] = FVector4f(Wave.Steepness, 0.0f, 0.0f, 0.0f);
-					}
 				}
 			}
+		}
 			return true;
 		});
 
