Fixed an issue where the water brush caches were not being recreated when the container was being GC'd.

diff --git a/Engine/Plugins/Experimental/Water/Source/Editor/Private/WaterBrushManager.cpp b/Engine/Plugins/Experimental/Water/Source/Editor/Private/WaterBrushManager.cpp
--- a/Engine/Plugins/Experimental/Water/Source/Editor/Private/WaterBrushManager.cpp
+++ b/Engine/Plugins/Experimental/Water/Source/Editor/Private/WaterBrushManager.cpp
@@ -1461,8 +1461,8 @@ void AWaterBrushManager::UpdateBrushCacheKeys()
 
 			AActor* Actor = CastChecked<AActor>(BrushActor.GetObject());
 			UWaterBodyBrushCacheContainer* WaterBrushCacheContainer = Cast<UWaterBodyBrushCacheContainer>(GetActorCache(Actor, UWaterBodyBrushCacheContainer::StaticClass()));
-			// If no cache entry, create one and associate it to the actor : 
-			if (WaterBrushCacheContainer == nullptr)
+			// If no cache entry or the cache entry is invalid, create one and associate it to the actor : 
+			if (!::IsValid(WaterBrushCacheContainer))
 			{
 				WaterBrushCacheContainer = NewObject<UWaterBodyBrushCacheContainer>(this, NAME_None, RF_Transactional);
 				check(!WaterBrushCacheContainer->Cache.CacheIsValid);
