Added a new AffectsWaterInfo virtual function to UWaterBodyComponent which allows it to indicate if it will ever write to the water info texture. This allows us to prevent rebuilding the water info texture when water bodies which don't affect the texture are changed.

diff --git a/Engine/Plugins/Experimental/Water/Source/Runtime/Private/WaterBodyComponent.cpp b/Engine/Plugins/Experimental/Water/Source/Runtime/Private/WaterBodyComponent.cpp
--- a/Engine/Plugins/Experimental/Water/Source/Runtime/Private/WaterBodyComponent.cpp
+++ b/Engine/Plugins/Experimental/Water/Source/Runtime/Private/WaterBodyComponent.cpp
@@ -200,6 +200,12 @@ bool UWaterBodyComponent::AffectsWaterMesh() const
 	return ShouldGenerateWaterMeshTile();
 }
 
+bool UWaterBodyComponent::AffectsWaterInfo() const
+{
+	// Currently only water bodies which are rendered by the water mesh can render into the water info texture
+	return ShouldGenerateWaterMeshTile();
+}
+
 #if WITH_EDITOR
 ETextureRenderTargetFormat UWaterBodyComponent::GetBrushRenderTargetFormat() const
 {
@@ -830,12 +836,17 @@ void UWaterBodyComponent::UpdateComponentVisibility(bool bAllowWaterMeshRebuild)
 	 		Component->SetHiddenInGame(bLocalHiddenInGame);
 	 	}
 
-		// If the component is being or can be rendered by the water mesh, rebuild it in case its visibility has changed : 
-		if (bAllowWaterMeshRebuild && (GetWaterBodyType() != EWaterBodyType::Transition))
+		if (AWaterZone* WaterZone = GetWaterZone())
 		{
-			if (AWaterZone* WaterZone = GetWaterZone())
+			// If the component is being or can be rendered by the water mesh or renders into the water info texture, rebuild it in case its visibility has changed : 
+			if (bAllowWaterMeshRebuild && AffectsWaterMesh())
 			{
-				WaterZone->MarkForRebuild(EWaterZoneRebuildFlags::All);
+				WaterZone->MarkForRebuild(EWaterZoneRebuildFlags::UpdateWaterMesh);
+			}
+
+			if (AffectsWaterInfo())
+			{
+				WaterZone->MarkForRebuild(EWaterZoneRebuildFlags::UpdateWaterInfoTexture);
 			}
 		}
 	}
@@ -1506,12 +1517,15 @@ AWaterZone* UWaterBodyComponent::GetWaterZone() const
 
 void UWaterBodyComponent::UpdateWaterBodyRenderData()
 {
-	GenerateWaterBodyMesh();
-	MarkRenderStateDirty();
-
-	if (AWaterZone* WaterZone = GetWaterZone())
+	if (AffectsWaterInfo())
 	{
-		WaterZone->MarkForRebuild(EWaterZoneRebuildFlags::UpdateWaterInfoTexture);
+		GenerateWaterBodyMesh();
+		MarkRenderStateDirty();
+
+		if (AWaterZone* WaterZone = GetWaterZone())
+		{
+			WaterZone->MarkForRebuild(EWaterZoneRebuildFlags::UpdateWaterInfoTexture);
+		}
 	}
 }
 
diff --git a/Engine/Plugins/Experimental/Water/Source/Runtime/Private/WaterZoneActor.cpp b/Engine/Plugins/Experimental/Water/Source/Runtime/Private/WaterZoneActor.cpp
--- a/Engine/Plugins/Experimental/Water/Source/Runtime/Private/WaterZoneActor.cpp
+++ b/Engine/Plugins/Experimental/Water/Source/Runtime/Private/WaterZoneActor.cpp
@@ -243,6 +243,12 @@ bool AWaterZone::UpdateWaterInfoTexture()
 		TArray<UWaterBodyComponent*> WaterBodies;
 		ForEachWaterBodyComponent([World, &WaterBodies, &WaterZMax, &WaterZMin, &bHasIncompleteShaderMaps](UWaterBodyComponent* Component)
 		{
+			// skip components which don't affect the water info texture
+			if (!Component->AffectsWaterInfo())
+			{
+				return true;
+			}
+
 			if (UMaterialInterface* WaterInfoMaterial = Component->GetWaterInfoMaterialInstance())
 			{
 				if (FMaterialResource* MaterialResource = WaterInfoMaterial->GetMaterialResource(World->Scene->GetFeatureLevel()))
diff --git a/Engine/Plugins/Experimental/Water/Source/Runtime/Public/WaterBodyComponent.h b/Engine/Plugins/Experimental/Water/Source/Runtime/Public/WaterBodyComponent.h
--- a/Engine/Plugins/Experimental/Water/Source/Runtime/Public/WaterBodyComponent.h
+++ b/Engine/Plugins/Experimental/Water/Source/Runtime/Public/WaterBodyComponent.h
@@ -86,6 +86,7 @@ class WATER_API UWaterBodyComponent : public UPrimitiveComponent
 public:
 	virtual bool AffectsLandscape() const;
 	virtual bool AffectsWaterMesh() const;
+	virtual bool AffectsWaterInfo() const;
 	virtual bool CanAffectWaterMesh() const { return false; }
 
 	void UpdateAll(bool bShapeOrPositionChanged);
