Make WaterZone and WaterBodyOcean non-spatially loaded by default

#rnx

diff --git a/Engine/Plugins/Experimental/Water/Source/Runtime/Private/WaterBodyOceanActor.cpp b/Engine/Plugins/Experimental/Water/Source/Runtime/Private/WaterBodyOceanActor.cpp
--- a/Engine/Plugins/Experimental/Water/Source/Runtime/Private/WaterBodyOceanActor.cpp
+++ b/Engine/Plugins/Experimental/Water/Source/Runtime/Private/WaterBodyOceanActor.cpp
@@ -17,6 +17,7 @@ AWaterBodyOcean::AWaterBodyOcean(const FObjectInitializer& ObjectInitializer)
 
 #if WITH_EDITORONLY_DATA
 	CollisionExtents_DEPRECATED = FVector(50000.f, 50000.f, 10000.f);
+	bIsSpatiallyLoaded = false;
 #endif // WITH_EDITORONLY_DATA
 }
 
diff --git a/Engine/Plugins/Experimental/Water/Source/Runtime/Private/WaterZoneActor.cpp b/Engine/Plugins/Experimental/Water/Source/Runtime/Private/WaterZoneActor.cpp
--- a/Engine/Plugins/Experimental/Water/Source/Runtime/Private/WaterZoneActor.cpp
+++ b/Engine/Plugins/Experimental/Water/Source/Runtime/Private/WaterZoneActor.cpp
@@ -58,6 +58,10 @@ AWaterZone::AWaterZone(const FObjectInitializer& Initializer)
 
 	ActorIcon = FWaterIconHelper::EnsureSpriteComponentCreated(this, TEXT("/Water/Icons/WaterZoneActorSprite"));
 #endif // WITH_EDITOR
+
+#if WITH_EDITORONLY_DATA
+	bIsSpatiallyLoaded = false;
+#endif
 }
 
 void AWaterZone::SetZoneExtent(FVector2D NewExtent)
diff --git a/Engine/Plugins/Experimental/Water/Source/Runtime/Public/WaterZoneActor.h b/Engine/Plugins/Experimental/Water/Source/Runtime/Public/WaterZoneActor.h
--- a/Engine/Plugins/Experimental/Water/Source/Runtime/Public/WaterZoneActor.h
+++ b/Engine/Plugins/Experimental/Water/Source/Runtime/Public/WaterZoneActor.h
@@ -48,6 +48,10 @@ public:
 	virtual void PostLoadSubobjects(FObjectInstancingGraph* OuterInstanceGraph) override;
 	virtual void PostLoad() override;
 
+#if WITH_EDITOR
+	virtual bool CanChangeIsSpatiallyLoadedFlag() const override { return false; }
+#endif
+
 	FVector2f GetWaterHeightExtents() const { return WaterHeightExtents; }
 	float GetGroundZMin() const { return GroundZMin; }
 
