Reworked the way water body dilation is implemented.
Dilated sections of the water body meshes are split and only rendered in a dedicated dilation pass which merges back into the water info texture in the water info merge pass. This prevents dilated mesh sections from overwriting the data of other water bodies. It allows avoids the need to embed runtime data (the water body index) into the vertex colors of the mesh which will enable serialized mesh data instead of runtime generated data.

diff --git a/Engine/Plugins/Experimental/Water/Shaders/Private/WaterInfoMerge.usf b/Engine/Plugins/Experimental/Water/Shaders/Private/WaterInfoMerge.usf
--- a/Engine/Plugins/Experimental/Water/Shaders/Private/WaterInfoMerge.usf
+++ b/Engine/Plugins/Experimental/Water/Shaders/Private/WaterInfoMerge.usf
@@ -5,6 +5,7 @@
 
 Texture2D<float4> DepthTexture;
 Texture2D<float4> ColorTexture;
+Texture2D<float4> DilationTexture;
 
 
 float2 WaterHeightExtents;
@@ -14,7 +15,6 @@ float CaptureZ;
 struct FWaterColorPassResult
 {
 	float2 Velocity;
-	int WaterBodyIndex;
 	float SceneDepth;
 };
 
@@ -22,7 +22,6 @@ FWaterColorPassResult MakeWaterColorPassResult(float4 InSampleData)
 {
 	FWaterColorPassResult Result;
 	Result.Velocity = InSampleData.rg;
-	Result.WaterBodyIndex = InSampleData.b;
 	Result.SceneDepth = InSampleData.a;
 	return Result;
 }
@@ -56,31 +55,60 @@ void Main(
 	const float GroundDeviceDepth = DecodeFloatRGB(DepthTexture.Load(int3(TextureCoordinates, 0)).rgb);
 	const float GroundSceneDepth = ConvertFromDeviceZ(GroundDeviceDepth);
 	const FWaterColorPassResult WaterInfo = MakeWaterColorPassResult(ColorTexture.Load(int3(TextureCoordinates, 0)));
+
 	const float WaterSceneDepth = WaterInfo.SceneDepth;
-	float WaterZ = CaptureZ - WaterSceneDepth;
+	const float WaterZ = CaptureZ - WaterSceneDepth;
 
 	float GroundZ = CaptureZ - GroundSceneDepth;
 
-	const float WaterZMin = WaterHeightExtents.x;
-	const float WaterZMax = WaterHeightExtents.y;
-	
-	// If this is a dilated region (indicated by WaterBodyIndex == -1) and it is above the ground we must push it down to prevent artifacts
-	// A negative NormalizedWaterZ is a special value for the material and will map to a value slightly below the terrain.
-	if (WaterInfo.WaterBodyIndex == -1 && WaterZ > GroundZ)
-	{
-		WaterZ = GroundZ - 32.f;
-	}
-
-	const float NormalizedWaterZ = NormalizeHeight(WaterZ, WaterZMin, WaterZMax);
-	
 	// If the GroundDepth is further than the far plane of the projection, set the GroundZ to some really far away value
 	if (GroundDeviceDepth == 1.0)
 	{
 		GroundZ = GroundZMin;
 	}
+
+	const float WaterZMin = WaterHeightExtents.x;
+	const float WaterZMax = WaterHeightExtents.y;
 	
-	const float NormalizedGroundZ = NormalizeHeight(GroundZ, GroundZMin, WaterZMax);
+	// if the measured WaterZ is below WaterZMin it means that there is no water at this pixel because it is not possible
+	// for water to be further than ZMin.
+	const bool bPixelHasValidWaterData = WaterZ >= WaterZMin;
+
+	if (bPixelHasValidWaterData)
+	{
+		const float NormalizedWaterZ = NormalizeHeight(WaterZ, WaterZMin, WaterZMax);
+		const float NormalizedGroundZ = NormalizeHeight(GroundZ, GroundZMin, WaterZMax);
 
+		OutColor = float4(WaterInfo.Velocity, NormalizedWaterZ, NormalizedGroundZ);
+		return;
+	}
+	else
+	{
+		const float DilationDeviceDepth = DecodeFloatRGB(DilationTexture.Load(int3(TextureCoordinates, 0)).rgb);
+		const float DilationSceneDepth = ConvertFromDeviceZ(DilationDeviceDepth);
+		const bool bPixelHasDilationData = DilationSceneDepth != 0.0;
+		
+		if (bPixelHasDilationData)
+		{
+			float DilatedZ = CaptureZ - DilationSceneDepth;
 
-	OutColor = float4(WaterInfo.Velocity, NormalizedWaterZ, NormalizedGroundZ);
+			// If the dilated water would render on top of terrain, squash it down below
+			if (DilatedZ > GroundZ)
+			{
+				DilatedZ = GroundZ - 64.f;
+			}
+
+			const float NormalizedWaterZ = NormalizeHeight(DilatedZ, WaterZMin, WaterZMax);
+			const float NormalizedGroundZ = NormalizeHeight(GroundZ, GroundZMin, WaterZMax);
+
+			OutColor = float4(0.0, 0.0, NormalizedWaterZ, NormalizedGroundZ);
+		}
+		else
+		{
+			// Pixel may not have any dilation data nor water data. In this case just write 0s to the texture
+
+			const float NormalizedGroundZ = NormalizeHeight(GroundZ, GroundZMin, WaterZMax);
+			OutColor = float4(0.0, 0.0, 0.0, NormalizedGroundZ);
+		}
+	}
 }
diff --git a/Engine/Plugins/Experimental/Water/Source/Runtime/Private/WaterBodyComponent.cpp b/Engine/Plugins/Experimental/Water/Source/Runtime/Private/WaterBodyComponent.cpp
--- a/Engine/Plugins/Experimental/Water/Source/Runtime/Private/WaterBodyComponent.cpp
+++ b/Engine/Plugins/Experimental/Water/Source/Runtime/Private/WaterBodyComponent.cpp
@@ -116,7 +116,7 @@ void UWaterBodyComponent::OnHiddenInGameChanged()
 
 FPrimitiveSceneProxy* UWaterBodyComponent::CreateSceneProxy()
 {
-	return new FWaterBodySceneProxy(this, WaterBodyMeshVertices, WaterBodyMeshIndices);
+	return new FWaterBodySceneProxy(this);
 }
 
 void UWaterBodyComponent::GetUsedMaterials(TArray<UMaterialInterface*>& OutMaterialInterfaces, bool bGetDebugMaterials) const
diff --git a/Engine/Plugins/Experimental/Water/Source/Runtime/Private/WaterBodyLakeComponent.cpp b/Engine/Plugins/Experimental/Water/Source/Runtime/Private/WaterBodyLakeComponent.cpp
--- a/Engine/Plugins/Experimental/Water/Source/Runtime/Private/WaterBodyLakeComponent.cpp
+++ b/Engine/Plugins/Experimental/Water/Source/Runtime/Private/WaterBodyLakeComponent.cpp
@@ -119,20 +119,17 @@ void UWaterBodyLakeComponent::GenerateWaterBodyMesh()
 		
 		if (Inset.Apply())
 		{
-			const uint32 IndexOffset = WaterBodyMeshVertices.Num();
 			for (const FVector3d& Vertex : LakeMesh.GetVerticesBuffer())
 			{
 				// push the set of dilated vertices to the persistent mesh
 				FDynamicMeshVertex MeshVertex(FVector3f(Vertex.X, Vertex.Y, 0.f));
-				MeshVertex.Position.Z = ShapeDilationZOffset;
 				MeshVertex.Color = FColor::Black;
-				MeshVertex.TextureCoordinate[0].X = -1;
-				WaterBodyMeshVertices.Add(MeshVertex);
+				DilatedWaterBodyMeshVertices.Add(MeshVertex);
 			}
 
 			for (const FIndex3i& Triangle : LakeMesh.GetTrianglesBuffer())
 			{
-				WaterBodyMeshIndices.Append({ IndexOffset + Triangle.A, IndexOffset + Triangle.B, IndexOffset + Triangle.C });
+				DilatedWaterBodyMeshIndices.Append({ (uint32)Triangle.A, (uint32)Triangle.B, (uint32)Triangle.C });
 			}
 		}
 		else
diff --git a/Engine/Plugins/Experimental/Water/Source/Runtime/Private/WaterBodyOceanComponent.cpp b/Engine/Plugins/Experimental/Water/Source/Runtime/Private/WaterBodyOceanComponent.cpp
--- a/Engine/Plugins/Experimental/Water/Source/Runtime/Private/WaterBodyOceanComponent.cpp
+++ b/Engine/Plugins/Experimental/Water/Source/Runtime/Private/WaterBodyOceanComponent.cpp
@@ -207,20 +207,17 @@ void UWaterBodyOceanComponent::GenerateWaterBodyMesh()
 		
 		if (Inset.Apply())
 		{
-			const uint32 IndexOffset = WaterBodyMeshVertices.Num();
 			for (const FVector3d& Vertex : OceanMesh.GetVerticesBuffer())
 			{
 				// push the set of dilated vertices to the persistent mesh
 				FDynamicMeshVertex MeshVertex(FVector3f(Vertex.X, Vertex.Y, 0.f));
-				MeshVertex.Position.Z = ShapeDilationZOffset;
 				MeshVertex.Color = FColor::Black;
-				MeshVertex.TextureCoordinate[0].X = -1;
-				WaterBodyMeshVertices.Add(MeshVertex);
+				DilatedWaterBodyMeshVertices.Add(MeshVertex);
 			}
 
 			for (const FIndex3i& Triangle : OceanMesh.GetTrianglesBuffer())
 			{
-				WaterBodyMeshIndices.Append({ IndexOffset + Triangle.A, IndexOffset + Triangle.B, IndexOffset + Triangle.C });
+				DilatedWaterBodyMeshIndices.Append({ (uint32)Triangle.A, (uint32)Triangle.B, (uint32)Triangle.C });
 			}
 		}
 		else
diff --git a/Engine/Plugins/Experimental/Water/Source/Runtime/Private/WaterBodyRiverComponent.cpp b/Engine/Plugins/Experimental/Water/Source/Runtime/Private/WaterBodyRiverComponent.cpp
--- a/Engine/Plugins/Experimental/Water/Source/Runtime/Private/WaterBodyRiverComponent.cpp
+++ b/Engine/Plugins/Experimental/Water/Source/Runtime/Private/WaterBodyRiverComponent.cpp
@@ -113,7 +113,15 @@ static FColor PackWaterFlow(float VelocityMagnitude, float DirectionAngle)
 	return Result;
 }
 
-static void AddVerticesForRiverSplineStep(float DistanceAlongSpline, const UWaterBodyRiverComponent* Component, const UWaterSplineComponent* SplineComp, const UWaterSplineMetadata* WaterSplineMetadata, TArray<FDynamicMeshVertex>& Vertices, TArray<uint32>& Indices)
+static void AddVerticesForRiverSplineStep(
+	float DistanceAlongSpline, 
+	const UWaterBodyRiverComponent* Component, 
+	const UWaterSplineComponent* SplineComp, 
+	const UWaterSplineMetadata* WaterSplineMetadata,
+	TArray<FDynamicMeshVertex>& Vertices, 
+	TArray<uint32>& Indices, 
+	TArray<FDynamicMeshVertex>& DilatedVertices, 
+	TArray<uint32>& DilatedIndices)
 {
 	check((Component != nullptr) && (SplineComp != nullptr) && (WaterSplineMetadata != nullptr));
 	
@@ -139,13 +147,9 @@ static void AddVerticesForRiverSplineStep(float DistanceAlongSpline, const UWate
 	FDynamicMeshVertex Right(FVector3f(Pos + OutwardDistance));
 
 	FDynamicMeshVertex DilatedFarLeft(FVector3f(Pos - OutwardDistance - DilationOffset));
-	DilatedFarLeft.Position.Z += Component->GetShapeDilationZOffsetFar();
 	FDynamicMeshVertex DilatedLeft(FVector3f(Pos - OutwardDistance));
-	DilatedLeft.Position.Z += Component->GetShapeDilationZOffset();
 	FDynamicMeshVertex DilatedRight(FVector3f(Pos + OutwardDistance));
-	DilatedRight.Position.Z += Component->GetShapeDilationZOffset();
 	FDynamicMeshVertex DilatedFarRight(FVector3f(Pos + OutwardDistance + DilationOffset));
-	DilatedFarRight.Position.Z += Component->GetShapeDilationZOffsetFar();
 
 	float FlowDirection = Tangent.HeadingAngle() + FMath::DegreesToRadians(Component->GetRelativeRotation().Yaw);
 	// Convert negative angles into positive angles
@@ -172,29 +176,31 @@ static void AddVerticesForRiverSplineStep(float DistanceAlongSpline, const UWate
 	DilatedRight.Color = EmptyFlowData;
 	DilatedFarRight.Color = EmptyFlowData;
 
-	// Embed the water body index in the vertex data so that we can distinguish between dilated and undilated regions of the texture
-	const int32 WaterBodyIndex = Component->GetWaterBodyIndex();
-	Left.TextureCoordinate[0].X = WaterBodyIndex;
-	Right.TextureCoordinate[0].X = WaterBodyIndex;
+	{
+		/* Non - dilated river segment geometry:
+		   2 --- 3
+		   |  /  |
+		   0 --- 1
+		*/
+		const uint32 BaseIndex = Vertices.Num();
+		Vertices.Append({ Left, Right });
+		Indices.Append({ BaseIndex, BaseIndex + 3, BaseIndex + 1, BaseIndex, BaseIndex + 2, BaseIndex + 3 });
+	}
 	
-	DilatedFarLeft.TextureCoordinate[0].X = -1;
-	DilatedLeft.TextureCoordinate[0].X = -1;
-	DilatedRight.TextureCoordinate[0].X = -1;
-	DilatedFarRight.TextureCoordinate[0].X = -1;
-
-	/* River segment geometry:
-		6---7,8---9,10--11
-		| /  |  /  |  / |
-		0---1,2---3,4---5
-	*/
-	const uint32 BaseIndex = Vertices.Num();
-	Vertices.Append({ DilatedFarLeft, DilatedLeft, Left, Right, DilatedRight, DilatedFarRight });
-	// Append left dilation quad
-	Indices.Append({ BaseIndex, BaseIndex + 7, BaseIndex + 1, BaseIndex, BaseIndex + 6, BaseIndex + 7 });
-	// Append main quad
-	Indices.Append({ BaseIndex + 2, BaseIndex + 9, BaseIndex + 3, BaseIndex + 2, BaseIndex + 8, BaseIndex + 9});
-	// Append right dilation quad
-	Indices.Append({ BaseIndex + 4, BaseIndex + 11, BaseIndex + 5, BaseIndex + 4, BaseIndex + 10, BaseIndex + 11 });
+	if (DilationAmount > 0.f)
+	{
+		/* Dilated River segment geometry:
+			4---5   6---7
+			| / |   | / |
+			0---1   2---3
+		*/
+		const uint32 BaseIndex = DilatedVertices.Num();
+		DilatedVertices.Append({ DilatedFarLeft, DilatedLeft, DilatedRight, DilatedFarRight });
+		// Append left dilation quad
+		DilatedIndices.Append({ BaseIndex, BaseIndex + 5, BaseIndex + 1, BaseIndex, BaseIndex + 4, BaseIndex + 5 });
+		// Append right dilation quad
+		DilatedIndices.Append({ BaseIndex + 2, BaseIndex + 6, BaseIndex + 7, BaseIndex + 2, BaseIndex + 7, BaseIndex + 3 });
+	}
 }
 
 enum class ERiverBoundaryEdge {
@@ -232,55 +238,34 @@ static void AddTerminalVerticesForRiverSpline(ERiverBoundaryEdge Edge, const UWa
 	}
 
 	FDynamicMeshVertex BackLeft(FVector3f(Pos - OutwardDistance + TangentialOffset - DilationOffset));
-	BackLeft.Position.Z += Component->GetShapeDilationZOffsetFar();
 	FDynamicMeshVertex Left(FVector3f(Pos - OutwardDistance + TangentialOffset));
-	Left.Position.Z += Component->GetShapeDilationZOffsetFar();
 	FDynamicMeshVertex Right(FVector3f(Pos + OutwardDistance + TangentialOffset));
-	Right.Position.Z += Component->GetShapeDilationZOffsetFar();
 	FDynamicMeshVertex BackRight(FVector3f(Pos + OutwardDistance + TangentialOffset + DilationOffset));
-	BackRight.Position.Z += Component->GetShapeDilationZOffsetFar();
 
-	// Initialize the vertex data to correct represent what we expect the dilated region to look like
-	// (no color and -1 UVs[0].x
+	// Initialize the vertex data to correct represent what we expect the dilated region to look like (no flow data)
 	BackLeft.Color = FColor(0.f);
-	BackLeft.TextureCoordinate[0].X = -1;
 	Left.Color = FColor(0.f);
-	Left.TextureCoordinate[0].X = -1;
 	Right.Color = FColor(0.f);
-	Right.TextureCoordinate[0].X = -1;
 	BackRight.Color = FColor(0.f);
-	BackRight.TextureCoordinate[0].X = -1;
-	
+
+	/* Dilated edge segment geometry:
+		4---5-----6---7
+		| / |  /  | / |
+		0---1-----2---3
+	*/
 
 	Vertices.Append({ BackLeft, Left, Right, BackRight });
+
+	const uint32 BaseIndex = Edge == ERiverBoundaryEdge::Start ? 0 : Vertices.Num() - 8;
+
+	// Since iterating the spline returns the final distance, all the side quads are already added for the final point so we only need to do this for the first point.
 	if (Edge == ERiverBoundaryEdge::Start)
-	{
-		/* Dilated front segment geometry: (ignore vertices 6-7 since they are the non-dilated vertices)
-			4---5,6---7,8---9
-			|    |     |    |
-			0----1-----2----3
-		*/
-		Indices.Append({ 0, 5, 1, 0, 4, 5});
-		Indices.Append({ 1, 8, 2, 1, 5, 8});
-		Indices.Append({ 2, 9, 3, 2, 8, 9});
-	}
-	else
-	{
-		/* Dilated back segment geometry: (ignore vertices 6-7 since they are the non-dilated vertices)
-			6----7-----8----9
-			|    |     |    |
-			0---1,2---3,4---5
-		*/
-		check(Vertices.Num() >= 10);
-		const uint32 BaseIndex = Vertices.Num() - 10;
-		// Since we don't know if we're at the end point or not in the AddVerticesForRiverSplineStep function,
-		// we end up adding a bunch of indices which link to the next set of vertices but they don't exist so we must trim them here.
-		Indices.SetNum(Indices.Num() - 18);
-
-		Indices.Append({ BaseIndex + 0, BaseIndex + 6, BaseIndex + 7, BaseIndex + 0, BaseIndex + 7, BaseIndex + 1 });
-		Indices.Append({ BaseIndex + 1, BaseIndex + 7, BaseIndex + 8, BaseIndex + 1, BaseIndex + 8, BaseIndex + 4 });
-		Indices.Append({ BaseIndex + 4, BaseIndex + 9, BaseIndex + 5, BaseIndex + 4, BaseIndex + 8, BaseIndex + 9 });
+	{ 
+
+		Indices.Append({ BaseIndex + 0, BaseIndex + 5, BaseIndex + 1, BaseIndex + 0, BaseIndex + 4, BaseIndex + 5 });
+		Indices.Append({ BaseIndex + 2, BaseIndex + 7, BaseIndex + 3, BaseIndex + 2, BaseIndex + 6, BaseIndex + 7 });
 	}
+	Indices.Append({ BaseIndex + 1, BaseIndex + 6, BaseIndex + 2, BaseIndex + 1, BaseIndex + 5, BaseIndex + 6 });
 }
 
 // ----------------------------------------------------------------------------------
@@ -309,19 +294,33 @@ void UWaterBodyRiverComponent::GenerateWaterBodyMesh()
 	TArray<FDynamicMeshVertex> Vertices;
 	TArray<uint32> Indices;
 
+	TArray<FDynamicMeshVertex> DilatedVertices;
+	TArray<uint32> DilatedIndices;
+
 	// Add an extra point at the start to dilate starting edge
-	AddTerminalVerticesForRiverSpline(ERiverBoundaryEdge::Start, this, SplineComp, WaterSplineMetadata, Vertices, Indices);
+	if (ShapeDilation > 0.f)
+	{
+		AddTerminalVerticesForRiverSpline(ERiverBoundaryEdge::Start, this, SplineComp, WaterSplineMetadata, DilatedVertices, DilatedIndices);
+	}
 
 	for (double DistanceAlongSpline : Distances)
 	{
-		AddVerticesForRiverSplineStep(DistanceAlongSpline, this, SplineComp, WaterSplineMetadata, Vertices, Indices);
+		AddVerticesForRiverSplineStep(DistanceAlongSpline, this, SplineComp, WaterSplineMetadata, Vertices, Indices, DilatedVertices, DilatedIndices);
 	}
+	// Remove the last two triangles since AddverticesForRiverSplineStep doesn't know which distance is the final distance and tries to create a triangle
+	// which links to vertices that don't exist.
+	Indices.SetNum(Indices.Num() - 6);
 	
 	// Add an extra point at the end to dilate ending edge
-	AddTerminalVerticesForRiverSpline(ERiverBoundaryEdge::End, this, SplineComp, WaterSplineMetadata, Vertices, Indices);
+	if (ShapeDilation > 0.f)
+	{
+		AddTerminalVerticesForRiverSpline(ERiverBoundaryEdge::End, this, SplineComp, WaterSplineMetadata, DilatedVertices, DilatedIndices);
+	}
 
 	WaterBodyMeshVertices = MoveTemp(Vertices);
 	WaterBodyMeshIndices = MoveTemp(Indices);
+	DilatedWaterBodyMeshVertices = MoveTemp(DilatedVertices);
+	DilatedWaterBodyMeshIndices = MoveTemp(DilatedIndices);
 }
 
 FBoxSphereBounds UWaterBodyRiverComponent::CalcBounds(const FTransform& LocalToWorld) const
diff --git a/Engine/Plugins/Experimental/Water/Source/Runtime/Private/WaterBodySceneProxy.cpp b/Engine/Plugins/Experimental/Water/Source/Runtime/Private/WaterBodySceneProxy.cpp
--- a/Engine/Plugins/Experimental/Water/Source/Runtime/Private/WaterBodySceneProxy.cpp
+++ b/Engine/Plugins/Experimental/Water/Source/Runtime/Private/WaterBodySceneProxy.cpp
@@ -29,30 +29,36 @@ public:
 	{}
 };
 
-void FWaterBodySceneProxy::InitResources(FWaterBodyMeshSection* Section)
+void FWaterBodySceneProxy::InitResources(FWaterBodyMeshSection& Section)
 {
-	check(Section != nullptr);
-	BeginInitResource(&Section->VertexBuffers.PositionVertexBuffer);
-	BeginInitResource(&Section->VertexBuffers.StaticMeshVertexBuffer);
-	BeginInitResource(&Section->VertexBuffers.ColorVertexBuffer);
-	BeginInitResource(&Section->IndexBuffer);
-	BeginInitResource(&Section->VertexFactory);
+	BeginInitResource(&Section.VertexBuffers.PositionVertexBuffer);
+	BeginInitResource(&Section.VertexBuffers.StaticMeshVertexBuffer);
+	BeginInitResource(&Section.VertexBuffers.ColorVertexBuffer);
+	BeginInitResource(&Section.IndexBuffer);
+	BeginInitResource(&Section.VertexFactory);
 }
 
-FWaterBodySceneProxy::FWaterBodySceneProxy(UWaterBodyComponent* Component, const TArray<FDynamicMeshVertex>& Vertices, TArray<uint32>& Indices)
+FWaterBodySceneProxy::FWaterBodySceneProxy(UWaterBodyComponent* Component)
 	: FPrimitiveSceneProxy(Component)
 {
-	if (Vertices.Num() > 0 && Indices.Num() > 0)
+	if (Component->WaterBodyMeshVertices.Num() > 0 && Component->WaterBodyMeshIndices.Num() > 0)
 	{
-		FWaterBodyMeshSection* NewSection = new FWaterBodyMeshSection(GetScene().GetFeatureLevel());
-		NewSection->IndexBuffer.Indices = Indices;
+		FWaterBodyMeshSection& NewSection = Sections.Emplace_GetRef(GetScene().GetFeatureLevel());
+		NewSection.IndexBuffer.Indices = Component->WaterBodyMeshIndices;
 
-		TArray<FDynamicMeshVertex> SectionVertices = Vertices;
-		NewSection->VertexBuffers.InitFromDynamicVertex(&NewSection->VertexFactory, SectionVertices);
+		NewSection.VertexBuffers.InitFromDynamicVertex(&NewSection.VertexFactory, Component->WaterBodyMeshVertices);
 
 		InitResources(NewSection);
+	}
+
+	if (Component->DilatedWaterBodyMeshVertices.Num() > 0 && Component->DilatedWaterBodyMeshIndices.Num() > 0)
+	{
+		FWaterBodyMeshSection& NewSection = DilatedSections.Emplace_GetRef(GetScene().GetFeatureLevel());
+		NewSection.IndexBuffer.Indices = Component->DilatedWaterBodyMeshIndices;
+
+		NewSection.VertexBuffers.InitFromDynamicVertex(&NewSection.VertexFactory, Component->DilatedWaterBodyMeshVertices);
 
-		Sections.Add(NewSection);
+		InitResources(NewSection);
 	}
 
 	if (UMaterialInstance* WaterInfoMaterial = Component->GetWaterInfoMaterialInstance())
@@ -63,15 +69,22 @@ FWaterBodySceneProxy::FWaterBodySceneProxy(UWaterBodyComponent* Component, const
 
 FWaterBodySceneProxy::~FWaterBodySceneProxy()
 {
-	for (FWaterBodyMeshSection* Section : Sections)
+	for (FWaterBodyMeshSection& Section : Sections)
 	{
-		Section->VertexBuffers.PositionVertexBuffer.ReleaseResource();
-		Section->VertexBuffers.StaticMeshVertexBuffer.ReleaseResource();
-		Section->VertexBuffers.ColorVertexBuffer.ReleaseResource();
-		Section->IndexBuffer.ReleaseResource();
-		Section->VertexFactory.ReleaseResource();
+		Section.VertexBuffers.PositionVertexBuffer.ReleaseResource();
+		Section.VertexBuffers.StaticMeshVertexBuffer.ReleaseResource();
+		Section.VertexBuffers.ColorVertexBuffer.ReleaseResource();
+		Section.IndexBuffer.ReleaseResource();
+		Section.VertexFactory.ReleaseResource();
+	}
 
-		delete Section;
+	for (FWaterBodyMeshSection& Section : DilatedSections)
+	{
+		Section.VertexBuffers.PositionVertexBuffer.ReleaseResource();
+		Section.VertexBuffers.StaticMeshVertexBuffer.ReleaseResource();
+		Section.VertexBuffers.ColorVertexBuffer.ReleaseResource();
+		Section.IndexBuffer.ReleaseResource();
+		Section.VertexFactory.ReleaseResource();
 	}
 }
 
@@ -89,8 +102,13 @@ void FWaterBodySceneProxy::GetDynamicMeshElements(const TArray<const FSceneView*
 		MaterialToUse = FallbackMaterial;
 	}
 
+	const bool bWithinWaterInfoPasses = CurrentWaterInfoPass != EWaterInfoPass::None;
+
+	const int32 ShowProxiesCVar = CVarWaterShowProxies.GetValueOnRenderThread();
+	const bool bDebugShowWaterSceneProxiesEnabled = !bWithinWaterInfoPasses && ((ShowProxiesCVar == 1 && IsSelected()) || ShowProxiesCVar >= 2);
+
 	// If we are not in the waterinfo pass and the cvar is not set to show opaque bodies, we should be in wireframe
-	const bool bWireframe = AllowDebugViewmodes() && !(bWithinWaterInfoPass) && (CVarWaterShowProxies.GetValueOnRenderThread() != 3);
+	const bool bWireframe = AllowDebugViewmodes() && !(bWithinWaterInfoPasses) && (CVarWaterShowProxies.GetValueOnRenderThread() != 3);
 
 	if (bWireframe)
 	{
@@ -102,38 +120,46 @@ void FWaterBodySceneProxy::GetDynamicMeshElements(const TArray<const FSceneView*
 		Collector.RegisterOneFrameMaterialProxy(WireframeMaterialInstance);
 		MaterialToUse = WireframeMaterialInstance;
 	}
-	
+
+	auto AddWaterBodyMeshSection = [&](const FWaterBodyMeshSection& Section, int32 ViewIndex)
+	{
+		FMeshBatch& Mesh = Collector.AllocateMesh();
+		FMeshBatchElement& BatchElement = Mesh.Elements[0];
+		BatchElement.IndexBuffer = &Section.IndexBuffer;
+		Mesh.bWireframe = bWireframe;
+		Mesh.VertexFactory = &Section.VertexFactory;
+		Mesh.MaterialRenderProxy = MaterialToUse;
+
+		bool bHasPrecomputedVolumetricLightmap;
+		FMatrix PreviousLocalToWorld;
+		int32 SingleCaptureIndex;
+		bool bOutputVelocity;
+		GetScene().GetPrimitiveUniformShaderParameters_RenderThread(GetPrimitiveSceneInfo(), bHasPrecomputedVolumetricLightmap, PreviousLocalToWorld, SingleCaptureIndex, bOutputVelocity);
+		bOutputVelocity |= AlwaysHasVelocity();
+
+		FDynamicPrimitiveUniformBuffer& DynamicPrimitiveUniformBuffer = Collector.AllocateOneFrameResource<FDynamicPrimitiveUniformBuffer>();
+		DynamicPrimitiveUniformBuffer.Set(GetLocalToWorld(), PreviousLocalToWorld, GetBounds(), GetLocalBounds(), true, bHasPrecomputedVolumetricLightmap, DrawsVelocity(), bOutputVelocity); BatchElement.PrimitiveUniformBufferResource = &DynamicPrimitiveUniformBuffer.UniformBuffer;
+
+		BatchElement.FirstIndex = 0;
+		BatchElement.NumPrimitives = Section.IndexBuffer.Indices.Num() / 3;
+		check(BatchElement.NumPrimitives != 0);
+		BatchElement.MinVertexIndex = 0;
+		BatchElement.MaxVertexIndex = Section.VertexBuffers.PositionVertexBuffer.GetNumVertices() - 1;
+		Mesh.ReverseCulling = IsLocalToWorldDeterminantNegative();
+		Mesh.Type = PT_TriangleList;
+		Mesh.DepthPriorityGroup = bWireframe ? SDPG_Foreground : SDPG_World;
+		Mesh.bCanApplyViewModeOverrides = false;
+		Collector.AddMesh(ViewIndex, Mesh);
+	};
+
 	for (int32 ViewIndex = 0; ViewIndex < Views.Num(); ++ViewIndex)
 	{
-		for (const FWaterBodyMeshSection* Section : Sections)
+		if ((CurrentWaterInfoPass == EWaterInfoPass::Color) || (bDebugShowWaterSceneProxiesEnabled))
 		{
-			FMeshBatch& Mesh = Collector.AllocateMesh();
-			FMeshBatchElement& BatchElement = Mesh.Elements[0];
-			BatchElement.IndexBuffer = &Section->IndexBuffer;
-			Mesh.bWireframe = bWireframe;
-			Mesh.VertexFactory = &Section->VertexFactory;
-			Mesh.MaterialRenderProxy = MaterialToUse;
-
-			bool bHasPrecomputedVolumetricLightmap;
-			FMatrix PreviousLocalToWorld;
-			int32 SingleCaptureIndex;
-			bool bOutputVelocity;
-			GetScene().GetPrimitiveUniformShaderParameters_RenderThread(GetPrimitiveSceneInfo(), bHasPrecomputedVolumetricLightmap, PreviousLocalToWorld, SingleCaptureIndex, bOutputVelocity);
-			bOutputVelocity |= AlwaysHasVelocity();
-
-			FDynamicPrimitiveUniformBuffer& DynamicPrimitiveUniformBuffer = Collector.AllocateOneFrameResource<FDynamicPrimitiveUniformBuffer>();
-			DynamicPrimitiveUniformBuffer.Set(GetLocalToWorld(), PreviousLocalToWorld, GetBounds(), GetLocalBounds(), true, bHasPrecomputedVolumetricLightmap, bOutputVelocity); BatchElement.PrimitiveUniformBufferResource = &DynamicPrimitiveUniformBuffer.UniformBuffer;
-
-			BatchElement.FirstIndex = 0;
-			BatchElement.NumPrimitives = Section->IndexBuffer.Indices.Num() / 3;
-			check(BatchElement.NumPrimitives != 0);
-			BatchElement.MinVertexIndex = 0;
-			BatchElement.MaxVertexIndex = Section->VertexBuffers.PositionVertexBuffer.GetNumVertices() - 1;
-			Mesh.ReverseCulling = IsLocalToWorldDeterminantNegative();
-			Mesh.Type = PT_TriangleList;
-			Mesh.DepthPriorityGroup = bWireframe ? SDPG_Foreground : SDPG_World;
-			Mesh.bCanApplyViewModeOverrides = false;
-			Collector.AddMesh(ViewIndex, Mesh);
+			for (const FWaterBodyMeshSection& Section : Sections)
+			{
+				AddWaterBodyMeshSection(Section, ViewIndex);
+			}
 		}
 	}
 }
@@ -170,7 +196,7 @@ uint32 FWaterBodySceneProxy::GetAllocatedSize() const
 
 bool FWaterBodySceneProxy::IsShown(const FSceneView* View) const
 {
-	if (!bWithinWaterInfoPass)
+	if (CurrentWaterInfoPass == EWaterInfoPass::None)
 	{
 		const int32 ShowProxiesCVar = CVarWaterShowProxies.GetValueOnRenderThread();
 		if ((ShowProxiesCVar == 1 && IsSelected()) || ShowProxiesCVar >= 2)
@@ -183,8 +209,3 @@ bool FWaterBodySceneProxy::IsShown(const FSceneView* View) const
 
 	return FPrimitiveSceneProxy::IsShown(View);
 }
-
-void FWaterBodySceneProxy::SetWithinWaterInfoPass(bool bInWithinWaterInfoPass)
-{
-	bWithinWaterInfoPass = bInWithinWaterInfoPass;
-}
diff --git a/Engine/Plugins/Experimental/Water/Source/Runtime/Private/WaterInfoRendering.cpp b/Engine/Plugins/Experimental/Water/Source/Runtime/Private/WaterInfoRendering.cpp
--- a/Engine/Plugins/Experimental/Water/Source/Runtime/Private/WaterInfoRendering.cpp
+++ b/Engine/Plugins/Experimental/Water/Source/Runtime/Private/WaterInfoRendering.cpp
@@ -31,10 +31,11 @@ namespace UE::WaterInfo
 
 struct FUpdateWaterInfoParams
 {
-	FSceneRenderer* DepthRenderer;
-	FSceneRenderer* ColorRenderer;
-	FRenderTarget* RenderTarget;
-	FTexture* OutputTexture;
+	FSceneRenderer* DepthRenderer = nullptr;
+	FSceneRenderer* ColorRenderer = nullptr;
+	FSceneRenderer* DilationRenderer = nullptr;
+	FRenderTarget* RenderTarget = nullptr;
+	FTexture* OutputTexture = nullptr;
 
 	FVector2D WaterZoneExtents;
 	FVector2f WaterHeightExtents;
@@ -58,6 +59,7 @@ public:
 		SHADER_PARAMETER_STRUCT_INCLUDE(FSceneTextureShaderParameters, SceneTextures)
 		SHADER_PARAMETER_RDG_TEXTURE(Texture2D, DepthTexture)
 		SHADER_PARAMETER_RDG_TEXTURE(Texture2D<float4>, ColorTexture)
+		SHADER_PARAMETER_RDG_TEXTURE(Texture2D, DilationTexture)
 		SHADER_PARAMETER(FVector2f, WaterHeightExtents)
 		SHADER_PARAMETER(float, GroundZMin)
 		SHADER_PARAMETER(float, CaptureZ)
@@ -85,7 +87,8 @@ static void MergeWaterInfoAndDepth(
 	FRDGTextureRef OutputTexture,
 	FRDGTextureRef DepthTexture,
 	FRDGTextureRef ColorTexture,
-	FUpdateWaterInfoParams Params)
+	FRDGTextureRef DilationTexture,
+	const FUpdateWaterInfoParams& Params)
 {
 	RDG_EVENT_SCOPE(GraphBuilder, "WaterInfoDepthMerge");
 
@@ -101,6 +104,7 @@ static void MergeWaterInfoAndDepth(
 		PassParameters->SceneTextures = SceneTextures.GetSceneTextureShaderParameters(ViewFamily.GetFeatureLevel());
 		PassParameters->DepthTexture = DepthTexture;
 		PassParameters->ColorTexture = ColorTexture;
+		PassParameters->DilationTexture = DilationTexture;
 		PassParameters->CaptureZ = Params.CaptureZ;
 		PassParameters->WaterHeightExtents = Params.WaterHeightExtents;
 		PassParameters->GroundZMin = Params.GroundZMin;
@@ -196,7 +200,7 @@ static void FinalizeWaterInfo(
 	FViewInfo& View,
 	FRDGTextureRef WaterInfoTexture,
 	FRDGTextureRef OutputTexture,
-	FUpdateWaterInfoParams Params)
+	const FUpdateWaterInfoParams& Params)
 {
 	RDG_EVENT_SCOPE(GraphBuilder, "WaterInfoFinalize");
 
@@ -276,19 +280,17 @@ static FMatrix BuildOrthoMatrix(float InOrthoWidth, float InOrthoHeight)
 
 // ---------------------------------------------------------------------------------------------------------------------
 
-static void SetWaterBodiesWithinWaterInfoPass(FSceneRenderer* SceneRenderer, bool bWithinWaterInfoPass)
+static void SetWaterBodiesWithinWaterInfoPass(FSceneRenderer* SceneRenderer, EWaterInfoPass InPass)
 {
+	TRACE_CPUPROFILER_EVENT_SCOPE(WaterInfo::SetWaterBodiesWithinWaterInfoPass);
 	if (SceneRenderer->Views[0].ShowOnlyPrimitives.IsSet())
 	{
-		for (FPrimitiveComponentId PrimId : SceneRenderer->Views[0].ShowOnlyPrimitives.GetValue())
+		for (FPrimitiveSceneProxy* PrimProxy : SceneRenderer->Scene->PrimitiveSceneProxies)
 		{
-			for (FPrimitiveSceneProxy* PrimProxy : SceneRenderer->Scene->PrimitiveSceneProxies)
+			if (PrimProxy && SceneRenderer->Views[0].ShowOnlyPrimitives->Contains(PrimProxy->GetPrimitiveComponentId()))
 			{
-				if (PrimProxy && PrimProxy->GetPrimitiveComponentId() == PrimId)
-				{
-					FWaterBodySceneProxy* WaterProxy = (FWaterBodySceneProxy*)PrimProxy;
-					WaterProxy->SetWithinWaterInfoPass(bWithinWaterInfoPass);
-				}
+				FWaterBodySceneProxy* WaterProxy = (FWaterBodySceneProxy*)PrimProxy;
+				WaterProxy->SetWithinWaterInfoPass(InPass);
 			}
 		}
 	}
@@ -397,6 +399,9 @@ static void UpdateWaterInfoRendering_RenderThread(
 	FRenderTarget* RenderTarget = Params.RenderTarget;
 	FTexture* OutputTexture = Params.OutputTexture;
 
+	TRefCountPtr<IPooledRenderTarget> ExtractedDepthTexture;
+	TRefCountPtr<IPooledRenderTarget> ExtractedColorTexture;
+
 	// Depth-only pass for actors which are considered the ground for water rendering
 	{
 		FSceneRenderer* DepthRenderer = Params.DepthRenderer;
@@ -413,7 +418,9 @@ static void UpdateWaterInfoRendering_RenderThread(
 		FRDGBuilder GraphBuilder(RHICmdList, RDG_EVENT_NAME("WaterInfoColorRendering"), ERDGBuilderFlags::AllowParallelExecute);
 		
 		FRDGTextureRef TargetTexture = RegisterExternalTexture(GraphBuilder, RenderTarget->GetRenderTargetTexture(), TEXT("WaterDepthTarget"));
-		
+
+		FRDGTextureDesc DepthTextureDesc(TargetTexture->Desc);
+		FRDGTextureRef DepthTexture = GraphBuilder.CreateTexture(DepthTextureDesc, TEXT("WaterDepthTexture"));
 
 		FViewInfo& View = DepthRenderer->Views[0];
 
@@ -426,10 +433,7 @@ static void UpdateWaterInfoRendering_RenderThread(
 			RDG_RHI_EVENT_SCOPE(GraphBuilder, RenderWaterInfoDepth);
 			DepthRenderer->Render(GraphBuilder);
 		}
-
-		FRDGTextureRef ShaderResourceTexture = RegisterExternalTexture(GraphBuilder, OutputTexture->TextureRHI, TEXT("WaterDepthTexture"));
-		AddCopyTexturePass(GraphBuilder, TargetTexture, ShaderResourceTexture);
-
+		
 		if (DepthRenderer->Scene->GetShadingPath() == EShadingPath::Mobile)
 		{
 			const FMinimalSceneTextures& SceneTextures = View.GetSceneTextures();
@@ -438,12 +442,21 @@ static void UpdateWaterInfoRendering_RenderThread(
 			CopySceneCaptureComponentToTarget(
 				GraphBuilder,
 				SceneTextures,
-				ShaderResourceTexture,
+				DepthTexture,
 				DepthRenderer->ViewFamily,
 				DepthRenderer->Views,
 				bNeedsFlippedRenderTarget);
 		}
+		else
+		{
+			FResolveParams ResolveParams;
+			AddCopyToResolveTargetPass(GraphBuilder, TargetTexture, DepthTexture, ResolveParams);
+		}
 		
+		// We currently can't have multiple scene renderers run within the same RDGBuilder. Therefore, we must
+		// extract the texture to allow it to survive until the water info pass which runs in a later RDG graph (however, still within the same frame).
+		GraphBuilder.QueueTextureExtraction(DepthTexture, &ExtractedDepthTexture);
+
 		GraphBuilder.Execute();
 
 		DepthRenderer->RenderThreadEnd(RHICmdList);
@@ -457,16 +470,17 @@ static void UpdateWaterInfoRendering_RenderThread(
 		FSceneRenderer::ViewExtensionPreRender_RenderThread(RHICmdList, ColorRenderer);
 
 		ColorRenderer->RenderThreadBegin(RHICmdList);
-
-		SetWaterBodiesWithinWaterInfoPass(ColorRenderer, true);
+		SetWaterBodiesWithinWaterInfoPass(ColorRenderer, EWaterInfoPass::Color);
 		
 		SCOPED_DRAW_EVENT(RHICmdList, ColorRendering_RT);
 		
 		FRDGBuilder GraphBuilder(RHICmdList, RDG_EVENT_NAME("WaterInfoColorRendering"), ERDGBuilderFlags::AllowParallelExecute);
-		FRDGTextureRef TargetTexture = RegisterExternalTexture(GraphBuilder, RenderTarget->GetRenderTargetTexture(), TEXT("WaterInfoTarget"));
 
-		FRDGTextureRef DepthTexture = GraphBuilder.CreateTexture(TargetTexture->Desc, TEXT("WaterInfoDepth"));
-		AddCopyTexturePass(GraphBuilder, TargetTexture, DepthTexture);
+		FRDGTextureRef TargetTexture = RegisterExternalTexture(GraphBuilder, RenderTarget->GetRenderTargetTexture(), TEXT("WaterColorTarget"));
+
+		FRDGTextureDesc ColorTextureDesc(TargetTexture->Desc);
+		ColorTextureDesc.Format = PF_A32B32G32R32F;
+		FRDGTextureRef ColorTexture = GraphBuilder.CreateTexture(ColorTextureDesc, TEXT("WaterColorTexture"));
 		
 		FViewInfo& View = ColorRenderer->Views[0];
 
@@ -481,9 +495,8 @@ static void UpdateWaterInfoRendering_RenderThread(
 		}
 
 		const FMinimalSceneTextures& SceneTextures = View.GetSceneTextures();
-		FRDGTextureDesc ColorTextureDesc(TargetTexture->Desc);
-		ColorTextureDesc.Format = PF_A32B32G32R32F;
-		FRDGTextureRef ColorTexture = GraphBuilder.CreateTexture(ColorTextureDesc, TEXT("WaterInfoColor"));
+
+		// This CopySceneCaptureComponentToTarget is required on all platforms as it extracts a higher pixel depth texture than the scene render target so we can't just copy out of it
 		{
 			const bool bNeedsFlippedRenderTarget = false;
 			RDG_EVENT_SCOPE(GraphBuilder, "CaptureSceneColor");
@@ -496,26 +509,110 @@ static void UpdateWaterInfoRendering_RenderThread(
 				bNeedsFlippedRenderTarget);
 		}
 
-		FRDGTextureRef MergeTargetTexture = GraphBuilder.CreateTexture(ColorTextureDesc, TEXT("WaterInfoMerged"));
-		MergeWaterInfoAndDepth(GraphBuilder, SceneTextures, ColorRenderer->ViewFamily, ColorRenderer->Views[0], MergeTargetTexture, DepthTexture, ColorTexture, Params);
+		// We currently can't have multiple scene renderers run within the same RDGBuilder. Therefore, we must
+		// extract the texture to allow it to survive until the water info pass which runs in a later RDG graph (however, still within the same frame).
+		GraphBuilder.QueueTextureExtraction(ColorTexture, &ExtractedColorTexture);
+
+		GraphBuilder.Execute();
+		
+		ColorRenderer->RenderThreadEnd(RHICmdList);
+	}
+
+	// Depth-only pass for water body dilated sections which get composited back into the water info texture at a lower priority than regular water body data
+	{
+		FSceneRenderer* DilationRenderer = Params.DilationRenderer;
+
+		// We need to execute the pre-render view extensions before we do any view dependent work.
+		FSceneRenderer::ViewExtensionPreRender_RenderThread(RHICmdList, DilationRenderer);
+
+		DilationRenderer->RenderThreadBegin(RHICmdList);
+		SetWaterBodiesWithinWaterInfoPass(DilationRenderer, EWaterInfoPass::Dilation);
+
+		SCOPED_DRAW_EVENT(RHICmdList, DilationRendering_RT);
+		
+		FRDGBuilder GraphBuilder(RHICmdList, RDG_EVENT_NAME("WaterInfoColorRendering"), ERDGBuilderFlags::AllowParallelExecute);
+		FRDGTextureRef TargetTexture = RegisterExternalTexture(GraphBuilder, RenderTarget->GetRenderTargetTexture(), TEXT("WaterDilationTarget"));
+
+		FRDGTextureRef DepthTexture = GraphBuilder.RegisterExternalTexture(ExtractedDepthTexture);
+		FRDGTextureRef ColorTexture = GraphBuilder.RegisterExternalTexture(ExtractedColorTexture);
+
+		FViewInfo& View = DilationRenderer->Views[0];
+
+		AddClearRenderTargetPass(GraphBuilder, TargetTexture, FLinearColor::Black, View.UnscaledViewRect);
+
+		View.bDisableQuerySubmissions = true;
+		View.bIgnoreExistingQueries = true;
+
+		{
+			RDG_RHI_EVENT_SCOPE(GraphBuilder, RenderWaterInfoDilation);
+			DilationRenderer->Render(GraphBuilder);
+		}
+
+		const FMinimalSceneTextures& SceneTextures = View.GetSceneTextures();
+		FRDGTextureDesc TextureDesc(TargetTexture->Desc);
+		FRDGTextureRef DilationTexture = GraphBuilder.CreateTexture(TextureDesc, TEXT("WaterDilationTexture"));
+
+		if (DilationRenderer->Scene->GetShadingPath() == EShadingPath::Mobile)
+		{
+			const bool bNeedsFlippedRenderTarget = false;
+			RDG_EVENT_SCOPE(GraphBuilder, "CaptureSceneColor");
+			CopySceneCaptureComponentToTarget(
+				GraphBuilder,
+				SceneTextures,
+				DilationTexture,
+				DilationRenderer->ViewFamily,
+				DilationRenderer->Views,
+				bNeedsFlippedRenderTarget);
+		}
+		else
+		{
+			FResolveParams ResolveParams;
+			AddCopyToResolveTargetPass(GraphBuilder, TargetTexture, DilationTexture, ResolveParams);
+		}
+
+		FRDGTextureRef MergeTargetTexture = GraphBuilder.CreateTexture(TextureDesc, TEXT("WaterInfoMerged"));
+		MergeWaterInfoAndDepth(GraphBuilder, SceneTextures, DilationRenderer->ViewFamily, DilationRenderer->Views[0], MergeTargetTexture, DepthTexture, ColorTexture, DilationTexture, Params);
 
 		FRDGTextureRef FinalizedTexture = GraphBuilder.CreateTexture(TargetTexture->Desc, TEXT("WaterInfoFinalized"));
-		FinalizeWaterInfo(GraphBuilder, SceneTextures, ColorRenderer->ViewFamily, ColorRenderer->Views[0], MergeTargetTexture, FinalizedTexture, Params);
+		FinalizeWaterInfo(GraphBuilder, SceneTextures, DilationRenderer->ViewFamily, DilationRenderer->Views[0], MergeTargetTexture, FinalizedTexture, Params);
 		
 		FRDGTextureRef ShaderResourceTexture = RegisterExternalTexture(GraphBuilder, OutputTexture->TextureRHI, TEXT("WaterInfoResolve"));
 		AddCopyTexturePass(GraphBuilder, FinalizedTexture, ShaderResourceTexture);
 		GraphBuilder.Execute();
-		
-		SetWaterBodiesWithinWaterInfoPass(ColorRenderer, false);
 
-		ColorRenderer->RenderThreadEnd(RHICmdList);
+		SetWaterBodiesWithinWaterInfoPass(DilationRenderer, EWaterInfoPass::None);
+		DilationRenderer->RenderThreadEnd(RHICmdList);
 	}
 
 	RHICmdList.Transition(FRHITransitionInfo(Params.OutputTexture->TextureRHI, ERHIAccess::RTV, ERHIAccess::SRVMask));
 }
 
-static FEngineShowFlags GetWaterInfoBaseShowFlags()
+struct FCreateWaterInfoSceneRendererParams
 {
+public:
+	FCreateWaterInfoSceneRendererParams(const WaterInfo::FRenderingContext& InContext) : Context(InContext) {}
+
+	const WaterInfo::FRenderingContext& Context;
+
+	FSceneInterface* Scene = nullptr;
+	FRenderTarget* RenderTarget = nullptr;
+
+	FIntPoint RenderTargetSize = FIntPoint(EForceInit::ForceInit);
+	FMatrix ViewRotationMatrix = FMatrix(EForceInit::ForceInit);
+	FVector ViewLocation = FVector::Zero();
+	FMatrix ProjectionMatrix = FMatrix(EForceInit::ForceInit);
+	ESceneCaptureSource CaptureSource = ESceneCaptureSource::SCS_MAX;
+	TSet<FPrimitiveComponentId> ShowOnlyPrimitives;
+};
+
+static FSceneRenderer* CreateWaterInfoSceneRenderer(const FCreateWaterInfoSceneRendererParams& Params)
+{
+	TRACE_CPUPROFILER_EVENT_SCOPE(WaterInfo::CreateWaterInfoRenderer);
+
+	check(Params.Scene != nullptr)
+	check(Params.RenderTarget != nullptr)
+	check(Params.CaptureSource != ESceneCaptureSource::SCS_MAX)
+
 	FEngineShowFlags ShowFlags(ESFIM_Game);
 	ShowFlags.NaniteMeshes = 0;
 	ShowFlags.Atmosphere = 0;
@@ -528,161 +625,63 @@ static FEngineShowFlags GetWaterInfoBaseShowFlags()
 	ShowFlags.Fog = 0;
 	ShowFlags.VolumetricFog = 0;
 	ShowFlags.DynamicShadows = 0;
-	return ShowFlags;
-}
-
-static FSceneRenderer* CreateWaterInfoDepthRenderer(
-	FSceneInterface* Scene,
-	FRenderTarget* RenderTarget,
-	const WaterInfo::FRenderingContext& Context,
-	FIntPoint RenderTargetSize,
-	const FMatrix& ViewRotationMatrix,
-	const FVector& ViewLocation,
-	const FMatrix& ProjectionMatrix)
-{
-	FEngineShowFlags ShowFlags = GetWaterInfoBaseShowFlags();
-
-	FSceneViewFamilyContext DepthViewFamily(FSceneViewFamily::ConstructionValues(
-		RenderTarget,
-		Scene,
-		ShowFlags)
-		.SetRealtimeUpdate(false)
-		.SetResolveScene(false));
-	DepthViewFamily.SceneCaptureSource = ESceneCaptureSource::SCS_DeviceDepth;
-
-	// Setup the view family
-	FSceneViewInitOptions DepthViewInitOptions;
-	DepthViewInitOptions.SetViewRectangle(FIntRect(0, 0, RenderTargetSize.X, RenderTargetSize.Y));
-	DepthViewInitOptions.ViewFamily = &DepthViewFamily;
-	DepthViewInitOptions.ViewActor = Context.ZoneToRender;
-	DepthViewInitOptions.ViewRotationMatrix = ViewRotationMatrix;
-	DepthViewInitOptions.ViewOrigin = ViewLocation;
-	DepthViewInitOptions.BackgroundColor = FLinearColor::Black;
-	DepthViewInitOptions.OverrideFarClippingPlaneDistance = -1.f;
-	DepthViewInitOptions.SceneViewStateInterface = nullptr;
-	DepthViewInitOptions.ProjectionMatrix = ProjectionMatrix;
-	DepthViewInitOptions.LODDistanceFactor = 0.001f;
-	DepthViewInitOptions.OverlayColor = FLinearColor::Black;
-
-	if (DepthViewFamily.Scene->GetWorld() != nullptr && DepthViewFamily.Scene->GetWorld()->GetWorldSettings() != nullptr)
-	{
-		DepthViewInitOptions.WorldToMetersScale = DepthViewFamily.Scene->GetWorld()->GetWorldSettings()->WorldToMeters;
-	}
-
-	FSceneView* DepthView = new FSceneView(DepthViewInitOptions);
-	DepthView->AntiAliasingMethod = EAntiAliasingMethod::AAM_None;
-	DepthView->SetupAntiAliasingMethod();
-
-	if (Context.GroundActors.Num() > 0)
-	{
-		DepthView->ShowOnlyPrimitives.Emplace();
-		DepthView->ShowOnlyPrimitives->Reserve(Context.GroundActors.Num());
-		for (TWeakObjectPtr<AActor> GroundActor : Context.GroundActors)
-		{
-			if (GroundActor.IsValid())
-			{
-				TInlineComponentArray<UPrimitiveComponent*> PrimComps(GroundActor.Get());
-				for (UPrimitiveComponent* PrimComp : PrimComps)
-				{
-					if (PrimComp)
-					{
-						DepthView->ShowOnlyPrimitives->Add(PrimComp->ComponentId);
-					}
-				}
-			}
-		}
-	}
-
-	DepthViewFamily.Views.Add(DepthView);
 	
-	DepthView->StartFinalPostprocessSettings(ViewLocation);
-	DepthView->EndFinalPostprocessSettings(DepthViewInitOptions);
-
-	DepthViewFamily.SetScreenPercentageInterface(new FLegacyScreenPercentageDriver(DepthViewFamily, 1.f));
-
-	DepthViewFamily.ViewExtensions = GEngine->ViewExtensions->GatherActiveExtensions(FSceneViewExtensionContext(Scene));
-	for (const FSceneViewExtensionRef& Extension : DepthViewFamily.ViewExtensions)
-	{
-		Extension->SetupViewFamily(DepthViewFamily);
-		Extension->SetupView(DepthViewFamily, *DepthView);
-	}
-
-	return FSceneRenderer::CreateSceneRenderer(&DepthViewFamily, nullptr);
-}
-
-static FSceneRenderer* CreateWaterInfoColorRenderer(
-	FSceneInterface* Scene,
-	FRenderTarget* RenderTarget,
-	const WaterInfo::FRenderingContext& Context,
-	FIntPoint RenderTargetSize,
-	const FMatrix& ViewRotationMatrix,
-	const FVector& ViewLocation,
-	const FMatrix& ProjectionMatrix)
-{
-	FEngineShowFlags ShowFlags = GetWaterInfoBaseShowFlags();
-	
-	FSceneViewFamilyContext ColorViewFamily(FSceneViewFamily::ConstructionValues(
-		RenderTarget,
-		Scene,
+	FSceneViewFamilyContext ViewFamily(FSceneViewFamily::ConstructionValues(
+		Params.RenderTarget,
+		Params.Scene,
 		ShowFlags)
 		.SetRealtimeUpdate(false)
 		.SetResolveScene(false));
-	ColorViewFamily.SceneCaptureSource = ESceneCaptureSource::SCS_SceneColorSceneDepth;
-
-	FSceneViewInitOptions ColorViewInitOptions;
-	ColorViewInitOptions.SetViewRectangle(FIntRect(0, 0, RenderTargetSize.X, RenderTargetSize.Y));
-	ColorViewInitOptions.ViewFamily = &ColorViewFamily;
-	ColorViewInitOptions.ViewActor = Context.ZoneToRender;
-	ColorViewInitOptions.ViewRotationMatrix = ViewRotationMatrix;
-	ColorViewInitOptions.ViewOrigin = ViewLocation;
-	ColorViewInitOptions.BackgroundColor = FLinearColor::Black;
-	ColorViewInitOptions.OverrideFarClippingPlaneDistance = -1.f;
-	ColorViewInitOptions.SceneViewStateInterface = nullptr;
-	ColorViewInitOptions.ProjectionMatrix = ProjectionMatrix;
-	ColorViewInitOptions.LODDistanceFactor = 0.001f;
-	ColorViewInitOptions.OverlayColor = FLinearColor::Black;
-
-	if (ColorViewFamily.Scene->GetWorld() != nullptr && ColorViewFamily.Scene->GetWorld()->GetWorldSettings() != nullptr)
+	ViewFamily.SceneCaptureSource = Params.CaptureSource;
+
+	FSceneViewInitOptions ViewInitOptions;
+	ViewInitOptions.SetViewRectangle(FIntRect(0, 0, Params.RenderTargetSize.X, Params.RenderTargetSize.Y));
+	ViewInitOptions.ViewFamily = &ViewFamily;
+	ViewInitOptions.ViewActor = Params.Context.ZoneToRender;
+	ViewInitOptions.ViewRotationMatrix = Params.ViewRotationMatrix;
+	ViewInitOptions.ViewOrigin = Params.ViewLocation;
+	ViewInitOptions.BackgroundColor = FLinearColor::Black;
+	ViewInitOptions.OverrideFarClippingPlaneDistance = -1.f;
+	ViewInitOptions.SceneViewStateInterface = nullptr;
+	ViewInitOptions.ProjectionMatrix = Params.ProjectionMatrix;
+	ViewInitOptions.LODDistanceFactor = 0.001f;
+	ViewInitOptions.OverlayColor = FLinearColor::Black;
+
+	if (ViewFamily.Scene->GetWorld() != nullptr && ViewFamily.Scene->GetWorld()->GetWorldSettings() != nullptr)
 	{
-		ColorViewInitOptions.WorldToMetersScale = ColorViewFamily.Scene->GetWorld()->GetWorldSettings()->WorldToMeters;
+		ViewInitOptions.WorldToMetersScale = ViewFamily.Scene->GetWorld()->GetWorldSettings()->WorldToMeters;
 	}
 
-	FSceneView* ColorView = new FSceneView(ColorViewInitOptions);
-	ColorView->bIsSceneCapture = true;
-	ColorView->AntiAliasingMethod = EAntiAliasingMethod::AAM_None;
-	ColorView->SetupAntiAliasingMethod();
+	FSceneView* View = new FSceneView(ViewInitOptions);
+	View->bIsSceneCapture = true;
+	View->AntiAliasingMethod = EAntiAliasingMethod::AAM_None;
+	View->SetupAntiAliasingMethod();
 
-	if (Context.WaterBodies.Num() > 0)
-	{
-		ColorView->ShowOnlyPrimitives.Emplace();
-		ColorView->ShowOnlyPrimitives->Reserve(Context.WaterBodies.Num());
-		for (const UWaterBodyComponent* WaterBodyToRender : Context.WaterBodies)
-		{
-			ColorView->ShowOnlyPrimitives->Add(WaterBodyToRender->ComponentId);
-		}
-	}
+	View->ShowOnlyPrimitives = Params.ShowOnlyPrimitives;
 
-	ColorViewFamily.Views.Add(ColorView);
+	ViewFamily.Views.Add(View);
 
-	ColorView->StartFinalPostprocessSettings(ViewLocation);
-	ColorView->EndFinalPostprocessSettings(ColorViewInitOptions);
+	View->StartFinalPostprocessSettings(Params.ViewLocation);
+	View->EndFinalPostprocessSettings(ViewInitOptions);
 
-	ColorViewFamily.SetScreenPercentageInterface(new FLegacyScreenPercentageDriver(ColorViewFamily, 1.f));
+	ViewFamily.SetScreenPercentageInterface(new FLegacyScreenPercentageDriver(ViewFamily, 1.f));
 
-	ColorViewFamily.ViewExtensions = GEngine->ViewExtensions->GatherActiveExtensions(FSceneViewExtensionContext(Scene));
-	for (const FSceneViewExtensionRef& Extension : ColorViewFamily.ViewExtensions)
+	ViewFamily.ViewExtensions = GEngine->ViewExtensions->GatherActiveExtensions(FSceneViewExtensionContext(Params.Scene));
+	for (const FSceneViewExtensionRef& Extension : ViewFamily.ViewExtensions)
 	{
-		Extension->SetupViewFamily(ColorViewFamily);
-		Extension->SetupView(ColorViewFamily, *ColorView);
+		Extension->SetupViewFamily(ViewFamily);
+		Extension->SetupView(ViewFamily, *View);
 	}
 
-	return FSceneRenderer::CreateSceneRenderer(&ColorViewFamily, nullptr);
+	return FSceneRenderer::CreateSceneRenderer(&ViewFamily, nullptr);
 }
 
 void UpdateWaterInfoRendering(
 	FSceneInterface* Scene,
 	const WaterInfo::FRenderingContext& Context)
 {
+	TRACE_CPUPROFILER_EVENT_SCOPE(WaterInfo::UpdateWaterInfoRendering);
+	
 	RenderCaptureInterface::FScopedCapture RenderCapture((RenderCaptureNextWaterInfoDraws != 0), TEXT("RenderWaterInfo"));
 	RenderCaptureNextWaterInfoDraws = FMath::Max(0, RenderCaptureNextWaterInfoDraws - 1);
 
@@ -697,38 +696,65 @@ void UpdateWaterInfoRendering(
 
 	// Zone rendering always happens facing towards negative z.
 	const FVector LookAt = ViewLocation - FVector(0.f, 0.f, 1.f);
-
-	FMatrix ViewRotationMat = FLookAtMatrix(ViewLocation, LookAt, FVector(0.f, -1.f, 0.f));
-	ViewRotationMat = ViewRotationMat.RemoveTranslation();
-	ViewRotationMat.RemoveScaling();
 	
 	const FIntPoint CaptureExtent(Context.TextureRenderTarget->GetSurfaceWidth(), Context.TextureRenderTarget->GetSurfaceHeight());
 
-	const FMatrix OrthoProj = BuildOrthoMatrix(ZoneExtent.X, ZoneExtent.Y);
-
-	FSceneRenderer* DepthRenderer = CreateWaterInfoDepthRenderer(
-		Scene,
-		Context.TextureRenderTarget->GameThread_GetRenderTargetResource(),
-		Context,
-		CaptureExtent,
-		ViewRotationMat,
-		ViewLocation,
-		OrthoProj);
+	// Initialize the generic parameters which are passed to each of the scene renderers
+	FCreateWaterInfoSceneRendererParams CreateSceneRendererParams(Context);
+	CreateSceneRendererParams.Scene = Scene;
+	CreateSceneRendererParams.RenderTarget = Context.TextureRenderTarget->GameThread_GetRenderTargetResource();
+	CreateSceneRendererParams.RenderTargetSize = CaptureExtent;
+	CreateSceneRendererParams.ViewLocation = ViewLocation;
+	CreateSceneRendererParams.ProjectionMatrix = BuildOrthoMatrix(ZoneExtent.X, ZoneExtent.Y);
+	CreateSceneRendererParams.ViewRotationMatrix = FLookAtMatrix(ViewLocation, LookAt, FVector(0.f, -1.f, 0.f));
+	CreateSceneRendererParams.ViewRotationMatrix = CreateSceneRendererParams.ViewRotationMatrix.RemoveTranslation();
+	CreateSceneRendererParams.ViewRotationMatrix.RemoveScaling();
+
+	TSet<FPrimitiveComponentId> ComponentsToRenderInDepthPass;
+	if (Context.GroundActors.Num() > 0)
+	{
+		ComponentsToRenderInDepthPass.Reserve(Context.GroundActors.Num());
+		for (TWeakObjectPtr<AActor> GroundActor : Context.GroundActors)
+		{
+			if (GroundActor.IsValid())
+			{
+				TInlineComponentArray<UPrimitiveComponent*> PrimComps(GroundActor.Get());
+				for (UPrimitiveComponent* PrimComp : PrimComps)
+				{
+					if (PrimComp)
+					{
+						ComponentsToRenderInDepthPass.Add(PrimComp->ComponentId);
+					}
+				}
+			}
+		}
+	}
+	CreateSceneRendererParams.CaptureSource = SCS_DeviceDepth;
+	CreateSceneRendererParams.ShowOnlyPrimitives = MoveTemp(ComponentsToRenderInDepthPass);
+	FSceneRenderer* DepthRenderer = CreateWaterInfoSceneRenderer(CreateSceneRendererParams);
 	
-	FSceneRenderer* ColorRenderer = CreateWaterInfoColorRenderer(
-		Scene,
-		Context.TextureRenderTarget->GameThread_GetRenderTargetResource(),
-		Context,
-		CaptureExtent,
-		ViewRotationMat,
-		ViewLocation,
-		OrthoProj);
+	TSet<FPrimitiveComponentId> ComponentsToRenderInColorAndDilationPass;
+	if (Context.WaterBodies.Num() > 0)
+	{
+		ComponentsToRenderInColorAndDilationPass.Reserve(Context.WaterBodies.Num());
+		for (const UWaterBodyComponent* WaterBodyToRender : Context.WaterBodies)
+		{
+			ComponentsToRenderInColorAndDilationPass.Add(WaterBodyToRender->ComponentId);
+		}
+	}
+	CreateSceneRendererParams.CaptureSource = SCS_SceneColorSceneDepth;
+	CreateSceneRendererParams.ShowOnlyPrimitives = MoveTemp(ComponentsToRenderInColorAndDilationPass);
+	FSceneRenderer* ColorRenderer = CreateWaterInfoSceneRenderer(CreateSceneRendererParams);
+
+	CreateSceneRendererParams.CaptureSource = SCS_DeviceDepth;
+	FSceneRenderer* DilationRenderer = CreateWaterInfoSceneRenderer(CreateSceneRendererParams);
 
 	FTextureRenderTargetResource* TextureRenderTargetResource = Context.TextureRenderTarget->GameThread_GetRenderTargetResource();
 
 	FUpdateWaterInfoParams Params;
 	Params.DepthRenderer = DepthRenderer;
 	Params.ColorRenderer = ColorRenderer;
+	Params.DilationRenderer = DilationRenderer;
 	Params.RenderTarget = TextureRenderTargetResource;
 	Params.OutputTexture = TextureRenderTargetResource;
 	Params.CaptureZ = ViewLocation.Z;
diff --git a/Engine/Plugins/Experimental/Water/Source/Runtime/Public/WaterBodyComponent.h b/Engine/Plugins/Experimental/Water/Source/Runtime/Public/WaterBodyComponent.h
--- a/Engine/Plugins/Experimental/Water/Source/Runtime/Public/WaterBodyComponent.h
+++ b/Engine/Plugins/Experimental/Water/Source/Runtime/Public/WaterBodyComponent.h
@@ -83,6 +83,7 @@ class WATER_API UWaterBodyComponent : public UPrimitiveComponent
 	GENERATED_UCLASS_BODY()
 
 	friend class AWaterBody;
+	friend class FWaterBodySceneProxy;
 public:
 	virtual bool AffectsLandscape() const;
 	virtual bool AffectsWaterMesh() const;
@@ -299,8 +300,6 @@ public:
 	/** Returns the minimum and maximum Z of the water surface, including waves */
 	virtual void GetSurfaceMinMaxZ(float& OutMinZ, float& OutMaxZ) const;
 
-	float GetShapeDilationZOffset() const { return ShapeDilationZOffset; }
-
 	virtual ALandscapeProxy* FindLandscape() const;
 
 	/** Returns what can be considered the single water depth of the water surface.
@@ -535,13 +534,13 @@ protected:
 	UPROPERTY(EditAnywhere, BlueprintReadWrite, Category = Navigation, meta = (EditCondition = "bCanAffectNavigation && bGenerateCollisions"))
 	TSubclassOf<UNavAreaBase> WaterNavAreaClass;
 
-	/** Push down the dilated region to prevent overwriting adjacent water body data */
-	UPROPERTY(EditAnywhere, BlueprintReadWrite, AdvancedDisplay, Category = Rendering)
-	float ShapeDilationZOffset = -64.f;
-
+	// #todo_water [roey]: these need to be refactored into individual section structs and serialized
 	TArray<FDynamicMeshVertex> WaterBodyMeshVertices;
 	TArray<uint32> WaterBodyMeshIndices;
 
+	TArray<FDynamicMeshVertex> DilatedWaterBodyMeshVertices;
+	TArray<uint32> DilatedWaterBodyMeshIndices;
+
 #if WITH_EDITORONLY_DATA
 	UPROPERTY()
 	bool bFillCollisionUnderWaterBodiesForNavmesh_DEPRECATED;
diff --git a/Engine/Plugins/Experimental/Water/Source/Runtime/Public/WaterBodyRiverComponent.h b/Engine/Plugins/Experimental/Water/Source/Runtime/Public/WaterBodyRiverComponent.h
--- a/Engine/Plugins/Experimental/Water/Source/Runtime/Public/WaterBodyRiverComponent.h
+++ b/Engine/Plugins/Experimental/Water/Source/Runtime/Public/WaterBodyRiverComponent.h
@@ -31,8 +31,6 @@ public:
 	void SetLakeTransitionMaterial(UMaterialInterface* InMat);
 	void SetOceanTransitionMaterial(UMaterialInterface* InMat);
 
-	float GetShapeDilationZOffsetFar() const { return ShapeDilationZOffsetFar; }
-
 protected:
 	/** UWaterBodyComponent Interface */
 	virtual void Reset() override;
@@ -70,8 +68,4 @@ protected:
 
 	UPROPERTY(Category = Debug, VisibleInstanceOnly, Transient, NonPIEDuplicateTransient, TextExportTransient, meta = (DisplayAfter = "OceanTransitionMaterial"))
 	UMaterialInstanceDynamic* OceanTransitionMID;
-
-	/** How far to push down the furthest vertices of the dilated portion. */
-	UPROPERTY(EditAnywhere, BlueprintReadWrite, AdvancedDisplay, Category = Rendering)
-	float ShapeDilationZOffsetFar = -128.f;
 };
\ No newline at end of file
diff --git a/Engine/Plugins/Experimental/Water/Source/Runtime/Public/WaterBodySceneProxy.h b/Engine/Plugins/Experimental/Water/Source/Runtime/Public/WaterBodySceneProxy.h
--- a/Engine/Plugins/Experimental/Water/Source/Runtime/Public/WaterBodySceneProxy.h
+++ b/Engine/Plugins/Experimental/Water/Source/Runtime/Public/WaterBodySceneProxy.h
@@ -15,10 +15,18 @@ class UWaterBodyRiverComponent;
 class FWaterBodyMeshSection;
 class FMaterialRenderProxy;
 
+enum class EWaterInfoPass
+{
+	None,
+	Depth,
+	Color,
+	Dilation,
+};
+
 class FWaterBodySceneProxy final : public FPrimitiveSceneProxy
 {
 public:
-	FWaterBodySceneProxy(UWaterBodyComponent* Component, const TArray<FDynamicMeshVertex>& Vertices, TArray<uint32>& Indices);
+	FWaterBodySceneProxy(UWaterBodyComponent* Component);
 	virtual ~FWaterBodySceneProxy();
 	virtual void GetDynamicMeshElements(const TArray<const FSceneView*>& Views, const FSceneViewFamily& ViewFamily, uint32 VisibilityMap, FMeshElementCollector& Collector) const override;
 	virtual FPrimitiveViewRelevance GetViewRelevance(const FSceneView* View) const override;
@@ -28,15 +36,16 @@ public:
 
 	bool IsShown(const FSceneView* View) const;
 
-	bool IsWithinWaterInfoPass() const { return bWithinWaterInfoPass; }
-	void SetWithinWaterInfoPass(bool bInWithinWaterInfoPass);
+	bool IsWithinWaterInfoPass(EWaterInfoPass InPass) const { return CurrentWaterInfoPass == InPass; }
+	void SetWithinWaterInfoPass(EWaterInfoPass InPass) { CurrentWaterInfoPass = InPass; }
 
 private:
-	void InitResources(FWaterBodyMeshSection* Section);
+	void InitResources(FWaterBodyMeshSection& Section);
 
-	TArray<FWaterBodyMeshSection*> Sections;
-	FMaterialRenderProxy* Material;
+	TArray<FWaterBodyMeshSection> Sections;
+	TArray<FWaterBodyMeshSection> DilatedSections;
+	FMaterialRenderProxy* Material = nullptr;
 
-	bool bWithinWaterInfoPass = false;
+	EWaterInfoPass CurrentWaterInfoPass = EWaterInfoPass::None;
 };
 
