Fixed deprecation warnings after merge

diff --git a/Engine/Plugins/Experimental/Water/Source/Runtime/Private/WaterBodySceneProxy.cpp b/Engine/Plugins/Experimental/Water/Source/Runtime/Private/WaterBodySceneProxy.cpp
--- a/Engine/Plugins/Experimental/Water/Source/Runtime/Private/WaterBodySceneProxy.cpp
+++ b/Engine/Plugins/Experimental/Water/Source/Runtime/Private/WaterBodySceneProxy.cpp
@@ -138,7 +138,8 @@ void FWaterBodySceneProxy::GetDynamicMeshElements(const TArray<const FSceneView*
 		bOutputVelocity |= AlwaysHasVelocity();
 
 		FDynamicPrimitiveUniformBuffer& DynamicPrimitiveUniformBuffer = Collector.AllocateOneFrameResource<FDynamicPrimitiveUniformBuffer>();
-		DynamicPrimitiveUniformBuffer.Set(GetLocalToWorld(), PreviousLocalToWorld, GetBounds(), GetLocalBounds(), true, bHasPrecomputedVolumetricLightmap, DrawsVelocity(), bOutputVelocity); BatchElement.PrimitiveUniformBufferResource = &DynamicPrimitiveUniformBuffer.UniformBuffer;
+		DynamicPrimitiveUniformBuffer.Set(GetLocalToWorld(), PreviousLocalToWorld, GetBounds(), GetLocalBounds(), true, bHasPrecomputedVolumetricLightmap, bOutputVelocity);
+		BatchElement.PrimitiveUniformBufferResource = &DynamicPrimitiveUniformBuffer.UniformBuffer;
 
 		BatchElement.FirstIndex = 0;
 		BatchElement.NumPrimitives = Section.IndexBuffer.Indices.Num() / 3;
diff --git a/Engine/Plugins/Experimental/Water/Source/Runtime/Private/WaterInfoRendering.cpp b/Engine/Plugins/Experimental/Water/Source/Runtime/Private/WaterInfoRendering.cpp
--- a/Engine/Plugins/Experimental/Water/Source/Runtime/Private/WaterInfoRendering.cpp
+++ b/Engine/Plugins/Experimental/Water/Source/Runtime/Private/WaterInfoRendering.cpp
@@ -449,8 +449,8 @@ static void UpdateWaterInfoRendering_RenderThread(
 		}
 		else
 		{
-			FResolveParams ResolveParams;
-			AddCopyToResolveTargetPass(GraphBuilder, TargetTexture, DepthTexture, ResolveParams);
+			FRHICopyTextureInfo CopyInfo;
+			AddCopyTexturePass(GraphBuilder, TargetTexture, DepthTexture, CopyInfo);
 		}
 		
 		// We currently can't have multiple scene renderers run within the same RDGBuilder. Therefore, we must
@@ -566,8 +566,8 @@ static void UpdateWaterInfoRendering_RenderThread(
 		}
 		else
 		{
-			FResolveParams ResolveParams;
-			AddCopyToResolveTargetPass(GraphBuilder, TargetTexture, DilationTexture, ResolveParams);
+			FRHICopyTextureInfo CopyInfo;
+			AddCopyTexturePass(GraphBuilder, TargetTexture, DilationTexture, CopyInfo);
 		}
 
 		FRDGTextureRef MergeTargetTexture = GraphBuilder.CreateTexture(TextureDesc, TEXT("WaterInfoMerged"));
