Fixed a bad merge causing dilated water sections to not render in the water dilation pass

diff --git a/Engine/Plugins/Experimental/Water/Source/Runtime/Private/WaterBodySceneProxy.cpp b/Engine/Plugins/Experimental/Water/Source/Runtime/Private/WaterBodySceneProxy.cpp
--- a/Engine/Plugins/Experimental/Water/Source/Runtime/Private/WaterBodySceneProxy.cpp
+++ b/Engine/Plugins/Experimental/Water/Source/Runtime/Private/WaterBodySceneProxy.cpp
@@ -162,6 +162,14 @@ void FWaterBodySceneProxy::GetDynamicMeshElements(const TArray<const FSceneView*
 				AddWaterBodyMeshSection(Section, ViewIndex);
 			}
 		}
+
+		if ((CurrentWaterInfoPass == EWaterInfoPass::Dilation) || (bDebugShowWaterSceneProxiesEnabled))
+		{
+			for (const FWaterBodyMeshSection& Section : DilatedSections)
+			{
+				AddWaterBodyMeshSection(Section, ViewIndex);
+			}
+		}
 	}
 }
 
