[NO NON-PII COMMIT MESSAGE]

diff --git a/Engine/Plugins/Experimental/Water/Source/Runtime/Private/WaterBodyComponent.cpp b/Engine/Plugins/Experimental/Water/Source/Runtime/Private/WaterBodyComponent.cpp
--- a/Engine/Plugins/Experimental/Water/Source/Runtime/Private/WaterBodyComponent.cpp
+++ b/Engine/Plugins/Experimental/Water/Source/Runtime/Private/WaterBodyComponent.cpp
@@ -259,6 +259,19 @@ void UWaterBodyComponent::SetUnderwaterPostProcessMaterial(UMaterialInterface* I
 	UpdateMaterialInstances();
 }
 
+void UWaterBodyComponent::SetWaterAndUnderWaterPostProcessMaterial(UMaterialInterface* InWaterMaterial, UMaterialInterface* InUnderWaterPostProcessMaterial)
+{
+	bool bUpdateInstances = WaterMaterial != InWaterMaterial || UnderwaterPostProcessMaterial != InUnderWaterPostProcessMaterial;
+
+	WaterMaterial = InWaterMaterial;
+	UnderwaterPostProcessMaterial = InUnderWaterPostProcessMaterial;
+
+	if (bUpdateInstances)
+	{
+		UpdateMaterialInstances();
+	}
+}
+
 bool UWaterBodyComponent::ShouldGenerateWaterMeshTile() const
 {
 	return ((GetWaterBodyType() != EWaterBodyType::Transition)
diff --git a/Engine/Plugins/Experimental/Water/Source/Runtime/Private/WaterBodyCustomComponent.cpp b/Engine/Plugins/Experimental/Water/Source/Runtime/Private/WaterBodyCustomComponent.cpp
--- a/Engine/Plugins/Experimental/Water/Source/Runtime/Private/WaterBodyCustomComponent.cpp
+++ b/Engine/Plugins/Experimental/Water/Source/Runtime/Private/WaterBodyCustomComponent.cpp
@@ -86,13 +86,21 @@ void UWaterBodyCustomComponent::OnUpdateBody(bool bWithExclusionVolumes)
 
 	CreateOrUpdateWaterMID();
 	MeshComp->SetStaticMesh(GetWaterMeshOverride());
-	MeshComp->SetMaterial(0, WaterMID);
 	MeshComp->SetCastShadow(false);
 	CopySharedCollisionSettingsToComponent(MeshComp);
 	CopySharedNavigationSettingsToComponent(MeshComp);
 	MeshComp->MarkRenderStateDirty();
 }
 
+void UWaterBodyCustomComponent::CreateOrUpdateWaterMID()
+{
+	Super::CreateOrUpdateWaterMID();
+	if (MeshComp != nullptr)
+	{
+		MeshComp->SetMaterial(0, WaterMID);
+	}
+}
+
 FPrimitiveSceneProxy* UWaterBodyCustomComponent::CreateSceneProxy()
 {
 	// Don't create a scene proxy for custom water body components since they don't render into the water info texture (yet)
diff --git a/Engine/Plugins/Experimental/Water/Source/Runtime/Public/WaterBodyComponent.h b/Engine/Plugins/Experimental/Water/Source/Runtime/Public/WaterBodyComponent.h
--- a/Engine/Plugins/Experimental/Water/Source/Runtime/Public/WaterBodyComponent.h
+++ b/Engine/Plugins/Experimental/Water/Source/Runtime/Public/WaterBodyComponent.h
@@ -187,6 +187,9 @@ public:
 	/** Sets under water post process material */
 	void SetUnderwaterPostProcessMaterial(UMaterialInterface* InMaterial);
 
+	UFUNCTION(BlueprintCallable, Category = Rendering)
+	void SetWaterAndUnderWaterPostProcessMaterial(UMaterialInterface* InWaterMaterial, UMaterialInterface* InUnderWaterPostProcessMaterial);
+
 	/** Returns water spline metadata */
 	UWaterSplineMetadata* GetWaterSplineMetadata() { return WaterSplineMetadata; }
 
@@ -391,7 +394,7 @@ protected:
 	void UpdateSplineComponent();
 	void UpdateExclusionVolumes();
 	bool UpdateWaterHeight();
-	void CreateOrUpdateWaterMID();
+	virtual void CreateOrUpdateWaterMID();
 	void CreateOrUpdateUnderwaterPostProcessMID();
 	void CreateOrUpdateWaterInfoMID();
 	void PrepareCurrentPostProcessSettings();
diff --git a/Engine/Plugins/Experimental/Water/Source/Runtime/Public/WaterBodyCustomComponent.h b/Engine/Plugins/Experimental/Water/Source/Runtime/Public/WaterBodyCustomComponent.h
--- a/Engine/Plugins/Experimental/Water/Source/Runtime/Public/WaterBodyCustomComponent.h
+++ b/Engine/Plugins/Experimental/Water/Source/Runtime/Public/WaterBodyCustomComponent.h
@@ -26,7 +26,8 @@ protected:
 	virtual void Reset() override;
 	virtual void BeginUpdateWaterBody() override;
 	virtual void OnUpdateBody(bool bWithExclusionVolumes) override;
-
+	virtual void CreateOrUpdateWaterMID() override;
+	
 	virtual FPrimitiveSceneProxy* CreateSceneProxy() override;
 #if WITH_EDITOR
 	virtual TArray<TSharedRef<FTokenizedMessage>> CheckWaterBodyStatus() const override;
