Water: Avoid check() in FScopedLandscapeLODOverride when processing a landscape made only of mesh proxies
* Can be the case in World Partition if all that is loaded are HLODs

diff --git a/Engine/Plugins/Experimental/Water/Source/Runtime/Private/WaterInfoRendering.cpp b/Engine/Plugins/Experimental/Water/Source/Runtime/Private/WaterInfoRendering.cpp
--- a/Engine/Plugins/Experimental/Water/Source/Runtime/Private/WaterInfoRendering.cpp
+++ b/Engine/Plugins/Experimental/Water/Source/Runtime/Private/WaterInfoRendering.cpp
@@ -327,7 +327,10 @@ private:
 
 	void SetOptimalLandscapeLODOverrides()
 	{
-		if (!LandscapeRenderSystems) return;
+		if (!LandscapeRenderSystems)
+		{
+			return;
+		}
 		
 		// In order to prevent overdrawing the landscape components, we compute the lowest-detailed LOD level which satisfies the pixel coverage of the Water Info texture
 		// and force it on all landscape components. This override is set different per Landscape actor in case there are multiple under the same water zone.
@@ -342,7 +345,7 @@ private:
 		for (const TPair<uint32, FLandscapeRenderSystem*>& Pair : *LandscapeRenderSystems)
 		{
 			FLandscapeRenderSystem* LandscapeRenderSystem = Pair.Value;
-			int32 OptimalLODLevel = -1;
+			int32 OptimalLODLevel = INDEX_NONE;
 
 			// All components within the same landscape (and thus its render system) should have the same number of quads and the same extent.
 			// therefore we can simply find the first component and compute its optimal LOD level.
@@ -353,6 +356,8 @@ private:
 					const double LandscapeComponentUnitsPerVertex = LandscapeSectionInfo->ComputeSectionResolution();
 					if (LandscapeComponentUnitsPerVertex <= 0.f)
 					{
+						// No section resolution probably means the section is a mesh proxy, which might not have regular units per vertex.
+						// Avoid computing optimal LOD in this case.
 						continue;
 					}
 
@@ -364,9 +369,6 @@ private:
 				}
 			}
 
-			// There should always be at least one valid component proxy and the optimal LOD level should never be negative.
-			check(OptimalLODLevel >= 0);
-
 			LandscapeLODOverridesToRestore.Add(LandscapeRenderSystem, LandscapeRenderSystem->ForcedLODOverride);
 			LandscapeRenderSystem->ForcedLODOverride = OptimalLODLevel;
 		}
