Implemented a deferred register for buoyancy components. Async callback is not always created before components attempt to register themselves.

diff --git a/Engine/Plugins/Experimental/Water/Source/Runtime/Private/BuoyancyManager.cpp b/Engine/Plugins/Experimental/Water/Source/Runtime/Private/BuoyancyManager.cpp
--- a/Engine/Plugins/Experimental/Water/Source/Runtime/Private/BuoyancyManager.cpp
+++ b/Engine/Plugins/Experimental/Water/Source/Runtime/Private/BuoyancyManager.cpp
@@ -100,7 +100,14 @@ void ABuoyancyManager::Register(UBuoyancyComponent* BuoyancyComponent)
 {
 	check(BuoyancyComponent);
 	BuoyancyComponents.AddUnique(BuoyancyComponent);
-	InitializeAsyncAux(BuoyancyComponent);
+	if(AsyncCallback)
+	{
+		InitializeAsyncAux(BuoyancyComponent);
+	}
+	else // AsyncCallback is not setup yet. We need to register this component at a later time.
+	{
+		BuoyancyComponentsToRegister.Emplace(BuoyancyComponent);
+	}
 }
 
 void ABuoyancyManager::Unregister(UBuoyancyComponent* BuoyancyComponent)
@@ -252,6 +259,15 @@ void ABuoyancyManager::BeginPlay()
 			OnPhysScenePreTickHandle = PhysScene->OnPhysScenePreTick.AddUObject(this, &ABuoyancyManager::Update);
 			AsyncCallback = PhysScene->GetSolver()->CreateAndRegisterSimCallbackObject_External<FBuoyancyManagerAsyncCallback>();
 		}
+
+		for (const TWeakObjectPtr<UBuoyancyComponent>& BuoyancyComponentPtr : BuoyancyComponentsToRegister)
+		{
+			if (UBuoyancyComponent* BuoyancyComponent = BuoyancyComponentPtr.Get())
+			{
+				InitializeAsyncAux(BuoyancyComponent);
+			}
+		}
+		BuoyancyComponentsToRegister.Empty();
 	}
 
 	Super::BeginPlay();
diff --git a/Engine/Plugins/Experimental/Water/Source/Runtime/Public/BuoyancyManager.h b/Engine/Plugins/Experimental/Water/Source/Runtime/Public/BuoyancyManager.h
--- a/Engine/Plugins/Experimental/Water/Source/Runtime/Public/BuoyancyManager.h
+++ b/Engine/Plugins/Experimental/Water/Source/Runtime/Public/BuoyancyManager.h
@@ -58,6 +58,9 @@ private:
 	UPROPERTY()
 	TArray<UBuoyancyComponent*> BuoyancyComponents;
 	TArray<UBuoyancyComponent*> BuoyancyComponentsActive;
+	
+	// List of buoyancy components to defer registration until this buoyancy manager is fully initialized.
+	TArray<TWeakObjectPtr<UBuoyancyComponent>> BuoyancyComponentsToRegister;
 
 	FBuoyancyManagerAsyncCallback* AsyncCallback;
 	int32 Timestamp;
