Fixed missing water in water waves asset editor due to missing material and made the water body larger to fill viewport completely.

diff --git a/Engine/Plugins/Experimental/Water/Source/Editor/Private/WaterWavesEditorViewport.cpp b/Engine/Plugins/Experimental/Water/Source/Editor/Private/WaterWavesEditorViewport.cpp
--- a/Engine/Plugins/Experimental/Water/Source/Editor/Private/WaterWavesEditorViewport.cpp
+++ b/Engine/Plugins/Experimental/Water/Source/Editor/Private/WaterWavesEditorViewport.cpp
@@ -37,6 +37,8 @@ void SWaterWavesEditorViewport::Construct(const FArguments& InArgs)
 	UWaterBodyComponent* WaterBodyComponent = CustomWaterBody->GetWaterBodyComponent();
 	check(WaterBodyComponent);
 	WaterBodyComponent->SetWaterMeshOverride(GetDefault<UWaterEditorSettings>()->WaterBodyCustomDefaults.GetWaterMesh());
+	WaterBodyComponent->SetWaterMaterial(GetDefault<UWaterEditorSettings>()->WaterBodyCustomDefaults.GetWaterMaterial());
+
 
 	UWaterSplineComponent* WaterSpline = CustomWaterBody->GetWaterSpline();
 	check(WaterSpline);
