Fix for water waves not working properly at large distances to world origin (LWC).

#rnx

diff --git a/Engine/Plugins/Experimental/Water/Shaders/Private/GerstnerWaveFunctions.ush b/Engine/Plugins/Experimental/Water/Shaders/Private/GerstnerWaveFunctions.ush
--- a/Engine/Plugins/Experimental/Water/Shaders/Private/GerstnerWaveFunctions.ush
+++ b/Engine/Plugins/Experimental/Water/Shaders/Private/GerstnerWaveFunctions.ush
@@ -345,43 +345,62 @@ float3 PackNormalAndWPONew(WaveOutput inWave)
 	return packedoutput;
 }
 
-WaveOutput GetSingleGerstnerWaveNew(int InWaveIndex, FWaterBodyHeader InWaterBodyHeader, float InTime, float2 InWorldPos)
+WaveOutput GetSingleGerstnerWaveNew(int InWaveIndex, FWaterBodyHeader InWaterBodyHeader, float InTime, FLWCVector2 InWorldPos)
 {
 	WaveOutput OutWave = (WaveOutput)0;
 		
 	FWaveParams CurrentWave = GetWaveDataNew(InWaveIndex, InWaterBodyHeader);
 		
-	float dispersion = 2 * PI / CurrentWave.Wavelength;
-	float2 wavevector = CurrentWave.Direction * dispersion;
-	float wavespeed = sqrt(dispersion * Gravity);
-	float wavetime = wavespeed * InTime;
+	float Dispersion = 2 * PI / CurrentWave.Wavelength;
+	float2 WaveVector = CurrentWave.Direction * Dispersion;
+	float WaveSpeed = sqrt(Dispersion * Gravity);
+	float WaveTime = WaveSpeed * InTime;
 
-	float wavepos = dot(InWorldPos, wavevector) - wavetime;
-		
-	float wavesin = sin(wavepos);
-	float wavecos = cos(wavepos);
+	// Use the offset of the normalized tile as world position.
+	// Doing proper LWC math is both more expensive and leads to less smooth/precise results.
+	InWorldPos = LWCNormalizeTile(InWorldPos);
+	float2 WorldPos = InWorldPos.Offset;
+
+	float WavePos = dot(WorldPos, WaveVector) - WaveTime;
+	float WaveSin = sin(WavePos);
+	float WaveCos = cos(WavePos);
+
+	// Fix up tile borders by blending between them
+	const float2 TileBorderDist = (UE_LWC_RENDER_TILE_SIZE * 0.5f) - abs(WorldPos);
+	const float BlendZoneWidth = 400.0f; // Blend over a range of 4 meters on each side of the tile
+	if (any(TileBorderDist < BlendZoneWidth))
+	{
+		const float2 BlendWorldPos = TileBorderDist;
+
+		const float2 BlendAlpha2 = saturate(TileBorderDist / BlendZoneWidth);
+		const float BlendAlpha = BlendAlpha2.x * BlendAlpha2.y;
+
+		const float BlendWavePos = dot(BlendWorldPos, WaveVector) - WaveTime;
+		WaveSin = lerp(sin(BlendWavePos), WaveSin, BlendAlpha);
+		WaveCos = lerp(cos(BlendWavePos), WaveCos, BlendAlpha);
+	}
 		
-	float wKA = CurrentWave.Amplitude * dispersion;
+	float wKA = CurrentWave.Amplitude * Dispersion;
 		
 	float q = CurrentWave.Steepness / wKA;
 		
-	OutWave.Normal.xy = wavesin * wKA * CurrentWave.Direction;
+	OutWave.Normal.xy = WaveSin * wKA * CurrentWave.Direction;
 		
 #if SOLVE_NORMAL_Z
-		OutWave.Normal.z = wavecos *  CurrentWave.Steepness * saturate( (CurrentWave.Amplitude * SteepnessThreshold) / CurrentWave.Wavelength );
-		//OutWave.Normal.z = wavecos *  wKA * (q/MaxWaves);
+		OutWave.Normal.z = WaveCos *  CurrentWave.Steepness * saturate( (CurrentWave.Amplitude * SteepnessThreshold) / CurrentWave.Wavelength );
+		//OutWave.Normal.z = WaveCos *  wKA * (q/MaxWaves);
 #else
 		OutWave.Normal.z = 0;
 #endif
 
-	OutWave.WPO.xy = -q * wavesin * CurrentWave.Direction *  CurrentWave.Amplitude;
-	OutWave.WPO.z = wavecos * CurrentWave.Amplitude;
+	OutWave.WPO.xy = -q * WaveSin * CurrentWave.Direction *  CurrentWave.Amplitude;
+	OutWave.WPO.z = WaveCos * CurrentWave.Amplitude;
 
 	return OutWave;
 		
 }
 
-WaveOutput GetAllGerstnerWavesNew(int InWaterBodyIndex, float InTime, float2 InWorldPos)
+WaveOutput GetAllGerstnerWavesNew(int InWaterBodyIndex, float InTime, FLWCVector2 InWorldPos)
 {		
 	WaveOutput OutWaves = (WaveOutput)0;
 	WaveOutput CurrentWave = (WaveOutput)0;
@@ -399,7 +418,7 @@ WaveOutput GetAllGerstnerWavesNew(int InWaterBodyIndex, float InTime, float2 InW
 	return OutWaves;
 }
 	
-WaveOutput GetRangeGerstnerWavesNew(int InMinWaveIndex, int InMaxWaveIndex, int InWaterBodyIndex, float InTime, float2 InWorldPos)
+WaveOutput GetRangeGerstnerWavesNew(int InMinWaveIndex, int InMaxWaveIndex, int InWaterBodyIndex, float InTime, FLWCVector2 InWorldPos)
 {
 	WaveOutput OutWaves = (WaveOutput)0;
 	WaveOutput CurrentWave = (WaveOutput)0;
diff --git a/Engine/Plugins/Experimental/Water/Source/Runtime/Private/GerstnerWaterWaves.cpp b/Engine/Plugins/Experimental/Water/Source/Runtime/Private/GerstnerWaterWaves.cpp
--- a/Engine/Plugins/Experimental/Water/Source/Runtime/Private/GerstnerWaterWaves.cpp
+++ b/Engine/Plugins/Experimental/Water/Source/Runtime/Private/GerstnerWaterWaves.cpp
@@ -4,6 +4,7 @@
 #include "GerstnerWaterWaveSubsystem.h"
 #include "WaterModule.h"
 #include "Engine/Engine.h"
+#include "Core/Public/Misc/LargeWorldRenderPosition.h"
 
 // ----------------------------------------------------------------------------------
 
@@ -59,11 +60,14 @@ float UGerstnerWaterWaves::GetWaveHeightAtPosition(const FVector& InPosition, fl
 
 	FVector SummedNormal(ForceInitToZero);
 
+	// Use the offset of the normalized tile as world position to match the shader behavior (see GerstnerWaveFunctions.ush).
+	FVector WorldPosition(FLargeWorldRenderPosition(InPosition).GetOffset());
+
 	for (const FGerstnerWave& Params : GetGerstnerWaves())
 	{
 		float FirstOffset1D;
 		FVector FirstNormal;
-		FVector FirstOffset = GetWaveOffsetAtPosition(Params, InPosition, InTime, FirstNormal, FirstOffset1D);
+		FVector FirstOffset = GetWaveOffsetAtPosition(Params, WorldPosition, InTime, FirstNormal, FirstOffset1D);
 
 		// Only non-zero steepness requires a second sample.
 		if (Params.Q != 0)
@@ -74,7 +78,7 @@ float UGerstnerWaterWaves::GetWaveHeightAtPosition(const FVector& InPosition, fl
 			//Lerp between the two sampled heights and normals.
 			const float TwoPi = 2 * PI;
 			const float WaveTime = Params.WaveSpeed * InTime;
-			float Position1D = FVector2D::DotProduct(FVector2D(InPosition.X, InPosition.Y), Params.WaveVector) - WaveTime;
+			float Position1D = FVector2D::DotProduct(FVector2D(WorldPosition.X, WorldPosition.Y), Params.WaveVector) - WaveTime;
 			float MappedPosition1D = Position1D >= 0.f ? FMath::Fmod(Position1D, TwoPi) : TwoPi - FMath::Abs(FMath::Fmod(Position1D, TwoPi)); //get positive modulos from negative numbers too
 
 			FVector SecondNormal;
@@ -88,7 +92,7 @@ float UGerstnerWaterWaves::GetWaveHeightAtPosition(const FVector& InPosition, fl
 			{
 				GuessOffset = -Params.Direction * Params.Q;
 			}
-			const FVector GuessPosition = InPosition + GuessOffset;
+			const FVector GuessPosition = WorldPosition + GuessOffset;
 			FVector SecondOffset = GetWaveOffsetAtPosition(Params, GuessPosition, InTime, SecondNormal, SecondOffset1D);
 			SecondOffset1D += (MappedPosition1D < PI) ? Params.Q : -Params.Q;
 			if (!(MappedPosition1D < PI))
@@ -120,10 +124,12 @@ float UGerstnerWaterWaves::GetWaveHeightAtPosition(const FVector& InPosition, fl
 float UGerstnerWaterWaves::GetSimpleWaveHeightAtPosition(const FVector& InPosition, float InWaterDepth, float InTime) const
 {
 	float WaveHeight = 0.f;
+	// Use the offset of the normalized tile as world position to match the shader behavior (see GerstnerWaveFunctions.ush).
+	FVector WorldPosition(FLargeWorldRenderPosition(InPosition).GetOffset());
 
 	for (const FGerstnerWave& Wave : GetGerstnerWaves())
 	{
-		WaveHeight += GetSimpleWaveOffsetAtPosition(Wave, InPosition, InTime);
+		WaveHeight += GetSimpleWaveOffsetAtPosition(Wave, WorldPosition, InTime);
 	}
 
 	return WaveHeight;
@@ -142,6 +148,7 @@ FVector UGerstnerWaterWaves::GetWaveOffsetAtPosition(const FGerstnerWave& InWave
 
 	float WaveSin = 0, WaveCos = 0;
 	FMath::SinCos(&WaveSin, &WaveCos, WavePosition);
+	BlendWaveBetweenLWCTiles(InWaveParams, InPosition, InTime, WaveSin, WaveCos);
 
 	FVector Offset;
 	OutOffset1D = -InWaveParams.Q * WaveSin;
@@ -158,11 +165,32 @@ float UGerstnerWaterWaves::GetSimpleWaveOffsetAtPosition(const FGerstnerWave& In
 {
 	const float WaveTime = InWaveParams.WaveSpeed * InTime;
 	const float WavePosition = FVector2D::DotProduct(FVector2D(InPosition.X, InPosition.Y), InWaveParams.WaveVector) - WaveTime;
-	const float WaveCos = FMath::Cos(WavePosition);
+	float WaveCos = FMath::Cos(WavePosition);
+	float WaveSinDummy = 0.0f;
+	BlendWaveBetweenLWCTiles(InWaveParams, InPosition, InTime, WaveSinDummy, WaveCos);
 	const float HeightOffset = WaveCos * InWaveParams.Amplitude;
 	return HeightOffset;
 }
 
+void UGerstnerWaterWaves::BlendWaveBetweenLWCTiles(const FGerstnerWave& InWaveParams, const FVector& InPosition, float InTime, float& WaveSin, float& WaveCos) const
+{
+	const FVector TileBorderDist = FVector(FLargeWorldRenderScalar::GetTileSize() * 0.5) - InPosition.GetAbs();
+	const double BlendZoneWidth = 400.0; // Blend over a range of 4 meters on each side of the tile
+	if (TileBorderDist.X < BlendZoneWidth || TileBorderDist.Y < BlendZoneWidth)
+	{
+		const FVector2D BlendWorldPos = FVector2D(TileBorderDist.X, TileBorderDist.Y);
+		const double BlendAlpha = FMath::Clamp(BlendWorldPos.X / BlendZoneWidth, 0.0, 1.0) * FMath::Clamp(BlendWorldPos.Y / BlendZoneWidth, 0.0, 1.0);
+
+		const float WaveTime = InWaveParams.WaveSpeed * InTime;
+		const float BlendWavePos = FVector2D::DotProduct(BlendWorldPos, InWaveParams.WaveVector) - WaveTime;
+		float BlendWaveSin = 0.0f;
+		float BlendWaveCos = 0.0f;
+		FMath::SinCos(&BlendWaveSin, &BlendWaveCos, BlendWavePos);
+		WaveSin = FMath::Lerp(BlendWaveSin, WaveSin, BlendAlpha);
+		WaveCos = FMath::Lerp(BlendWaveCos, WaveCos, BlendAlpha);
+	}
+}
+
 void UGerstnerWaterWaves::RecomputeWaves(bool bAllowBPScript)
 {
 	GerstnerWaves.Empty();
diff --git a/Engine/Plugins/Experimental/Water/Source/Runtime/Public/GerstnerWaterWaves.h b/Engine/Plugins/Experimental/Water/Source/Runtime/Public/GerstnerWaterWaves.h
--- a/Engine/Plugins/Experimental/Water/Source/Runtime/Public/GerstnerWaterWaves.h
+++ b/Engine/Plugins/Experimental/Water/Source/Runtime/Public/GerstnerWaterWaves.h
@@ -211,4 +211,5 @@ public:
 private:
 	FVector GetWaveOffsetAtPosition(const FGerstnerWave& InWaveParams, const FVector& InPosition, float InTime, FVector& OutNormal, float& OutOffset1D) const;
 	float GetSimpleWaveOffsetAtPosition(const FGerstnerWave& InWaveParams, const FVector& InPosition, float InTime) const;
+	void BlendWaveBetweenLWCTiles(const FGerstnerWave& InWaveParams, const FVector& InPosition, float InTime, float& WaveSin, float& WaveCos) const;
 };
\ No newline at end of file
