Ensure shader maps for ground actors are compiled before rendering the water info texture.

diff --git a/Engine/Plugins/Experimental/Water/Source/Runtime/Private/WaterInfoRendering.h b/Engine/Plugins/Experimental/Water/Source/Runtime/Private/WaterInfoRendering.h
--- a/Engine/Plugins/Experimental/Water/Source/Runtime/Private/WaterInfoRendering.h
+++ b/Engine/Plugins/Experimental/Water/Source/Runtime/Private/WaterInfoRendering.h
@@ -19,7 +19,6 @@ namespace UE::WaterInfo
 		TArray<UWaterBodyComponent*> WaterBodies;
 		TArray<TWeakObjectPtr<AActor>> GroundActors;
 		float CaptureZ;
-		int32 LandscapeLODOverride;
 	};
 	
 void UpdateWaterInfoRendering(
diff --git a/Engine/Plugins/Experimental/Water/Source/Runtime/Private/WaterZoneActor.cpp b/Engine/Plugins/Experimental/Water/Source/Runtime/Private/WaterZoneActor.cpp
--- a/Engine/Plugins/Experimental/Water/Source/Runtime/Private/WaterZoneActor.cpp
+++ b/Engine/Plugins/Experimental/Water/Source/Runtime/Private/WaterZoneActor.cpp
@@ -241,11 +241,15 @@ bool AWaterZone::UpdateWaterInfoTexture()
 		float WaterZMin(TNumericLimits<float>::Max());
 		float WaterZMax(TNumericLimits<float>::Lowest());
 	
-		bool bHasIncompleteShaderMaps = false;
+		
+		// Collect a list of all materials used in the water info render to ensure they have complete shaders maps.
+		// If they do not, we must submit compile jobs for them and wait until they are finished before re-rendering.
+		TArray<UMaterialInterface*> UsedMaterials;
+
 		// #todo_water [roey]: we should try caching this list to avoid potentially iterating over a lot of water bodies which may not belong to this zone specifically.
 		// For now whenever we update the water info texture we will collect all water bodies within the zone and pass those to the renderer each time this function is called.
 		TArray<UWaterBodyComponent*> WaterBodies;
-		ForEachWaterBodyComponent([World, &WaterBodies, &WaterZMax, &WaterZMin, &bHasIncompleteShaderMaps](UWaterBodyComponent* Component)
+		ForEachWaterBodyComponent([World, &WaterBodies, &WaterZMax, &WaterZMin, &UsedMaterials](UWaterBodyComponent* Component)
 		{
 			// skip components which don't affect the water info texture
 			if (!Component->AffectsWaterInfo())
@@ -255,15 +259,7 @@ bool AWaterZone::UpdateWaterInfoTexture()
 
 			if (UMaterialInterface* WaterInfoMaterial = Component->GetWaterInfoMaterialInstance())
 			{
-				if (FMaterialResource* MaterialResource = WaterInfoMaterial->GetMaterialResource(World->Scene->GetFeatureLevel()))
-				{
-					if (!MaterialResource->IsGameThreadShaderMapComplete())
-					{
-						MaterialResource->SubmitCompileJobs_GameThread(EShaderCompileJobPriority::ForceLocal);
-						bHasIncompleteShaderMaps = true;
-						return true;
-					}
-				}
+				UsedMaterials.Add(WaterInfoMaterial);
 			}
 
 			WaterBodies.Add(Component);
@@ -273,11 +269,6 @@ bool AWaterZone::UpdateWaterInfoTexture()
 			return true;
 		});
 
-		if (bHasIncompleteShaderMaps)
-		{
-			return false;
-		}
-
 		// If we don't have any water bodies we don't need to do anything.
 		if (WaterBodies.Num() == 0)
 		{
@@ -290,7 +281,6 @@ bool AWaterZone::UpdateWaterInfoTexture()
 		GroundZMin = TNumericLimits<float>::Max();
 		float GroundZMax = TNumericLimits<float>::Lowest();
 
-		int32 LandscapeLODOverride = 0;
 		TArray<TWeakObjectPtr<AActor>> GroundActors;
 		for (ALandscapeProxy* LandscapeProxy : TActorRange<ALandscapeProxy>(World))
 		{
@@ -298,12 +288,48 @@ bool AWaterZone::UpdateWaterInfoTexture()
 			GroundZMin = FMath::Min(GroundZMin, LandscapeBox.Min.Z);
 			GroundZMax = FMath::Max(GroundZMax, LandscapeBox.Max.Z);
 			GroundActors.Add(LandscapeProxy);
+		}
+
+		// Check all the ground actors have complete shader maps before we try to render them into the water info texture
+		for (TWeakObjectPtr<AActor> GroundActorPtr : GroundActors)
+		{
+			if (AActor* GroundActor = GroundActorPtr.Get())
+			{
+				TInlineComponentArray<UPrimitiveComponent*> PrimitiveComponents(GroundActor);
+
+				for (UPrimitiveComponent* PrimitiveComponent : PrimitiveComponents)
+				{
+					TArray<UMaterialInterface*> TmpUsedMaterials;
+					PrimitiveComponent->GetUsedMaterials(TmpUsedMaterials, false);
+					UsedMaterials.Append(TmpUsedMaterials);
+				}
+			}
+		}
+
+		// Loop through all used materials and ensure that compile jobs are submitted for all which do not have complete shader maps before early-ing out of the info update.
+		bool bHasIncompleteShaderMaps = false;
+		const ERHIFeatureLevel::Type FeatureLevel = World->Scene->GetFeatureLevel();
+		for (UMaterialInterface* Material : UsedMaterials)
+		{
+			if (Material)
+			{
+				if (FMaterialResource* MaterialResource = Material->GetMaterialResource(FeatureLevel))
+				{
+					if (!MaterialResource->IsGameThreadShaderMapComplete())
+					{
+						MaterialResource->SubmitCompileJobs_GameThread(EShaderCompileJobPriority::ForceLocal);
+						bHasIncompleteShaderMaps = true;
+					}
+				}
+			}
+		}
 
-			// Target mip 64x64 for a capture
-			int32 LOD64x64 = FMath::CeilLogTwo(LandscapeProxy->SubsectionSizeQuads + 1) - 6;
-			LandscapeLODOverride = FMath::Max(LandscapeLODOverride, LOD64x64);
+		if (bHasIncompleteShaderMaps)
+		{
+			return false;
 		}
 
+
 		const ETextureRenderTargetFormat Format = bHalfPrecisionTexture ? ETextureRenderTargetFormat::RTF_RGBA16f : RTF_RGBA32f;
 		WaterInfoTexture = FWaterUtils::GetOrCreateTransientRenderTarget2D(WaterInfoTexture, TEXT("WaterInfoTexture"), RenderTargetResolution, Format);
 
@@ -313,7 +339,6 @@ bool AWaterZone::UpdateWaterInfoTexture()
 		Context.GroundActors = MoveTemp(GroundActors);
 		Context.CaptureZ = FMath::Max(WaterZMax, GroundZMax) + CaptureZOffset;
 		Context.TextureRenderTarget = WaterInfoTexture;
-		Context.LandscapeLODOverride = LandscapeLODOverride;
 
 		if (UWaterSubsystem* WaterSubsystem = UWaterSubsystem::GetWaterSubsystem(World))
 		{
