Implemented missing memory tracking for water classes

diff --git a/Engine/Plugins/Experimental/Water/Source/Runtime/Private/BuoyancyComponent.cpp b/Engine/Plugins/Experimental/Water/Source/Runtime/Private/BuoyancyComponent.cpp
--- a/Engine/Plugins/Experimental/Water/Source/Runtime/Private/BuoyancyComponent.cpp
+++ b/Engine/Plugins/Experimental/Water/Source/Runtime/Private/BuoyancyComponent.cpp
@@ -247,6 +247,15 @@ void UBuoyancyComponent::TickComponent(float DeltaTime, enum ELevelTick TickType
 	Update(DeltaTime);
 }
 
+void UBuoyancyComponent::GetResourceSizeEx(FResourceSizeEx& CumulativeResourceSize)
+{
+	Super::GetResourceSizeEx(CumulativeResourceSize);
+
+	// Account for all non-editor data properties :
+	CumulativeResourceSize.AddDedicatedSystemMemoryBytes(sizeof(CurAsyncInput) + sizeof(CurAsyncOutput) + sizeof(NextAsyncOutput) + sizeof(CurAsyncType) + sizeof(OutputInterpAlpha) + OutputsWaitingOn.GetAllocatedSize()
+		+ sizeof(PontoonConfiguration) + ConfiguredPontoonCoefficients.GetAllocatedSize() + sizeof(VelocityPontoonIndex));
+}
+
 void UBuoyancyComponent::SetupWaterBodyOverlaps()
 {
 	//SimulatingComponent->SetCollisionObjectType(ECollisionChannel::ECC_PhysicsBody);
diff --git a/Engine/Plugins/Experimental/Water/Source/Runtime/Private/GerstnerWaterWaveSubsystem.cpp b/Engine/Plugins/Experimental/Water/Source/Runtime/Private/GerstnerWaterWaveSubsystem.cpp
--- a/Engine/Plugins/Experimental/Water/Source/Runtime/Private/GerstnerWaterWaveSubsystem.cpp
+++ b/Engine/Plugins/Experimental/Water/Source/Runtime/Private/GerstnerWaterWaveSubsystem.cpp
@@ -38,6 +38,13 @@ void UGerstnerWaterWaveSubsystem::Unregister(FGerstnerWaterWaveViewExtension* In
 	GerstnerWaterWaveViewExtensions.Remove(InViewExtension);
 }
 
+void UGerstnerWaterWaveSubsystem::GetResourceSizeEx(FResourceSizeEx& CumulativeResourceSize)
+{
+	Super::GetResourceSizeEx(CumulativeResourceSize);
+
+	CumulativeResourceSize.AddDedicatedSystemMemoryBytes(GerstnerWaterWaveViewExtensions.GetAllocatedSize());
+}
+
 void UGerstnerWaterWaveSubsystem::BeginFrameCallback()
 {
 	// In case there was a change, all registered view extensions need to update their GPU data : 
diff --git a/Engine/Plugins/Experimental/Water/Source/Runtime/Private/LakeCollisionComponent.cpp b/Engine/Plugins/Experimental/Water/Source/Runtime/Private/LakeCollisionComponent.cpp
--- a/Engine/Plugins/Experimental/Water/Source/Runtime/Private/LakeCollisionComponent.cpp
+++ b/Engine/Plugins/Experimental/Water/Source/Runtime/Private/LakeCollisionComponent.cpp
@@ -118,7 +118,7 @@ FPrimitiveSceneProxy* ULakeCollisionComponent::CreateSceneProxy()
 			return Result;
 		}
 		virtual uint32 GetMemoryFootprint(void) const override { return(sizeof(*this) + GetAllocatedSize()); }
-		uint32 GetAllocatedSize(void) const { return(FPrimitiveSceneProxy::GetAllocatedSize()); }
+		uint32 GetAllocatedSize(void) const { return FPrimitiveSceneProxy::GetAllocatedSize() + AggregateGeom.GetAllocatedSize(); }
 
 	private:
 		FKAggregateGeom AggregateGeom;
diff --git a/Engine/Plugins/Experimental/Water/Source/Runtime/Private/OceanCollisionComponent.cpp b/Engine/Plugins/Experimental/Water/Source/Runtime/Private/OceanCollisionComponent.cpp
--- a/Engine/Plugins/Experimental/Water/Source/Runtime/Private/OceanCollisionComponent.cpp
+++ b/Engine/Plugins/Experimental/Water/Source/Runtime/Private/OceanCollisionComponent.cpp
@@ -183,7 +183,7 @@ FPrimitiveSceneProxy* UOceanCollisionComponent::CreateSceneProxy()
 			return Result;
 		}
 		virtual uint32 GetMemoryFootprint(void) const override { return(sizeof(*this) + GetAllocatedSize()); }
-		uint32 GetAllocatedSize(void) const { return(FPrimitiveSceneProxy::GetAllocatedSize()); }
+		uint32 GetAllocatedSize(void) const { return FPrimitiveSceneProxy::GetAllocatedSize() + AggregateGeom.GetAllocatedSize(); }
 
 	private:
 		FKAggregateGeom AggregateGeom;
diff --git a/Engine/Plugins/Experimental/Water/Source/Runtime/Private/WaterBodyComponent.cpp b/Engine/Plugins/Experimental/Water/Source/Runtime/Private/WaterBodyComponent.cpp
--- a/Engine/Plugins/Experimental/Water/Source/Runtime/Private/WaterBodyComponent.cpp
+++ b/Engine/Plugins/Experimental/Water/Source/Runtime/Private/WaterBodyComponent.cpp
@@ -1312,6 +1312,18 @@ void UWaterBodyComponent::OnGenerateOverlapEventsChanged()
 	}
 }
 
+void UWaterBodyComponent::GetResourceSizeEx(FResourceSizeEx& CumulativeResourceSize)
+{
+	Super::GetResourceSizeEx(CumulativeResourceSize);
+
+	// Account for all non-editor data properties :
+	CumulativeResourceSize.AddDedicatedSystemMemoryBytes(
+		WaterBodyMeshIndices.GetAllocatedSize()
+		+ WaterBodyMeshVertices.GetAllocatedSize()
+		+ DilatedWaterBodyMeshIndices.GetAllocatedSize()
+		+ DilatedWaterBodyMeshVertices.GetAllocatedSize());
+}
+
 bool UWaterBodyComponent::SetDynamicParametersOnMID(UMaterialInstanceDynamic* InMID)
 {
 	UWaterSubsystem* WaterSubsystem = UWaterSubsystem::GetWaterSubsystem(GetWorld());
@@ -1535,6 +1547,12 @@ void UWaterBodyComponent::UpdateWaterBodyRenderData()
 		GenerateWaterBodyMesh();
 		MarkRenderStateDirty();
 
+		// Compact the arrays since they are not going to change unless everything is rebuilt
+		WaterBodyMeshVertices.Shrink();
+		WaterBodyMeshIndices.Shrink();
+		DilatedWaterBodyMeshVertices.Shrink();
+		DilatedWaterBodyMeshIndices.Shrink();
+
 		if (AWaterZone* WaterZone = GetWaterZone())
 		{
 			WaterZone->MarkForRebuild(EWaterZoneRebuildFlags::UpdateWaterInfoTexture);
diff --git a/Engine/Plugins/Experimental/Water/Source/Runtime/Private/WaterBodySceneProxy.cpp b/Engine/Plugins/Experimental/Water/Source/Runtime/Private/WaterBodySceneProxy.cpp
--- a/Engine/Plugins/Experimental/Water/Source/Runtime/Private/WaterBodySceneProxy.cpp
+++ b/Engine/Plugins/Experimental/Water/Source/Runtime/Private/WaterBodySceneProxy.cpp
@@ -200,7 +200,8 @@ uint32 FWaterBodySceneProxy::GetMemoryFootprint() const
 
 uint32 FWaterBodySceneProxy::GetAllocatedSize() const
 {
-	return FPrimitiveSceneProxy::GetAllocatedSize();
+
+	return FPrimitiveSceneProxy::GetAllocatedSize() + Sections.GetAllocatedSize() + DilatedSections.GetAllocatedSize();
 }
 
 bool FWaterBodySceneProxy::IsShown(const FSceneView* View) const
diff --git a/Engine/Plugins/Experimental/Water/Source/Runtime/Public/BuoyancyComponent.h b/Engine/Plugins/Experimental/Water/Source/Runtime/Public/BuoyancyComponent.h
--- a/Engine/Plugins/Experimental/Water/Source/Runtime/Public/BuoyancyComponent.h
+++ b/Engine/Plugins/Experimental/Water/Source/Runtime/Public/BuoyancyComponent.h
@@ -20,14 +20,17 @@ class WATER_API UBuoyancyComponent : public UActorComponent
 public:
 	UBuoyancyComponent(const FObjectInitializer& ObjectInitializer);
 
+	//~ Begin UActorComponent Interface.	
 	virtual void BeginPlay() override;
 	virtual void EndPlay(const EEndPlayReason::Type EndPlayReason) override;
+	virtual void TickComponent(float DeltaTime, enum ELevelTick TickType, FActorComponentTickFunction* ThisTickFunction) override;
+	//~ End UActorComponent Interface
 
+	//~ Begin UObject Interface.	
+	virtual void GetResourceSizeEx(FResourceSizeEx& CumulativeResourceSize) override;
 	virtual void PostLoad() override;
-
 	virtual void Serialize(FArchive& Ar) override;
-
-	virtual void TickComponent(float DeltaTime, enum ELevelTick TickType, FActorComponentTickFunction* ThisTickFunction) override;
+	//~ End UObject Interface
 
 	virtual void Update(float DeltaTime);
 
diff --git a/Engine/Plugins/Experimental/Water/Source/Runtime/Public/GerstnerWaterWaveSubsystem.h b/Engine/Plugins/Experimental/Water/Source/Runtime/Public/GerstnerWaterWaveSubsystem.h
--- a/Engine/Plugins/Experimental/Water/Source/Runtime/Public/GerstnerWaterWaveSubsystem.h
+++ b/Engine/Plugins/Experimental/Water/Source/Runtime/Public/GerstnerWaterWaveSubsystem.h
@@ -26,6 +26,10 @@ public:
 
 	void RebuildGPUData() { bRebuildGPUData = true; }
 
+	//~ Begin UObject Interface.	
+	virtual void GetResourceSizeEx(FResourceSizeEx& CumulativeResourceSize) override;
+	//~ End UObject Interface
+
 private:
 	void BeginFrameCallback();
 
diff --git a/Engine/Plugins/Experimental/Water/Source/Runtime/Public/WaterBodyComponent.h b/Engine/Plugins/Experimental/Water/Source/Runtime/Public/WaterBodyComponent.h
--- a/Engine/Plugins/Experimental/Water/Source/Runtime/Public/WaterBodyComponent.h
+++ b/Engine/Plugins/Experimental/Water/Source/Runtime/Public/WaterBodyComponent.h
@@ -410,6 +410,7 @@ protected:
 	virtual bool MoveComponentImpl(const FVector& Delta, const FQuat& NewRotation, bool bSweep, FHitResult* Hit = NULL, EMoveComponentFlags MoveFlags = MOVECOMP_NoFlags, ETeleportType Teleport = ETeleportType::None) override;
 	virtual void OnComponentCollisionSettingsChanged(bool bUpdateOverlaps) override;
 	virtual void OnGenerateOverlapEventsChanged() override;
+	virtual void GetResourceSizeEx(FResourceSizeEx& CumulativeResourceSize) override;
 
 #if WITH_EDITOR
 	virtual void PreEditUndo() override;
diff --git a/Engine/Plugins/Experimental/Water/Source/Runtime/Public/WaterMeshSceneProxy.h b/Engine/Plugins/Experimental/Water/Source/Runtime/Public/WaterMeshSceneProxy.h
--- a/Engine/Plugins/Experimental/Water/Source/Runtime/Public/WaterMeshSceneProxy.h
+++ b/Engine/Plugins/Experimental/Water/Source/Runtime/Public/WaterMeshSceneProxy.h
@@ -43,7 +43,7 @@ public:
 
 	uint32 GetAllocatedSize() const
 	{
-		return(FPrimitiveSceneProxy::GetAllocatedSize());
+		return(FPrimitiveSceneProxy::GetAllocatedSize() + (WaterVertexFactories.GetAllocatedSize() + WaterVertexFactories.Num() * sizeof(WaterVertexFactoryType)) + WaterQuadTree.GetAllocatedSize());
 	}
 
 #if WITH_WATER_SELECTION_SUPPORT
diff --git a/Engine/Plugins/Experimental/Water/Source/Runtime/Public/WaterQuadTree.h b/Engine/Plugins/Experimental/Water/Source/Runtime/Public/WaterQuadTree.h
--- a/Engine/Plugins/Experimental/Water/Source/Runtime/Public/WaterQuadTree.h
+++ b/Engine/Plugins/Experimental/Water/Source/Runtime/Public/WaterQuadTree.h
@@ -175,10 +175,10 @@ struct FWaterQuadTree
 	const TArray<FMaterialRenderProxy*>& GetWaterMaterials() const { return WaterMaterials; }
 
 	/** Calculate the world distance to a LOD */
-	static float GetLODDistance(int32 InLODLevel, float InLODScale)
-	{
-		return FMath::Pow(2.0f, (float)(InLODLevel + 1)) * InLODScale;
-	}
+	static float GetLODDistance(int32 InLODLevel, float InLODScale) { return FMath::Pow(2.0f, (float)(InLODLevel + 1)) * InLODScale; }
+
+	/** Total memory dynamically allocated by this object */
+	uint32 GetAllocatedSize() const { return NodeData.GetAllocatedSize() + WaterMaterials.GetAllocatedSize() + FarMeshData.GetAllocatedSize(); }
 
 #if WITH_WATER_SELECTION_SUPPORT
 	/** Obtain all possible hit proxies (proxies of all the water bodies) */
@@ -267,6 +267,9 @@ private:
 
 		/** Render data for all water bodies in this tree, indexed by the nodes */
 		TArray<FWaterBodyRenderData> WaterBodyRenderData;
+
+		/** Total memory dynamically allocated by this object */
+		uint32 GetAllocatedSize() const { return Nodes.GetAllocatedSize() + WaterBodyRenderData.GetAllocatedSize(); }
 	} NodeData;
 
 	TArray<FMaterialRenderProxy*> WaterMaterials;
@@ -286,9 +289,6 @@ private:
 		/** Material for the Far Distance Mesh, its material render proxy will be cached in WaterMaterials when BuildMaterialIndices is called */
 		const UMaterialInterface* Material = nullptr;
 
-		/** Cached material index */
-		int16 MaterialIndex = INDEX_NONE;
-
 		void Clear()
 		{
 			InstanceData.Empty();
@@ -296,6 +296,12 @@ private:
 			MaterialIndex = INDEX_NONE;
 		}
 
+		/** Total memory dynamically allocated by this object */
+		uint32 GetAllocatedSize() const { InstanceData.GetAllocatedSize(); }
+
+		/** Cached material index */
+		int16 MaterialIndex = INDEX_NONE;
+
 	} FarMeshData;
 
 	/** If true, the tree may not change */
diff --git a/Engine/Source/Runtime/Engine/Classes/PhysicsEngine/AggregateGeom.h b/Engine/Source/Runtime/Engine/Classes/PhysicsEngine/AggregateGeom.h
--- a/Engine/Source/Runtime/Engine/Classes/PhysicsEngine/AggregateGeom.h
+++ b/Engine/Source/Runtime/Engine/Classes/PhysicsEngine/AggregateGeom.h
@@ -65,6 +65,8 @@ struct ENGINE_API FKAggregateGeom
 
 	int32 GetElementCount(EAggCollisionShape::Type Type) const;
 
+	uint32 GetAllocatedSize() const { return SphereElems.GetAllocatedSize() + SphylElems.GetAllocatedSize() + BoxElems.GetAllocatedSize() + ConvexElems.GetAllocatedSize() + TaperedCapsuleElems.GetAllocatedSize(); }
+
 	FKShapeElem* GetElement(const EAggCollisionShape::Type Type, const int32 Index)
 	{
 		switch (Type)
