Fix CIS issue for missing return statement

diff --git a/Engine/Plugins/Experimental/Water/Source/Runtime/Public/WaterQuadTree.h b/Engine/Plugins/Experimental/Water/Source/Runtime/Public/WaterQuadTree.h
--- a/Engine/Plugins/Experimental/Water/Source/Runtime/Public/WaterQuadTree.h
+++ b/Engine/Plugins/Experimental/Water/Source/Runtime/Public/WaterQuadTree.h
@@ -297,7 +297,7 @@ private:
 		}
 
 		/** Total memory dynamically allocated by this object */
-		uint32 GetAllocatedSize() const { InstanceData.GetAllocatedSize(); }
+		uint32 GetAllocatedSize() const { return InstanceData.GetAllocatedSize(); }
 
 		/** Cached material index */
 		int16 MaterialIndex = INDEX_NONE;
