Remove callbacks for spline/spline metadata data change events when UWaterBodyComponent is unregistered.

Callbacks are registered in OnRegister and during property change events like sliding inputs, the component is unregistered and re-registered many times causing thousands of callbacks to be run per event.

diff --git a/Engine/Plugins/Experimental/Water/Source/Runtime/Private/WaterBodyComponent.cpp b/Engine/Plugins/Experimental/Water/Source/Runtime/Private/WaterBodyComponent.cpp
--- a/Engine/Plugins/Experimental/Water/Source/Runtime/Private/WaterBodyComponent.cpp
+++ b/Engine/Plugins/Experimental/Water/Source/Runtime/Private/WaterBodyComponent.cpp
@@ -674,13 +674,21 @@ void UWaterBodyComponent::OnRegister()
 	CreateOrUpdateWaterInfoMID();
 
 #if WITH_EDITOR
-	RegisterOnChangeWaterSplineMetadata(WaterSplineMetadata, /*bRegister = */true);
-	GetWaterSpline()->OnSplineDataChanged().AddUObject(this, &UWaterBodyComponent::OnSplineDataChanged);
+	RegisterOnChangeWaterSplineData(/*bRegister = */true); 
 
 	CreateWaterSpriteComponent();
 #endif // WITH_EDITOR
 }
 
+void UWaterBodyComponent::OnUnregister()
+{
+#if WITH_EDITOR
+	RegisterOnChangeWaterSplineData(/*bRegister = */false);
+#endif // WITH_EDITOR
+
+	Super::OnUnregister();
+}
+
 TArray<AWaterBodyIsland*> UWaterBodyComponent::GetIslands() const
 {
 	TArray<AWaterBodyIsland*> IslandActors;
@@ -1031,20 +1039,34 @@ void UWaterBodyComponent::OnWaterSplineMetadataChanged(UWaterSplineMetadata* InW
 	OnWaterBodyChanged(bShapeOrPositionChanged);
 }
 
-void UWaterBodyComponent::RegisterOnChangeWaterSplineMetadata(UWaterSplineMetadata* InWaterSplineMetadata, bool bRegister)
+PRAGMA_DISABLE_OPTIMIZATION
+void UWaterBodyComponent::RegisterOnChangeWaterSplineData(bool bRegister)
 {
-	if (InWaterSplineMetadata != nullptr)
+	if (UWaterSplineComponent* WaterSpline = GetWaterSpline())
+	{
+		if (bRegister)
+		{
+			WaterSpline->OnSplineDataChanged().AddUObject(this, &UWaterBodyComponent::OnSplineDataChanged);
+		}
+		else
+		{
+			WaterSpline->OnSplineDataChanged().RemoveAll(this);
+		}
+	}
+
+	if (WaterSplineMetadata != nullptr)
 	{
 		if (bRegister)
 		{
-			InWaterSplineMetadata->OnChangeData.AddUObject(this, &UWaterBodyComponent::OnWaterSplineMetadataChanged);
+			WaterSplineMetadata->OnChangeData.AddUObject(this, &UWaterBodyComponent::OnWaterSplineMetadataChanged);
 		}
 		else
 		{
-			InWaterSplineMetadata->OnChangeData.RemoveAll(this);
+			WaterSplineMetadata->OnChangeData.RemoveAll(this);
 		}
 	}
 }
+PRAGMA_ENABLE_OPTIMIZATION
 
 #endif // WITH_EDITOR
 
@@ -1195,6 +1217,7 @@ void UWaterBodyComponent::UpdateSplineComponent()
 
 void UWaterBodyComponent::OnWaterBodyChanged(bool bShapeOrPositionChanged, bool bWeightmapSettingsChanged)
 {
+	TRACE_CPUPROFILER_EVENT_SCOPE("UWaterBodyComponent::OnWaterBodyChanged")
 	// It's possible to get called without a water spline after the Redo of a water body deletion (i.e. the water body actor gets deleted again, hence its SplineComp is restored to nullptr)
 	//  This is a very-edgy case that needs to be checked everywhere that UpdateAll might hook into so it's simpler to just skip it all. The actor is in limbo by then anyway (it only survives because
 	//  of the editor transaction) :
@@ -1269,7 +1292,7 @@ void UWaterBodyComponent::DeprecateData()
 void UWaterBodyComponent::OnComponentDestroyed(bool bDestroyingHierarchy)
 {
 #if WITH_EDITOR
-	RegisterOnChangeWaterSplineMetadata(WaterSplineMetadata, /*bRegister = */false);
+	RegisterOnChangeWaterSplineData(/*bRegister = */false);
 	RegisterOnUpdateWavesData(GetWaterWaves(), /*bRegister = */false);
 #endif // WITH_EDITOR
 
diff --git a/Engine/Plugins/Experimental/Water/Source/Runtime/Public/WaterBodyComponent.h b/Engine/Plugins/Experimental/Water/Source/Runtime/Public/WaterBodyComponent.h
--- a/Engine/Plugins/Experimental/Water/Source/Runtime/Public/WaterBodyComponent.h
+++ b/Engine/Plugins/Experimental/Water/Source/Runtime/Public/WaterBodyComponent.h
@@ -260,6 +260,7 @@ public:
 
 	/** Component interface */
 	virtual void OnRegister() override;
+	virtual void OnUnregister() override;
 	virtual void PostDuplicate(bool bDuplicateForPie) override;
 
 	UFUNCTION(BlueprintCallable, Category=Water)
@@ -420,9 +421,8 @@ protected:
 	void OnSplineDataChanged();
 	void RegisterOnUpdateWavesData(UWaterWavesBase* InWaterWaves, bool bRegister);
 	void OnWavesDataUpdated(UWaterWavesBase* InWaterWaves, EPropertyChangeType::Type InChangeType);
-
 	void OnWaterSplineMetadataChanged(UWaterSplineMetadata* InWaterSplineMetadata, FPropertyChangedEvent& PropertyChangedEvent);
-	void RegisterOnChangeWaterSplineMetadata(UWaterSplineMetadata* InWaterSplineMetadata, bool bRegister);
+	void RegisterOnChangeWaterSplineData(bool bRegister);
 
 	void CreateWaterSpriteComponent();
 #endif // WITH_EDITOR
