Removed leftover disable optimization macros

#rnx

diff --git a/Engine/Plugins/Experimental/Water/Source/Runtime/Private/WaterBodyComponent.cpp b/Engine/Plugins/Experimental/Water/Source/Runtime/Private/WaterBodyComponent.cpp
--- a/Engine/Plugins/Experimental/Water/Source/Runtime/Private/WaterBodyComponent.cpp
+++ b/Engine/Plugins/Experimental/Water/Source/Runtime/Private/WaterBodyComponent.cpp
@@ -1039,7 +1039,6 @@ void UWaterBodyComponent::OnWaterSplineMetadataChanged(UWaterSplineMetadata* InW
 	OnWaterBodyChanged(bShapeOrPositionChanged);
 }
 
-PRAGMA_DISABLE_OPTIMIZATION
 void UWaterBodyComponent::RegisterOnChangeWaterSplineData(bool bRegister)
 {
 	if (UWaterSplineComponent* WaterSpline = GetWaterSpline())
@@ -1066,7 +1065,6 @@ void UWaterBodyComponent::RegisterOnChangeWaterSplineData(bool bRegister)
 		}
 	}
 }
-PRAGMA_ENABLE_OPTIMIZATION
 
 #endif // WITH_EDITOR
 
