WorldPartition: Added HLOD support for water bodies
* AWaterBody / UWaterBodyComponent are now HLOD relevant
* Added custom HLOD builder for UWaterBodyComponent
* FWaterBodyDefaults now provides an "HLOD Material" which is set on the components when they are created.
* HLOD mesh is built from water mesh, then simplified according to the expected HLOD distance.

diff --git a/Engine/Plugins/Experimental/Water/Source/Editor/Private/WaterBodyActorFactory.cpp b/Engine/Plugins/Experimental/Water/Source/Editor/Private/WaterBodyActorFactory.cpp
--- a/Engine/Plugins/Experimental/Water/Source/Editor/Private/WaterBodyActorFactory.cpp
+++ b/Engine/Plugins/Experimental/Water/Source/Editor/Private/WaterBodyActorFactory.cpp
@@ -40,6 +40,7 @@ void UWaterBodyActorFactory::PostSpawnActor(UObject* Asset, AActor* NewActor)
 	if (const FWaterBodyDefaults* WaterBodyDefaults = GetWaterBodyDefaults())
 	{
 		WaterBodyComponent->SetWaterMaterial(WaterBodyDefaults->GetWaterMaterial());
+		WaterBodyComponent->SetHLODMaterial(WaterBodyDefaults->GetWaterHLODMaterial());
 		WaterBodyComponent->SetUnderwaterPostProcessMaterial(WaterBodyDefaults->GetUnderwaterPostProcessMaterial());
 
 		UWaterSplineComponent* WaterSpline = WaterBodyComponent->GetWaterSpline();
diff --git a/Engine/Plugins/Experimental/Water/Source/Editor/Private/WaterEditorSettings.cpp b/Engine/Plugins/Experimental/Water/Source/Editor/Private/WaterEditorSettings.cpp
--- a/Engine/Plugins/Experimental/Water/Source/Editor/Private/WaterEditorSettings.cpp
+++ b/Engine/Plugins/Experimental/Water/Source/Editor/Private/WaterEditorSettings.cpp
@@ -31,6 +31,7 @@ UMaterialInterface* FWaterZoneActorDefaults::GetFarDistanceMaterial() const
 
 FWaterBodyDefaults::FWaterBodyDefaults()
 	: WaterMaterial(FSoftObjectPath(TEXT("/Water/Materials/WaterSurface/Water_Material.Water_Material")))
+	, WaterHLODMaterial(FSoftObjectPath(TEXT("/Water/Materials/HLOD/HLODWater.HLODWater")))
 	, UnderwaterPostProcessMaterial(FSoftObjectPath(TEXT("/Water/Materials/PostProcessing/M_UnderWater_PostProcess_Volume.M_UnderWater_PostProcess_Volume")))
 {}
 
@@ -39,6 +40,11 @@ UMaterialInterface* FWaterBodyDefaults::GetWaterMaterial() const
 	return WaterMaterial.LoadSynchronous();
 }
 
+UMaterialInterface* FWaterBodyDefaults::GetWaterHLODMaterial() const
+{
+	return WaterHLODMaterial.LoadSynchronous();
+}
+
 UMaterialInterface* FWaterBodyDefaults::GetUnderwaterPostProcessMaterial() const
 {
 	return UnderwaterPostProcessMaterial.LoadSynchronous();
@@ -222,4 +228,4 @@ UMaterialInterface* UWaterEditorSettings::GetDefaultDrawCanvasMaterial() const
 UMaterialInterface* UWaterEditorSettings::GetDefaultRenderRiverSplineDepthsMaterial() const
 {
 	return DefaultRenderRiverSplineDepthsMaterial.LoadSynchronous();
-}
\ No newline at end of file
+}
diff --git a/Engine/Plugins/Experimental/Water/Source/Editor/Public/WaterEditorSettings.h b/Engine/Plugins/Experimental/Water/Source/Editor/Public/WaterEditorSettings.h
--- a/Engine/Plugins/Experimental/Water/Source/Editor/Public/WaterEditorSettings.h
+++ b/Engine/Plugins/Experimental/Water/Source/Editor/Public/WaterEditorSettings.h
@@ -68,6 +68,9 @@ public:
 	UMaterialInterface* GetWaterMaterial() const;
 	FSoftObjectPath GetWaterMaterialPath() const { return WaterMaterial.ToSoftObjectPath(); }
 
+	UMaterialInterface* GetWaterHLODMaterial() const;
+	FSoftObjectPath GetWaterHLODMaterialPath() const { return WaterHLODMaterial.ToSoftObjectPath(); }
+
 	UMaterialInterface* GetUnderwaterPostProcessMaterial() const;
 	FSoftObjectPath GetUnderwaterPostProcessMaterialPath() const { return UnderwaterPostProcessMaterial.ToSoftObjectPath(); }
 
@@ -75,6 +78,9 @@ protected:
 	UPROPERTY(EditAnywhere, config, Category = Rendering)
 	TSoftObjectPtr<UMaterialInterface> WaterMaterial;
 
+	UPROPERTY(EditAnywhere, config, Category = Rendering)
+	TSoftObjectPtr<UMaterialInterface> WaterHLODMaterial;
+
 	UPROPERTY(EditAnywhere, config, Category = Rendering)
 	TSoftObjectPtr<UMaterialInterface> UnderwaterPostProcessMaterial;
 };
diff --git a/Engine/Plugins/Experimental/Water/Source/Runtime/Private/WaterBodyActor.cpp b/Engine/Plugins/Experimental/Water/Source/Runtime/Private/WaterBodyActor.cpp
--- a/Engine/Plugins/Experimental/Water/Source/Runtime/Private/WaterBodyActor.cpp
+++ b/Engine/Plugins/Experimental/Water/Source/Runtime/Private/WaterBodyActor.cpp
@@ -567,6 +567,11 @@ void AWaterBody::PostRegisterAllComponents()
 	}
 }
 
+bool AWaterBody::IsHLODRelevant() const
+{
+	return true;
+}
+
 #if WITH_EDITOR
 void AWaterBody::SetActorHiddenInGame(bool bNewHidden)
 {
diff --git a/Engine/Plugins/Experimental/Water/Source/Runtime/Private/WaterBodyComponent.cpp b/Engine/Plugins/Experimental/Water/Source/Runtime/Private/WaterBodyComponent.cpp
--- a/Engine/Plugins/Experimental/Water/Source/Runtime/Private/WaterBodyComponent.cpp
+++ b/Engine/Plugins/Experimental/Water/Source/Runtime/Private/WaterBodyComponent.cpp
@@ -35,12 +35,15 @@
 #include "Logging/MessageLog.h"
 #include "Logging/TokenizedMessage.h"
 #include "WaterBodySceneProxy.h"
+#include "Engine/StaticMesh.h"
 
 #if WITH_EDITOR
 #include "WaterIconHelper.h"
 #include "Components/BillboardComponent.h"
 #include "Modules/ModuleManager.h"
 #include "WaterModule.h"
+#include "StaticMeshAttributes.h"
+#include "WaterBodyHLODBuilder.h"
 #endif
 
 #define LOCTEXT_NAMESPACE "Water"
@@ -100,6 +103,16 @@ UWaterBodyComponent::UWaterBodyComponent(const FObjectInitializer& ObjectInitial
 #endif // WITH_EDITORONLY_DATA
 }
 
+bool UWaterBodyComponent::IsHLODRelevant() const
+{
+	if (!bEnableAutoLODGeneration)
+	{
+		return false;
+	}
+
+	return true;
+}
+
 void UWaterBodyComponent::OnVisibilityChanged()
 {
 	Super::OnVisibilityChanged();
@@ -1257,6 +1270,18 @@ void UWaterBodyComponent::PostLoad()
 	{
 		WaterMeshOverride = nullptr;
 	}
+
+	// If available, use far mesh material as the HLOD material for water bodies created before HLOD support was added.
+	if (GetLinkerCustomVersion(FFortniteMainBranchObjectVersion::GUID) < FFortniteMainBranchObjectVersion::WaterHLODSupportAdded)
+	{
+		if (const AWaterZone* WaterZone = GetWaterZone())
+		{
+			const UWaterMeshComponent* WaterMeshComponent = WaterZone->GetWaterMeshComponent();
+			check(WaterMeshComponent);
+
+			WaterHLODMaterial = WaterMeshComponent->FarDistanceMaterial;
+		}
+	}
 #endif // WITH_EDITORONLY_DATA
 
 	DeprecateData();
@@ -1644,6 +1669,74 @@ void UWaterBodyComponent::UpdateWaterSpriteComponent()
 		SpriteComponent->MarkRenderStateDirty();
 	}
 }
+
+TSubclassOf<UHLODBuilder> UWaterBodyComponent::GetCustomHLODBuilderClass() const
+{
+	return UWaterBodyHLODBuilder::StaticClass();
+}
+
+FMeshDescription UWaterBodyComponent::GetHLODMeshDescription() const
+{
+	if (WaterMeshOverride)
+	{
+		return *WaterMeshOverride->GetMeshDescription(WaterMeshOverride->GetNumLODs() - 1);
+	}
+	
+	FMeshDescription MeshDescription;
+
+	FStaticMeshAttributes StaticMeshAttributes(MeshDescription);
+	StaticMeshAttributes.Register();
+
+	TVertexAttributesRef<FVector3f> VertexPositions = StaticMeshAttributes.GetVertexPositions();
+
+	const int32 NumVertices = WaterBodyMeshVertices.Num();
+	const int32 NumTriangles = WaterBodyMeshIndices.Num() / 3;
+
+	MeshDescription.ReserveNewVertices(NumVertices);
+	MeshDescription.ReserveNewVertexInstances(NumVertices);
+	MeshDescription.ReserveNewTriangles(NumTriangles);
+
+	FPolygonGroupID PolygonGroupID = MeshDescription.CreatePolygonGroup();
+
+	// Positions
+	for (int32 VertexIndex = 0; VertexIndex < NumVertices; ++VertexIndex)
+	{
+		FVertexID VertexID = MeshDescription.CreateVertex();
+		VertexPositions[VertexID] = WaterBodyMeshVertices[VertexIndex].Position;
+	}
+
+	// Triangles
+	for (int32 TriangleIndex = 0; TriangleIndex < NumTriangles; ++TriangleIndex)
+	{
+		TStaticArray<FVertexInstanceID, 3> VertexInstanceIDs;
+
+		for (int32 Corner = 0; Corner < 3; ++Corner)
+		{
+			uint32 VertexIndex = WaterBodyMeshIndices[TriangleIndex * 3 + Corner];
+
+			FVertexID VertexID(VertexIndex);
+			FVertexInstanceID VertexInstanceID = MeshDescription.CreateVertexInstance(VertexID);
+			
+			VertexInstanceIDs[Corner] = VertexInstanceID;
+		}
+
+		// Create a triangle
+		MeshDescription.CreateTriangle(PolygonGroupID, VertexInstanceIDs);
+	}
+	
+	return MeshDescription;
+}
+
+UMaterialInterface* UWaterBodyComponent::GetHLODMaterial() const
+{
+	return WaterHLODMaterial;
+}
+
+void UWaterBodyComponent::SetHLODMaterial(UMaterialInterface* InMaterial)
+{
+	WaterHLODMaterial = InMaterial;
+}
+
 #endif // WITH_EDITOR
 
 #undef LOCTEXT_NAMESPACE
diff --git a/Engine/Plugins/Experimental/Water/Source/Runtime/Private/WaterBodyHLODBuilder.cpp b/Engine/Plugins/Experimental/Water/Source/Runtime/Private/WaterBodyHLODBuilder.cpp
new file mode 100644
--- /dev/null
+++ b/Engine/Plugins/Experimental/Water/Source/Runtime/Private/WaterBodyHLODBuilder.cpp
@@ -0,0 +1,126 @@
+// Copyright Epic Games, Inc. All Rights Reserved.
+
+#include "WaterBodyHLODBuilder.h"
+
+#include "WaterBodyComponent.h"
+#include "WaterZoneActor.h"
+#include "WaterMeshComponent.h"
+
+#include "MeshDescription.h"
+#include "TriangleTypes.h"
+#include "Engine/StaticMesh.h"
+#include "StaticMeshAttributes.h"
+#include "Engine/StaticMeshSourceData.h"
+
+#include "DynamicMesh/DynamicMesh3.h"
+#include "MeshSimplification.h"
+#include "MeshDescriptionToDynamicMesh.h"
+#include "DynamicMeshToMeshDescription.h"
+
+#include "Serialization/ArchiveCrc32.h"
+#include "Engine/HLODProxy.h"
+
+UWaterBodyHLODBuilder::UWaterBodyHLODBuilder(const FObjectInitializer& ObjectInitializer)
+	: Super(ObjectInitializer)
+{
+}
+
+static double GetPixelToWorldUnit(const double InDistanceZ)
+{
+	// Generate a projection matrix.
+	static const FIntPoint ScreenSize(1920, 1080);
+	static const float HalfFOVRad = FMath::DegreesToRadians(45.0f);
+	static const FMatrix InvProjectionMatrix = FPerspectiveMatrix(HalfFOVRad, ScreenSize.X, ScreenSize.Y, 0.01f).Inverse();
+
+	const FVector2D ScreenPixel((ScreenSize / 2) + FIntPoint(1, 0));
+	FVector Origin, Direction;
+	FSceneView::DeprojectScreenToWorld(ScreenPixel, FIntRect(FIntPoint(0, 0), ScreenSize), InvProjectionMatrix, Origin, Direction);
+	const FVector4d WorldPos = -Direction * InDistanceZ;
+	return WorldPos.X;
+}
+
+uint32 UWaterBodyHLODBuilder::ComputeHLODHash(const UActorComponent* InSourceComponent) const
+{
+	FArchiveCrc32 Ar;
+
+	if (const UWaterBodyComponent* WaterBodyComponent = Cast<UWaterBodyComponent>(InSourceComponent))
+	{
+		uint32 TransformHash = UHLODProxy::GetCRC(WaterBodyComponent->GetComponentTransform());
+		Ar << TransformHash;
+		
+		FMeshDescription HLODMesh = WaterBodyComponent->GetHLODMeshDescription();
+		Ar << HLODMesh;
+
+		UMaterialInterface* HLODMaterial = WaterBodyComponent->GetHLODMaterial();
+		Ar << HLODMaterial;
+	}
+
+	return Ar.GetCrc();
+}
+
+TArray<UActorComponent*> UWaterBodyHLODBuilder::Build(const FHLODBuildContext& InHLODBuildContext, const TArray<UActorComponent*>& InSourceComponents) const
+{
+	TArray<UWaterBodyComponent*> SourceWaterBodyComponents = FilterComponents<UWaterBodyComponent>(InSourceComponents);
+	
+	TArray<UActorComponent*> HLODComponents;
+	TArray<UStaticMesh*> StaticMeshes;
+
+	// Compute the maximum error to keep displacement under a single pixel for the requested HLOD draw distance.
+	const double MaxDisplacement = GetPixelToWorldUnit(InHLODBuildContext.MinVisibleDistance);
+	const double MaxDisplacementSqr = FMath::Square(MaxDisplacement);
+
+	for (UWaterBodyComponent* WaterBodyComponent : SourceWaterBodyComponents)
+	{
+		UStaticMesh* StaticMesh = NewObject<UStaticMesh>(InHLODBuildContext.AssetsOuter);
+
+		// Mesh
+		{
+			FMeshDescription HLODMesh = WaterBodyComponent->GetHLODMeshDescription();
+
+			// Simplify mesh
+			{
+				// MeshDescription -> DynamicMesh
+				UE::Geometry::FDynamicMesh3 DynamicMesh;
+				FMeshDescriptionToDynamicMesh MeshDescriptionToDynamicMesh;
+				MeshDescriptionToDynamicMesh.bDisableAttributes = true;
+				MeshDescriptionToDynamicMesh.Convert(&HLODMesh, DynamicMesh, false);
+
+				// Simplify
+				UE::Geometry::FQEMSimplification Simplifier(&DynamicMesh);
+				Simplifier.bPreserveBoundaryShape = false;
+				Simplifier.SimplifyToMaxError(MaxDisplacementSqr);
+				
+				// DynamicMesh -> MeshDescription
+				FMeshDescription SimplifiedWaterMesh = HLODMesh;
+				FDynamicMeshToMeshDescription DynamicMeshToMeshDescription;
+				DynamicMeshToMeshDescription.Convert(&DynamicMesh, SimplifiedWaterMesh);
+								
+				HLODMesh = SimplifiedWaterMesh;
+			}
+
+			StaticMesh->AddSourceModel();
+			StaticMesh->CreateMeshDescription(0, HLODMesh);
+			StaticMesh->CommitMeshDescription(0);
+			StaticMesh->ImportVersion = EImportStaticMeshVersion::LastVersion;
+		}
+		
+		// Material
+		{
+			UMaterialInterface* HLODMaterial = WaterBodyComponent->GetHLODMaterial();
+			StaticMesh->AddMaterial(HLODMaterial);
+		}
+		
+		StaticMeshes.Add(StaticMesh);
+
+		FName ComponentName = MakeUniqueObjectName(nullptr, UStaticMeshComponent::StaticClass(), TEXT("WaterHLODMeshComponent"));
+		UStaticMeshComponent* StaticMeshComponent = NewObject<UStaticMeshComponent>(GetTransientPackage(), ComponentName);
+		StaticMeshComponent->SetRelativeTransform(WaterBodyComponent->GetComponentTransform());
+		StaticMeshComponent->SetStaticMesh(StaticMesh);
+		HLODComponents.Add(StaticMeshComponent);
+	}
+
+	UStaticMesh::BatchBuild(StaticMeshes);
+
+	return HLODComponents;
+}
+
diff --git a/Engine/Plugins/Experimental/Water/Source/Runtime/Public/WaterBodyActor.h b/Engine/Plugins/Experimental/Water/Source/Runtime/Public/WaterBodyActor.h
--- a/Engine/Plugins/Experimental/Water/Source/Runtime/Public/WaterBodyActor.h
+++ b/Engine/Plugins/Experimental/Water/Source/Runtime/Public/WaterBodyActor.h
@@ -68,6 +68,7 @@ public:
 	virtual void Serialize(FArchive& Ar) override;
 	virtual void PostLoad() override;
 	virtual void PostRegisterAllComponents() override;
+	virtual bool IsHLODRelevant() const override;
 
 #if WITH_EDITOR
 	virtual void SetActorHiddenInGame(bool bNewHidden) override;
diff --git a/Engine/Plugins/Experimental/Water/Source/Runtime/Public/WaterBodyComponent.h b/Engine/Plugins/Experimental/Water/Source/Runtime/Public/WaterBodyComponent.h
--- a/Engine/Plugins/Experimental/Water/Source/Runtime/Public/WaterBodyComponent.h
+++ b/Engine/Plugins/Experimental/Water/Source/Runtime/Public/WaterBodyComponent.h
@@ -17,6 +17,11 @@
 #include "WaterWaves.h"
 #include "WaterBodyTypes.h"
 #include "DynamicMeshBuilder.h"
+
+#if WITH_EDITOR
+#include "MeshDescription.h"
+#endif
+
 #include "WaterBodyComponent.generated.h"
 
 class UWaterSplineComponent;
@@ -101,6 +106,10 @@ public:
 	virtual TArray<UPrimitiveComponent*> GetBrushRenderableComponents() const { return TArray<UPrimitiveComponent*>(); }
 	
 	virtual void UpdateWaterSpriteComponent();
+
+	virtual FMeshDescription GetHLODMeshDescription() const;
+	virtual UMaterialInterface* GetHLODMaterial() const;
+	virtual void SetHLODMaterial(UMaterialInterface* InMaterial);
 #endif //WITH_EDITOR
 
 	/** Returns whether the body supports waves */
@@ -330,7 +339,12 @@ public:
 	
 	/** Set the navigation area class */
 	void SetNavAreaClass(TSubclassOf<UNavAreaBase> NewWaterNavAreaClass) { WaterNavAreaClass = NewWaterNavAreaClass; }
+
 protected:
+	//~ Begin UActorComponent interface.
+	virtual bool IsHLODRelevant() const override;
+	//~ End UActorComponent interface.
+
 	//~ Begin USceneComponent Interface.
 	virtual void OnVisibilityChanged();
 	virtual void OnHiddenInGameChanged();
@@ -425,6 +439,8 @@ protected:
 	void RegisterOnChangeWaterSplineData(bool bRegister);
 
 	void CreateWaterSpriteComponent();
+
+	virtual TSubclassOf<class UHLODBuilder> GetCustomHLODBuilderClass() const override;
 #endif // WITH_EDITOR
 
 public:
@@ -470,6 +486,9 @@ public:
 	UPROPERTY(Category = Rendering, EditAnywhere, BlueprintReadOnly)
 	UMaterialInterface* WaterMaterial;
 
+	UPROPERTY(Category = HLOD, EditAnywhere, BlueprintReadOnly, meta = (DisplayName = "Water HLOD Material"))
+	UMaterialInterface* WaterHLODMaterial;
+
 	/** Post process material to apply when the camera goes underwater (only available when bGenerateCollisions is true because collisions are needed to detect if it's under water). */
 	UPROPERTY(Category = Rendering, EditAnywhere, BlueprintReadOnly, meta = (EditCondition = "bGenerateCollisions", DisplayAfter = "WaterMaterial"))
 	UMaterialInterface* UnderwaterPostProcessMaterial;
diff --git a/Engine/Plugins/Experimental/Water/Source/Runtime/Public/WaterBodyHLODBuilder.h b/Engine/Plugins/Experimental/Water/Source/Runtime/Public/WaterBodyHLODBuilder.h
new file mode 100644
--- /dev/null
+++ b/Engine/Plugins/Experimental/Water/Source/Runtime/Public/WaterBodyHLODBuilder.h
@@ -0,0 +1,22 @@
+// Copyright Epic Games, Inc. All Rights Reserved.
+
+#pragma once
+
+#include "CoreMinimal.h"
+
+#include "WorldPartition/HLOD/HLODBuilder.h"
+#include "WaterBodyHLODBuilder.generated.h"
+
+UCLASS()
+class UWaterBodyHLODBuilder : public UHLODBuilder
+{
+	GENERATED_UCLASS_BODY()
+
+public:
+	virtual uint32 ComputeHLODHash(const UActorComponent* InSourceComponent) const override;
+
+	/**
+	 * Components created with this method need to be properly outered & assigned to your target actor.
+	 */
+	virtual TArray<UActorComponent*> Build(const FHLODBuildContext& InHLODBuildContext, const TArray<UActorComponent*>& InSourceComponents) const override;
+};
diff --git a/Engine/Plugins/Experimental/Water/Source/Runtime/Water.Build.cs b/Engine/Plugins/Experimental/Water/Source/Runtime/Water.Build.cs
--- a/Engine/Plugins/Experimental/Water/Source/Runtime/Water.Build.cs
+++ b/Engine/Plugins/Experimental/Water/Source/Runtime/Water.Build.cs
@@ -46,10 +46,18 @@ public class Water : ModuleRules
         {
 			bWithWaterSelectionSupport = true;
 
+			PublicDependencyModuleNames.AddRange(
+				new string[] {
+					"MeshDescription"
+				}
+			);
+
 			PrivateDependencyModuleNames.AddRange(
 				new string[] {
 					"SourceControl",
-					"UnrealEd"
+					"UnrealEd",
+					"StaticMeshDescription",
+					"MeshConversion"
 				}
 			);
 		}
diff --git a/Engine/Source/Runtime/Core/Public/UObject/FortniteMainBranchObjectVersion.h b/Engine/Source/Runtime/Core/Public/UObject/FortniteMainBranchObjectVersion.h
--- a/Engine/Source/Runtime/Core/Public/UObject/FortniteMainBranchObjectVersion.h
+++ b/Engine/Source/Runtime/Core/Public/UObject/FortniteMainBranchObjectVersion.h
@@ -226,6 +226,9 @@ struct CORE_API FFortniteMainBranchObjectVersion
 		// if a widget exposes its named slot to everyone (even if it has content), which by default they wont any longer.
 		WidgetInheritedNamedSlots,
 
+		// Added water HLOD material
+		WaterHLODSupportAdded,
+
 		// -----<new versions can be added above this line>-------------------------------------------------
 		VersionPlusOne,
 		LatestVersion = VersionPlusOne - 1
diff --git a/Engine/Source/Runtime/Engine/Classes/Engine/HLODProxy.h b/Engine/Source/Runtime/Engine/Classes/Engine/HLODProxy.h
--- a/Engine/Source/Runtime/Engine/Classes/Engine/HLODProxy.h
+++ b/Engine/Source/Runtime/Engine/Classes/Engine/HLODProxy.h
@@ -52,6 +52,7 @@ public:
 	/** Build a unique key for the LOD actor, used to determine if the actor needs rebuilding */
 	static FName GenerateKeyForActor(const ALODActor* LODActor, bool bMustUndoLevelTransform = true);
 
+	static uint32 GetCRC(const FTransform& InTransform, uint32 InCRC = 0);
 	static uint32 GetCRC(UMaterialInterface* InMaterialInterface, uint32 InCRC = 0);
 	static uint32 GetCRC(UTexture* InTexture, uint32 InCRC = 0);
 	static uint32 GetCRC(UStaticMesh* InStaticMesh, uint32 InCRC = 0, bool bInConsiderPhysicData = false);
diff --git a/Engine/Source/Runtime/Engine/Private/HLODProxy.cpp b/Engine/Source/Runtime/Engine/Private/HLODProxy.cpp
--- a/Engine/Source/Runtime/Engine/Private/HLODProxy.cpp
+++ b/Engine/Source/Runtime/Engine/Private/HLODProxy.cpp
@@ -425,7 +425,7 @@ static void AppendRoundedTransform(const FTransform& InTransform, FArchive& Ar)
 	AppendRoundedTransform(InTransform.Rotator(), InTransform.GetLocation(), InTransform.GetScale3D(), Ar);
 }
 
-static int32 GetTransformCRC(const FTransform& InTransform, uint32 InCRC = 0)
+uint32 UHLODProxy::GetCRC(const FTransform& InTransform, uint32 InCRC)
 {
 	FArchiveCrc32 Ar(InCRC);
 	AppendRoundedTransform(InTransform, Ar);
@@ -537,7 +537,7 @@ FName UHLODProxy::GenerateKeyForActor(const ALODActor* LODActor, bool bMustUndoL
 		}
 
 		// HLODBakingTransform
-		CRC = GetTransformCRC(LODActor->GetLevel()->GetWorldSettings()->HLODBakingTransform, CRC);
+		CRC = GetCRC(LODActor->GetLevel()->GetWorldSettings()->HLODBakingTransform, CRC);
 
 		// screen size + override
 		{
@@ -753,13 +753,13 @@ bool UHLODProxy::SetHLODBakingTransform(const FTransform& InTransform)
 {
 	bool bChanged = false;
 
-	int32 NewTransformCRC = GetTransformCRC(InTransform);
+	int32 NewTransformCRC = GetCRC(InTransform);
 
 	for (auto ItHLODActor = HLODActors.CreateIterator(); ItHLODActor; ++ItHLODActor)
 	{
 		UHLODProxyDesc* HLODProxyDesc = ItHLODActor.Key();
 
-		int32 OldTransformCRC = GetTransformCRC(HLODProxyDesc->HLODBakingTransform);
+		int32 OldTransformCRC = GetCRC(HLODProxyDesc->HLODBakingTransform);
 		if (OldTransformCRC != NewTransformCRC)
 		{
 			HLODProxyDesc->HLODBakingTransform = InTransform;
diff --git a/Engine/Source/Runtime/Engine/Public/WorldPartition/HLOD/HLODBuilder.h b/Engine/Source/Runtime/Engine/Public/WorldPartition/HLOD/HLODBuilder.h
--- a/Engine/Source/Runtime/Engine/Public/WorldPartition/HLOD/HLODBuilder.h
+++ b/Engine/Source/Runtime/Engine/Public/WorldPartition/HLOD/HLODBuilder.h
@@ -99,7 +99,7 @@ public:
 	virtual uint32 ComputeHLODHash(const UActorComponent* InSourceComponent) const;
 
 	/**
-	 * Components created with this method needs to be properly outered & assigned to your target actor.
+	 * Components created with this method need to be properly outered & assigned to your target actor.
 	 */
 	virtual TArray<UActorComponent*> Build(const FHLODBuildContext& InHLODBuildContext, const TArray<UActorComponent*>& InSourceComponents) const PURE_VIRTUAL(UHLODBuilder::CreateComponents, return {};);
 
diff --git a/Engine/Source/Runtime/Landscape/Classes/LandscapeComponent.h b/Engine/Source/Runtime/Landscape/Classes/LandscapeComponent.h
--- a/Engine/Source/Runtime/Landscape/Classes/LandscapeComponent.h
+++ b/Engine/Source/Runtime/Landscape/Classes/LandscapeComponent.h
@@ -852,7 +852,7 @@ public:
 	virtual void InvalidateLightingCacheDetailed(bool bInvalidateBuildEnqueuedLighting, bool bTranslationOnly) override;
 #endif
 	virtual void PropagateLightingScenarioChange() override;
-	virtual bool IsHLODRelevant() const { return true; }
+	virtual bool IsHLODRelevant() const override;
 	//~ End UActorComponent Interface.
 
 	/** Gets the landscape info object for this landscape */
diff --git a/Engine/Source/Runtime/Landscape/Private/Landscape.cpp b/Engine/Source/Runtime/Landscape/Private/Landscape.cpp
--- a/Engine/Source/Runtime/Landscape/Private/Landscape.cpp
+++ b/Engine/Source/Runtime/Landscape/Private/Landscape.cpp
@@ -1613,6 +1613,11 @@ void ULandscapeComponent::PropagateLightingScenarioChange()
 	FComponentRecreateRenderStateContext Context(this);
 }
 
+bool ULandscapeComponent::IsHLODRelevant() const
+{
+	return bEnableAutoLODGeneration;
+}
+
 TArray<URuntimeVirtualTexture*> const& ULandscapeComponent::GetRuntimeVirtualTextures() const
 {
 	return GetLandscapeProxy()->RuntimeVirtualTextures;
diff --git a/Engine/Source/Runtime/Landscape/Private/LandscapeHLODBuilder.cpp b/Engine/Source/Runtime/Landscape/Private/LandscapeHLODBuilder.cpp
--- a/Engine/Source/Runtime/Landscape/Private/LandscapeHLODBuilder.cpp
+++ b/Engine/Source/Runtime/Landscape/Private/LandscapeHLODBuilder.cpp
@@ -33,25 +33,6 @@ ULandscapeHLODBuilder::ULandscapeHLODBuilder(const FObjectInitializer& ObjectIni
 
 #if WITH_EDITOR
 
-static uint32 GetTransformHash(const FRotator& ComponentRotation, const FVector& ComponentLocation, const FVector& ComponentScale)
-{
-	FArchiveCrc32 Ar;
-
-	// Include transform - round sufficiently to ensure stability
-	FIntVector Location(FMath::RoundToInt(ComponentLocation.X), FMath::RoundToInt(ComponentLocation.Y), FMath::RoundToInt(ComponentLocation.Z));
-	Ar << Location;
-
-	FRotator Rotator(ComponentRotation.GetDenormalized());
-	FIntVector Rotation(FMath::RoundToInt(Rotator.Pitch), FMath::RoundToInt(Rotator.Yaw), FMath::RoundToInt(Rotator.Roll));
-	Ar << Rotation;
-
-	float SCALE_FACTOR = 100;
-	FIntVector Scale(FMath::RoundToInt(ComponentScale.X * SCALE_FACTOR), FMath::RoundToInt(ComponentScale.Y * SCALE_FACTOR), FMath::RoundToInt(ComponentScale.Z * SCALE_FACTOR));
-	Ar << Scale;
-
-	return Ar.GetCrc();
-}
-
 uint32 ULandscapeHLODBuilder::ComputeHLODHash(const UActorComponent* InSourceComponent) const
 {
 	FArchiveCrc32 Ar;
@@ -68,7 +49,7 @@ uint32 ULandscapeHLODBuilder::ComputeHLODHash(const UActorComponent* InSourceCom
 		UE_LOG(LogHLODBuilder, VeryVerbose, TEXT("     - LODScreenSize = %x"), Ar.GetCrc());
 		
 		// LS Transform
-		uint32 TransformHash = GetTransformHash(LSComponent->GetComponentRotation(), LSComponent->GetComponentLocation(), LSComponent->GetComponentScale());
+		uint32 TransformHash = UHLODProxy::GetCRC(LSComponent->GetComponentTransform());
 		Ar << TransformHash;
 		UE_LOG(LogHLODBuilder, VeryVerbose, TEXT("     - TransformHash = %x"), Ar.GetCrc());
 
diff --git a/Engine/Source/Runtime/Landscape/Public/LandscapeHLODBuilder.h b/Engine/Source/Runtime/Landscape/Public/LandscapeHLODBuilder.h
--- a/Engine/Source/Runtime/Landscape/Public/LandscapeHLODBuilder.h
+++ b/Engine/Source/Runtime/Landscape/Public/LandscapeHLODBuilder.h
@@ -14,10 +14,10 @@ class LANDSCAPE_API ULandscapeHLODBuilder : public UHLODBuilder
 
 public:
 #if WITH_EDITOR
-	virtual uint32 ComputeHLODHash(const UActorComponent* InSourceComponent) const;
+	virtual uint32 ComputeHLODHash(const UActorComponent* InSourceComponent) const override;
 
 	/**
-	 * Components created with this method needs to be properly outered & assigned to your target actor.
+	 * Components created with this method need to be properly outered & assigned to your target actor.
 	 */
 	virtual TArray<UActorComponent*> Build(const FHLODBuildContext& InHLODBuildContext, const TArray<UActorComponent*>& InSourceComponents) const override;
 #endif
