CIS fix
#rnx

diff --git a/Engine/Plugins/Experimental/Water/Source/Runtime/Private/WaterBodyHLODBuilder.cpp b/Engine/Plugins/Experimental/Water/Source/Runtime/Private/WaterBodyHLODBuilder.cpp
--- a/Engine/Plugins/Experimental/Water/Source/Runtime/Private/WaterBodyHLODBuilder.cpp
+++ b/Engine/Plugins/Experimental/Water/Source/Runtime/Private/WaterBodyHLODBuilder.cpp
@@ -2,6 +2,8 @@
 
 #include "WaterBodyHLODBuilder.h"
 
+#if WITH_EDITOR
+
 #include "WaterBodyComponent.h"
 #include "WaterZoneActor.h"
 #include "WaterMeshComponent.h"
@@ -124,3 +126,4 @@ TArray<UActorComponent*> UWaterBodyHLODBuilder::Build(const FHLODBuildContext& I
 	return HLODComponents;
 }
 
+#endif // #if WITH_EDITOR
diff --git a/Engine/Plugins/Experimental/Water/Source/Runtime/Public/WaterBodyHLODBuilder.h b/Engine/Plugins/Experimental/Water/Source/Runtime/Public/WaterBodyHLODBuilder.h
--- a/Engine/Plugins/Experimental/Water/Source/Runtime/Public/WaterBodyHLODBuilder.h
+++ b/Engine/Plugins/Experimental/Water/Source/Runtime/Public/WaterBodyHLODBuilder.h
@@ -13,10 +13,12 @@ class UWaterBodyHLODBuilder : public UHLODBuilder
 	GENERATED_UCLASS_BODY()
 
 public:
+#if WITH_EDITOR
 	virtual uint32 ComputeHLODHash(const UActorComponent* InSourceComponent) const override;
 
 	/**
 	 * Components created with this method need to be properly outered & assigned to your target actor.
 	 */
 	virtual TArray<UActorComponent*> Build(const FHLODBuildContext& InHLODBuildContext, const TArray<UActorComponent*>& InSourceComponents) const override;
+#endif
 };
