CIS fix
#rnx

diff --git a/Engine/Plugins/Experimental/Water/Source/Runtime/Private/WaterBodyComponent.cpp b/Engine/Plugins/Experimental/Water/Source/Runtime/Private/WaterBodyComponent.cpp
--- a/Engine/Plugins/Experimental/Water/Source/Runtime/Private/WaterBodyComponent.cpp
+++ b/Engine/Plugins/Experimental/Water/Source/Runtime/Private/WaterBodyComponent.cpp
@@ -105,12 +105,11 @@ UWaterBodyComponent::UWaterBodyComponent(const FObjectInitializer& ObjectInitial
 
 bool UWaterBodyComponent::IsHLODRelevant() const
 {
-	if (!bEnableAutoLODGeneration)
-	{
-		return false;
-	}
-
-	return true;
+#if WITH_EDITOR
+	return bEnableAutoLODGeneration;
+#else
+	return false;
+#endif
 }
 
 void UWaterBodyComponent::OnVisibilityChanged()
diff --git a/Engine/Source/Runtime/Landscape/Private/Landscape.cpp b/Engine/Source/Runtime/Landscape/Private/Landscape.cpp
--- a/Engine/Source/Runtime/Landscape/Private/Landscape.cpp
+++ b/Engine/Source/Runtime/Landscape/Private/Landscape.cpp
@@ -1615,7 +1615,11 @@ void ULandscapeComponent::PropagateLightingScenarioChange()
 
 bool ULandscapeComponent::IsHLODRelevant() const
 {
+#if WITH_EDITOR
 	return bEnableAutoLODGeneration;
+#else
+	return false;
+#endif
 }
 
 TArray<URuntimeVirtualTexture*> const& ULandscapeComponent::GetRuntimeVirtualTextures() const
