CIS fix
#rnx

diff --git a/Engine/Plugins/Experimental/Water/Source/Runtime/Private/WaterBodyHLODBuilder.cpp b/Engine/Plugins/Experimental/Water/Source/Runtime/Private/WaterBodyHLODBuilder.cpp
--- a/Engine/Plugins/Experimental/Water/Source/Runtime/Private/WaterBodyHLODBuilder.cpp
+++ b/Engine/Plugins/Experimental/Water/Source/Runtime/Private/WaterBodyHLODBuilder.cpp
@@ -22,11 +22,15 @@
 #include "Serialization/ArchiveCrc32.h"
 #include "Engine/HLODProxy.h"
 
+#endif
+
 UWaterBodyHLODBuilder::UWaterBodyHLODBuilder(const FObjectInitializer& ObjectInitializer)
 	: Super(ObjectInitializer)
 {
 }
 
+#if WITH_EDITOR
+
 static double GetPixelToWorldUnit(const double InDistanceZ)
 {
 	// Generate a projection matrix.
