Water: fixed river direction angle unwind logic for angles > 2 PI.

diff --git a/Engine/Plugins/Experimental/Water/Source/Runtime/Private/WaterBodyRiverComponent.cpp b/Engine/Plugins/Experimental/Water/Source/Runtime/Private/WaterBodyRiverComponent.cpp
--- a/Engine/Plugins/Experimental/Water/Source/Runtime/Private/WaterBodyRiverComponent.cpp
+++ b/Engine/Plugins/Experimental/Water/Source/Runtime/Private/WaterBodyRiverComponent.cpp
@@ -152,10 +152,12 @@ static void AddVerticesForRiverSplineStep(
 	FDynamicMeshVertex DilatedFarRight(FVector3f(Pos + OutwardDistance + DilationOffset));
 
 	float FlowDirection = Tangent.HeadingAngle() + FMath::DegreesToRadians(Component->GetRelativeRotation().Yaw);
-	// Convert negative angles into positive angles
-	while (FlowDirection < 0.f)
+
+	// Restrict all angles between [0, 2 PI]. UnwindRadians returns a value between [-Pi, Pi] so we must remap again:
+	FlowDirection = FMath::UnwindRadians(FlowDirection);
+	if (FlowDirection < 0.f)
 	{
-		FlowDirection = TWO_PI + FlowDirection;
+		FlowDirection += TWO_PI;
 	}
 
 	// If negative velocity, inverse the direction and change the velocity back to positive.
