Removed unused function

#rnx

diff --git a/Engine/Plugins/Experimental/Water/Source/Editor/Private/WaterEditorSubsystem.cpp b/Engine/Plugins/Experimental/Water/Source/Editor/Private/WaterEditorSubsystem.cpp
--- a/Engine/Plugins/Experimental/Water/Source/Editor/Private/WaterEditorSubsystem.cpp
+++ b/Engine/Plugins/Experimental/Water/Source/Editor/Private/WaterEditorSubsystem.cpp
@@ -22,71 +22,6 @@
 
 #define LOCTEXT_NAMESPACE "WaterEditorSubsystem"
 
-void UpdateSingleTexture(UTexture2D*& DestTexture, UTextureRenderTarget2D* SrcRenderTarget, UObject* Outer, const TCHAR* TextureName)
-{
-	uint32 TextureFlags = CTF_Default;
-	if (!DestTexture)
-	{
-		DestTexture = SrcRenderTarget->ConstructTexture2D(Outer, TextureName, RF_NoFlags, TextureFlags);
-	}
-
-	const EPixelFormat PixelFormat = SrcRenderTarget->GetFormat();
-	ETextureSourceFormat TextureFormat = TSF_Invalid;
-	switch (PixelFormat)
-	{
-	case PF_B8G8R8A8:
-		TextureFormat = TSF_BGRA8;
-		break;
-	case PF_FloatRGBA:
-		TextureFormat = TSF_RGBA16F;
-		break;
-	}
-
-	bool bTextureModified = false;
-
-	UTextureRenderTarget2D::FTextureChangingDelegate OnTextureChanging;
-	OnTextureChanging.BindLambda([&bTextureModified](UTexture* InTexture)
-	{
-		if (!bTextureModified)
-		{
-			InTexture->Modify();
-			InTexture->PreEditChange(nullptr);
-			bTextureModified = true;
-		}
-	});
-	
-	// Verify if we need to update the destination texture
-	bool bMustUpdateTexture = false;
-
-	// Compare LOD group
-	TEnumAsByte<TextureGroup> TexLODGroup = GetDefault<UWaterEditorSettings>()->TextureGroupForGeneratedTextures;
-	bMustUpdateTexture |= DestTexture->LODGroup != TexLODGroup;
-	
-	// Compare mip gen settings
-	TEnumAsByte<TextureMipGenSettings> TexMipGenSetting = TMGS_NoMipmaps;
-	bMustUpdateTexture |= DestTexture->MipGenSettings != TexMipGenSetting;
-
-	// Compare max texture size
-	int32 TexMaxSize = GetDefault<UWaterEditorSettings>()->MaxWaterVelocityAndHeightTextureSize;
-	bMustUpdateTexture |= DestTexture->MaxTextureSize != TexMaxSize;
-
-	// Update the texture if needed
-	if (bMustUpdateTexture)
-	{
-		OnTextureChanging.Execute(DestTexture);
-	}
-
-	DestTexture->LODGroup = TexLODGroup;
-	DestTexture->MipGenSettings = TexMipGenSetting;
-	DestTexture->MaxTextureSize = TexMaxSize;
-	SrcRenderTarget->UpdateTexture2D(DestTexture, TextureFormat, TextureFlags, nullptr, OnTextureChanging);
-
-	if (bTextureModified)
-	{
-		DestTexture->PostEditChange();
-	}
-}
-
 UWaterEditorSubsystem::UWaterEditorSubsystem()
 {
 	struct FConstructorStatics
