Water: changed default UWaterBodyComponent::CanAffectWaterMesh to return true and implemented it for custom water bodies to return false.

diff --git a/Engine/Plugins/Experimental/Water/Source/Editor/Private/WaterLandscapeBrush.cpp b/Engine/Plugins/Experimental/Water/Source/Editor/Private/WaterLandscapeBrush.cpp
--- a/Engine/Plugins/Experimental/Water/Source/Editor/Private/WaterLandscapeBrush.cpp
+++ b/Engine/Plugins/Experimental/Water/Source/Editor/Private/WaterLandscapeBrush.cpp
@@ -182,7 +182,7 @@ void AWaterLandscapeBrush::OnActorChanged(AActor* Actor, bool bWeightmapSettings
 		}
 
 		// Force rebuild the mesh if a water body actor has been added or removed (islands don't affect the water mesh so it's not necessary for them): 
-		bRebuildWaterMesh = WaterBrushActor->CanAffectWaterMesh();
+		bRebuildWaterMesh = WaterBrushActor->CanEverAffectWaterMesh();
 		bForceUpdateBrush = true;
 	}
 
diff --git a/Engine/Plugins/Experimental/Water/Source/Runtime/Public/WaterBodyActor.h b/Engine/Plugins/Experimental/Water/Source/Runtime/Public/WaterBodyActor.h
--- a/Engine/Plugins/Experimental/Water/Source/Runtime/Public/WaterBodyActor.h
+++ b/Engine/Plugins/Experimental/Water/Source/Runtime/Public/WaterBodyActor.h
@@ -46,7 +46,7 @@ public:
 	//~ Begin IWaterBrushActorInterface interface
 	virtual bool AffectsLandscape() const override { return WaterBodyComponent->AffectsLandscape(); }
 	virtual bool AffectsWaterMesh() const override { return WaterBodyComponent->AffectsWaterMesh(); }
-	virtual bool CanAffectWaterMesh() const override { return WaterBodyComponent->CanAffectWaterMesh(); }
+	virtual bool CanEverAffectWaterMesh() const override { return WaterBodyComponent->CanEverAffectWaterMesh(); }
 
 #if WITH_EDITOR
 	virtual const FWaterCurveSettings& GetWaterCurveSettings() const override { return WaterBodyComponent->GetWaterCurveSettings(); }
diff --git a/Engine/Plugins/Experimental/Water/Source/Runtime/Public/WaterBodyComponent.h b/Engine/Plugins/Experimental/Water/Source/Runtime/Public/WaterBodyComponent.h
--- a/Engine/Plugins/Experimental/Water/Source/Runtime/Public/WaterBodyComponent.h
+++ b/Engine/Plugins/Experimental/Water/Source/Runtime/Public/WaterBodyComponent.h
@@ -93,7 +93,7 @@ public:
 	virtual bool AffectsLandscape() const;
 	virtual bool AffectsWaterMesh() const;
 	virtual bool AffectsWaterInfo() const;
-	virtual bool CanAffectWaterMesh() const { return false; }
+	virtual bool CanEverAffectWaterMesh() const { return true; }
 
 	void UpdateAll(bool bShapeOrPositionChanged);
 
@@ -340,6 +340,9 @@ public:
 	/** Set the navigation area class */
 	void SetNavAreaClass(TSubclassOf<UNavAreaBase> NewWaterNavAreaClass) { WaterNavAreaClass = NewWaterNavAreaClass; }
 
+	
+	UE_DEPRECATED(5.1, "Renamed to CanEverAffectWaterMesh")
+	virtual bool CanAffectWaterMesh() const { return true; }
 protected:
 	//~ Begin UActorComponent interface.
 	virtual bool IsHLODRelevant() const override;
diff --git a/Engine/Plugins/Experimental/Water/Source/Runtime/Public/WaterBodyCustomComponent.h b/Engine/Plugins/Experimental/Water/Source/Runtime/Public/WaterBodyCustomComponent.h
--- a/Engine/Plugins/Experimental/Water/Source/Runtime/Public/WaterBodyCustomComponent.h
+++ b/Engine/Plugins/Experimental/Water/Source/Runtime/Public/WaterBodyCustomComponent.h
@@ -20,6 +20,7 @@ public:
 	virtual EWaterBodyType GetWaterBodyType() const override { return EWaterBodyType::Transition; }
 	virtual TArray<UPrimitiveComponent*> GetCollisionComponents(bool bInOnlyEnabledComponents = true) const override;
 	virtual TArray<UPrimitiveComponent*> GetStandardRenderableComponents() const override;
+	virtual bool CanEverAffectWaterMesh() const { return false; }
 
 protected:
 	/** AWaterBody Interface */
diff --git a/Engine/Plugins/Experimental/Water/Source/Runtime/Public/WaterBodyIslandActor.h b/Engine/Plugins/Experimental/Water/Source/Runtime/Public/WaterBodyIslandActor.h
--- a/Engine/Plugins/Experimental/Water/Source/Runtime/Public/WaterBodyIslandActor.h
+++ b/Engine/Plugins/Experimental/Water/Source/Runtime/Public/WaterBodyIslandActor.h
@@ -25,7 +25,7 @@ public:
 	//~ Begin IWaterBrushActorInterface interface
 	virtual bool AffectsLandscape() const override { return true; }
 	virtual bool AffectsWaterMesh() const override { return false; }
-	virtual bool CanAffectWaterMesh() const override { return false; }
+	virtual bool CanEverAffectWaterMesh() const override { return false; }
 
 #if WITH_EDITOR
 	virtual const FWaterCurveSettings& GetWaterCurveSettings() const { return WaterCurveSettings; }
diff --git a/Engine/Plugins/Experimental/Water/Source/Runtime/Public/WaterBrushActorInterface.h b/Engine/Plugins/Experimental/Water/Source/Runtime/Public/WaterBrushActorInterface.h
--- a/Engine/Plugins/Experimental/Water/Source/Runtime/Public/WaterBrushActorInterface.h
+++ b/Engine/Plugins/Experimental/Water/Source/Runtime/Public/WaterBrushActorInterface.h
@@ -40,7 +40,10 @@ class WATER_API IWaterBrushActorInterface
 	/**
 	 * Returns true if this water actor could potentially be affecting (i.e. being rendered by) a AWaterZoneActor :
 	 */
-	virtual bool CanAffectWaterMesh() const = 0;
+	virtual bool CanEverAffectWaterMesh() const = 0;
+
+	UE_DEPRECATED(5.1, "Renamed to CanEverAffectWaterMesh")
+	virtual bool CanAffectWaterMesh() const { return CanEverAffectWaterMesh(); };
 
 #if WITH_EDITOR
 	/** 
