MGPU:  Water not rendering.  Water Info render pass needs to run across all GPUs.

#rnx

diff --git a/Engine/Plugins/Experimental/Water/Source/Runtime/Private/WaterInfoRendering.cpp b/Engine/Plugins/Experimental/Water/Source/Runtime/Private/WaterInfoRendering.cpp
--- a/Engine/Plugins/Experimental/Water/Source/Runtime/Private/WaterInfoRendering.cpp
+++ b/Engine/Plugins/Experimental/Water/Source/Runtime/Private/WaterInfoRendering.cpp
@@ -527,7 +527,6 @@ static void UpdateWaterInfoRendering_RenderThread(
 		FRDGBuilder GraphBuilder(RHICmdList, RDG_EVENT_NAME("WaterInfoColorRendering"), ERDGBuilderFlags::AllowParallelExecute);
 
 
-
 		// We need to execute the pre-render view extensions before we do any view dependent work.
 		FSceneRenderer::ViewExtensionPreRender_RenderThread(GraphBuilder, DilationRenderer);
 
@@ -654,6 +653,8 @@ static FSceneRenderer* CreateWaterInfoSceneRenderer(const FCreateWaterInfoSceneR
 
 	FSceneView* View = new FSceneView(ViewInitOptions);
 	View->bIsSceneCapture = true;
+	View->GPUMask = FRHIGPUMask::All();
+	View->bOverrideGPUMask = true;
 	View->AntiAliasingMethod = EAntiAliasingMethod::AAM_None;
 	View->SetupAntiAliasingMethod();
 
