Add setting to FBuoyancyData so that relative pontoon locations may be centered around the object's center of mass.

diff --git a/Engine/Plugins/Experimental/Water/Source/Runtime/Private/BuoyancyComponent.cpp b/Engine/Plugins/Experimental/Water/Source/Runtime/Private/BuoyancyComponent.cpp
--- a/Engine/Plugins/Experimental/Water/Source/Runtime/Private/BuoyancyComponent.cpp
+++ b/Engine/Plugins/Experimental/Water/Source/Runtime/Private/BuoyancyComponent.cpp
@@ -377,31 +377,40 @@ void UBuoyancyComponent::ComputePontoonCoefficients()
 	TArray<float>& PontoonCoefficients = ConfiguredPontoonCoefficients.FindOrAdd(PontoonConfiguration);
 	if (PontoonCoefficients.Num() == 0)
 	{
-		TArray<FVector> LocalPontoonLocations;
-		if (!SimulatingComponent)
-		{
-			return;
-		}
-		for (int32 PontoonIndex = 0; PontoonIndex < BuoyancyData.Pontoons.Num(); ++PontoonIndex)
+		if (FBodyInstance* BodyInstance = SimulatingComponent->GetBodyInstance())
 		{
-			const FSphericalPontoon& Pontoon = BuoyancyData.Pontoons[PontoonIndex];
-			if (PontoonConfiguration & (1 << PontoonIndex))
+			TArray<FVector> LocalPontoonLocations;
+			const FVector& LocalCOM = BodyInstance->GetMassSpaceLocal().GetLocation();
+
+			if (!SimulatingComponent)
 			{
-				if (Pontoon.bUseCenterSocket)
-				{
-					const FVector LocalPosition = SimulatingComponent->GetSocketTransform(Pontoon.CenterSocket, ERelativeTransformSpace::RTS_ParentBoneSpace).GetLocation();
-					LocalPontoonLocations.Add(LocalPosition);
-				}
-				else
+				return;
+			}
+
+			for (int32 PontoonIndex = 0; PontoonIndex < BuoyancyData.Pontoons.Num(); ++PontoonIndex)
+			{
+				const FSphericalPontoon& Pontoon = BuoyancyData.Pontoons[PontoonIndex];
+				if (PontoonConfiguration & (1 << PontoonIndex))
 				{
-					LocalPontoonLocations.Add(Pontoon.RelativeLocation);
+					if (Pontoon.bUseCenterSocket)
+					{
+						const FVector LocalPosition = SimulatingComponent->GetSocketTransform(Pontoon.CenterSocket, ERelativeTransformSpace::RTS_ParentBoneSpace).GetLocation();
+						LocalPontoonLocations.Add(LocalPosition);
+					}
+					else
+					{
+						// If using the relative location for the pontoon and the buoyancydata indicates
+						// that we should center on COM, then shift the relative location to be centered.
+						const FVector PontoonRelativeLocation
+							= BuoyancyData.bCenterPontoonsOnCOM
+							? Pontoon.RelativeLocation - LocalCOM
+							: Pontoon.RelativeLocation;
+						LocalPontoonLocations.Add(PontoonRelativeLocation);
+					}
 				}
 			}
-		}
-		PontoonCoefficients.AddZeroed(LocalPontoonLocations.Num());
-		if (FBodyInstance* BodyInstance = SimulatingComponent->GetBodyInstance())
-		{
-			const FVector& LocalCOM = BodyInstance->GetMassSpaceLocal().GetLocation();
+
+			PontoonCoefficients.AddZeroed(LocalPontoonLocations.Num());
 			//Distribute a mass of 1 to each pontoon so that we get a scaling factor based on position relative to CoM
 			FString ErrMsg;
 			bool ComputeSuccess = FSimpleSuspensionHelpers::ComputeSprungMasses(LocalPontoonLocations, LocalCOM, 1.f, PontoonCoefficients, &ErrMsg);
diff --git a/Engine/Plugins/Experimental/Water/Source/Runtime/Public/BuoyancyTypes.h b/Engine/Plugins/Experimental/Water/Source/Runtime/Public/BuoyancyTypes.h
--- a/Engine/Plugins/Experimental/Water/Source/Runtime/Public/BuoyancyTypes.h
+++ b/Engine/Plugins/Experimental/Water/Source/Runtime/Public/BuoyancyTypes.h
@@ -184,6 +184,10 @@ struct FBuoyancyData
 	UPROPERTY(EditDefaultsOnly, BlueprintReadOnly, Category = Buoyancy)
 	TArray<FSphericalPontoon> Pontoons;
 
+	/** If true, center pontoons around center of mass when using relative locations
+		(not used when pontoon locations are specified via sockets) */
+	UPROPERTY(EditDefaultsOnly, BlueprintReadOnly, Category = Buoyancy)
+	bool bCenterPontoonsOnCOM = true;
 
 	/** Increases buoyant force applied on each pontoon. */
 	UPROPERTY(EditDefaultsOnly, Category = Buoyancy)
