Water: allow Z scaling and non-z rotation for water bodies not rendered by the water mesh.

diff --git a/Engine/Plugins/Experimental/Water/Source/Runtime/Private/WaterBodyComponent.cpp b/Engine/Plugins/Experimental/Water/Source/Runtime/Private/WaterBodyComponent.cpp
--- a/Engine/Plugins/Experimental/Water/Source/Runtime/Private/WaterBodyComponent.cpp
+++ b/Engine/Plugins/Experimental/Water/Source/Runtime/Private/WaterBodyComponent.cpp
@@ -936,13 +936,15 @@ void UWaterBodyComponent::OnPostEditChangeProperty(FPropertyChangedEvent& Proper
 	{
 		bShapeOrPositionChanged = true;
 	}
-	else if (PropertyChangedEvent.MemberProperty && PropertyChangedEvent.MemberProperty->GetFName() == FName(TEXT("RelativeScale3D")) &&
-		PropertyName == FName(TEXT("Z")))
+	else if (PropertyChangedEvent.MemberProperty && PropertyChangedEvent.MemberProperty->GetFName() == FName(TEXT("RelativeScale3D")))
 	{
-		// Prevent scaling on the Z axis
-		FVector Scale = GetRelativeScale3D();
-		Scale.Z = 1.f;
-		SetRelativeScale3D(Scale);
+		// All water bodies which can ever be rendered by the water mesh shouldn't have a z-scale.
+		if (CanEverAffectWaterMesh())
+		{
+			FVector Scale = GetRelativeScale3D();
+			Scale.Z = 1.f;
+			SetRelativeScale3D(Scale);
+		}
 	}
 }
 
@@ -1323,16 +1325,19 @@ void UWaterBodyComponent::OnComponentDestroyed(bool bDestroyingHierarchy)
 
 bool UWaterBodyComponent::MoveComponentImpl(const FVector& Delta, const FQuat& NewRotation, bool bSweep, FHitResult* Hit, EMoveComponentFlags MoveFlags, ETeleportType Teleport)
 {
-	// Prevent Z-Scale
-	FVector Scale = GetRelativeScale3D();
-	Scale.Z = 1.f;
-	SetRelativeScale3D(Scale);
-
-	// Restrict rotation to the Z-axis only
 	FQuat CorrectedRotation = NewRotation;
-	CorrectedRotation.X = 0.f;
-	CorrectedRotation.Y = 0.f;
 
+	// All water bodies which can ever be rendered by the water mesh shouldn't have a z-scale or non-z rotation
+	if (CanEverAffectWaterMesh())
+	{
+		FVector Scale = GetRelativeScale3D();
+		Scale.Z = 1.f;
+		SetRelativeScale3D(Scale);
+
+		// Restrict rotation to the Z-axis only
+		CorrectedRotation.X = 0.f;
+		CorrectedRotation.Y = 0.f;
+	}
 	return Super::MoveComponentImpl(Delta, CorrectedRotation, bSweep, Hit, MoveFlags, Teleport);
 }
 
