Water: Defined filter strings for water spline metadata customization rows in order for them to be searchable.

diff --git a/Engine/Plugins/Experimental/Water/Source/Editor/Private/WaterSplineMetadataDetails.cpp b/Engine/Plugins/Experimental/Water/Source/Editor/Private/WaterSplineMetadataDetails.cpp
--- a/Engine/Plugins/Experimental/Water/Source/Editor/Private/WaterSplineMetadataDetails.cpp
+++ b/Engine/Plugins/Experimental/Water/Source/Editor/Private/WaterSplineMetadataDetails.cpp
@@ -218,6 +218,7 @@ void FWaterSplineMetadataDetails::GenerateChildContent(IDetailGroup& DetailGroup
 {
 	DetailGroup.AddWidgetRow()
 		.RowTag("Depth")
+		.FilterString(LOCTEXT("Depth", "Depth"))
 		.Visibility(TAttribute<EVisibility>(this, &FWaterSplineMetadataDetails::IsDepthVisible))
 		.NameContent()
 		.HAlign(HAlign_Left)
@@ -250,6 +251,7 @@ void FWaterSplineMetadataDetails::GenerateChildContent(IDetailGroup& DetailGroup
 
 	DetailGroup.AddWidgetRow()
 		.RowTag("RiverWidth")
+		.FilterString(LOCTEXT("RiverWidth", "River Width"))
 		.Visibility(TAttribute<EVisibility>(this, &FWaterSplineMetadataDetails::IsRiverWidthVisible))
 		.NameContent()
 		.HAlign(HAlign_Left)
@@ -282,6 +284,7 @@ void FWaterSplineMetadataDetails::GenerateChildContent(IDetailGroup& DetailGroup
 
 	DetailGroup.AddWidgetRow()
 		.RowTag("Velocity")
+		.FilterString(LOCTEXT("Velocity", "Velocity"))
 		.Visibility(TAttribute<EVisibility>(this, &FWaterSplineMetadataDetails::IsEnabled))
 		.NameContent()
 		.HAlign(HAlign_Left)
@@ -312,6 +315,7 @@ void FWaterSplineMetadataDetails::GenerateChildContent(IDetailGroup& DetailGroup
 
 	DetailGroup.AddWidgetRow()
 		.RowTag("AudioIntensity")
+		.FilterString(LOCTEXT("AudioIntensity", "Audio Intensity"))
 		.Visibility(TAttribute<EVisibility>(this, &FWaterSplineMetadataDetails::IsEnabled))
 		.NameContent()
 		.HAlign(HAlign_Left)
