Water: Fixed issue with lingering river mesh components after undo

diff --git a/Engine/Plugins/Experimental/Water/Source/Runtime/Private/WaterBodyComponent.cpp b/Engine/Plugins/Experimental/Water/Source/Runtime/Private/WaterBodyComponent.cpp
--- a/Engine/Plugins/Experimental/Water/Source/Runtime/Private/WaterBodyComponent.cpp
+++ b/Engine/Plugins/Experimental/Water/Source/Runtime/Private/WaterBodyComponent.cpp
@@ -898,12 +898,17 @@ void UWaterBodyComponent::PostEditUndo()
 {
 	Super::PostEditUndo();
 
-	OnWaterBodyChanged(/*bShapeOrPositionChanged*/true, /*bWeightmapSettingsChanged*/true);
+	// Since this component may become unregistered/deleted if we are undoing the creation of a water body.
+	// Ensure we only trigger updates if this component is registered.
+	if (IsRegistered())
+	{
+		OnWaterBodyChanged(/*bShapeOrPositionChanged*/true, /*bWeightmapSettingsChanged*/true);
 
-	// On undo, when PostEditChangeProperty is called, PropertyChangedEvent is fake so we need to register to the new object here :
-	RegisterOnUpdateWavesData(GetWaterWaves(), /*bRegister = */true);
+		// On undo, when PostEditChangeProperty is called, PropertyChangedEvent is fake so we need to register to the new object here :
+		RegisterOnUpdateWavesData(GetWaterWaves(), /*bRegister = */true);
 
-	RequestGPUWaveDataUpdate();
+		RequestGPUWaveDataUpdate();
+	}
 }
 
 void UWaterBodyComponent::PostEditImport()
