Water: fixed error icon appearing in blueprint preview viewport.

diff --git a/Engine/Plugins/Experimental/Water/Source/Runtime/Private/WaterBodyComponent.cpp b/Engine/Plugins/Experimental/Water/Source/Runtime/Private/WaterBodyComponent.cpp
--- a/Engine/Plugins/Experimental/Water/Source/Runtime/Private/WaterBodyComponent.cpp
+++ b/Engine/Plugins/Experimental/Water/Source/Runtime/Private/WaterBodyComponent.cpp
@@ -956,7 +956,9 @@ void UWaterBodyComponent::OnPostEditChangeProperty(FPropertyChangedEvent& Proper
 TArray<TSharedRef<FTokenizedMessage>> UWaterBodyComponent::CheckWaterBodyStatus() const
 {
 	TArray<TSharedRef<FTokenizedMessage>> Result;
-	if (!IsTemplate())
+
+	const UWorld* World = GetWorld();
+	if (!IsTemplate() && World && World->WorldType != EWorldType::EditorPreview)
 	{
 		if (AffectsWaterMesh() && (GetWaterZone() == nullptr))
 		{
