Water: fixed issue where ocean mesh calculation wasn't in local space and caused issues away from the origin.

diff --git a/Engine/Plugins/Experimental/Water/Source/Runtime/Private/WaterBodyOceanComponent.cpp b/Engine/Plugins/Experimental/Water/Source/Runtime/Private/WaterBodyOceanComponent.cpp
--- a/Engine/Plugins/Experimental/Water/Source/Runtime/Private/WaterBodyOceanComponent.cpp
+++ b/Engine/Plugins/Experimental/Water/Source/Runtime/Private/WaterBodyOceanComponent.cpp
@@ -154,7 +154,7 @@ void UWaterBodyOceanComponent::GenerateWaterBodyMesh()
 	}
 
 	FVector OceanLocation = GetComponentLocation();
-	FPolygon2d OceanBoundingPolygon = FPolygon2d::MakeRectangle(FVector2d(OceanLocation.X, OceanLocation.Y), VisualExtents.X, VisualExtents.Y);
+	FPolygon2d OceanBoundingPolygon = FPolygon2d::MakeRectangle(FVector2d(0, 0), VisualExtents.X, VisualExtents.Y);
 
 	FConstrainedDelaunay2d Triangulation;
 	Triangulation.FillRule = FConstrainedDelaunay2d::EFillRule::Positive;
