Water: Allow water bodies with null materials to cut holes in the water mesh quad tree.

diff --git a/Engine/Plugins/Experimental/Water/Source/Runtime/Private/WaterBodyComponent.cpp b/Engine/Plugins/Experimental/Water/Source/Runtime/Private/WaterBodyComponent.cpp
--- a/Engine/Plugins/Experimental/Water/Source/Runtime/Private/WaterBodyComponent.cpp
+++ b/Engine/Plugins/Experimental/Water/Source/Runtime/Private/WaterBodyComponent.cpp
@@ -286,7 +286,8 @@ void UWaterBodyComponent::SetWaterAndUnderWaterPostProcessMaterial(UMaterialInte
 
 bool UWaterBodyComponent::ShouldGenerateWaterMeshTile() const
 {
-	return ((GetWaterBodyType() != EWaterBodyType::Transition)
+	return bAlwaysGenerateWaterMeshTiles
+		|| ((GetWaterBodyType() != EWaterBodyType::Transition)
 		&& (GetWaterMeshOverride() == nullptr)
 		&& (GetWaterMaterial() != nullptr));
 }
diff --git a/Engine/Plugins/Experimental/Water/Source/Runtime/Private/WaterMeshComponent.cpp b/Engine/Plugins/Experimental/Water/Source/Runtime/Private/WaterMeshComponent.cpp
--- a/Engine/Plugins/Experimental/Water/Source/Runtime/Private/WaterMeshComponent.cpp
+++ b/Engine/Plugins/Experimental/Water/Source/Runtime/Private/WaterMeshComponent.cpp
@@ -260,19 +260,23 @@ void UWaterMeshComponent::RebuildWaterMesh(float InTileSize, const FIntPoint& In
 		// Assign material instance
 		UMaterialInstanceDynamic* WaterMaterial = WaterBodyComponent->GetWaterMaterialInstance();
 		RenderData.Material = WaterMaterial;
-		if (!IsMaterialUsedWithWater(RenderData.Material))
-		{
-			RenderData.Material = UMaterial::GetDefaultMaterial(MD_Surface);
-		}
-		else
+
+		if (RenderData.Material)
 		{
-			// Add ocean height as a scalar parameter
-			WaterBodyComponent->SetDynamicParametersOnMID(WaterMaterial);
+			if (!IsMaterialUsedWithWater(RenderData.Material))
+			{
+				RenderData.Material = UMaterial::GetDefaultMaterial(MD_Surface);
+			}
+			else
+			{
+				// Add ocean height as a scalar parameter
+				WaterBodyComponent->SetDynamicParametersOnMID(WaterMaterial);
+			}
+
+			// Add material so that the component keeps track of all potential materials used
+			UsedMaterials.Add(RenderData.Material);
 		}
 
-		// Add material so that the component keeps track of all potential materials used
-		UsedMaterials.Add(RenderData.Material);
-		
 		// Min and max user defined priority range. (Input also clamped on OverlapMaterialPriority in AWaterBody)
 		const int32 MinPriority = -8192;
 		const int32 MaxPriority = 8191;
diff --git a/Engine/Plugins/Experimental/Water/Source/Runtime/Private/WaterQuadTree.cpp b/Engine/Plugins/Experimental/Water/Source/Runtime/Private/WaterQuadTree.cpp
--- a/Engine/Plugins/Experimental/Water/Source/Runtime/Private/WaterQuadTree.cpp
+++ b/Engine/Plugins/Experimental/Water/Source/Runtime/Private/WaterQuadTree.cpp
@@ -287,6 +287,9 @@ void FWaterQuadTree::FNode::AddNodes(FNodeData& InNodeData, const FBox& InMeshBo
 	if (InWaterBody.Priority >= ThisWaterBody.Priority)
 	{
 		WaterBodyIndex = InWaterBodyIndex;
+
+		// Cache whether or not this node has a material
+		HasMaterial = InNodeData.WaterBodyRenderData[WaterBodyIndex].Material != nullptr;
 	}
 
 	// Reset the flags before going through the children. These flags will be turned off by recursion if the state changes
@@ -473,6 +476,32 @@ void FWaterQuadTree::Unlock(bool bPruneRedundantNodes)
 
 	if (bPruneRedundantNodes)
 	{
+		auto SwapRemove = [&](int32 NodeIndex, int32 EndIndex)
+		{
+			if (NodeIndex != EndIndex)
+			{
+				// Swap to back. All the children of this node would have already been removed (or didn't exist to begin with), so don't care about those
+				NodeData.Nodes.SwapMemory(NodeIndex, EndIndex);
+
+				// Patch up the newly moved good node (parent and children)
+				FNode& MovedNode = NodeData.Nodes[NodeIndex];
+				FNode& MovedNodeParent = NodeData.Nodes[MovedNode.ParentIndex];
+
+				for (int32 i = 0; i < 4; i++)
+				{
+					if (MovedNode.Children[i] > 0)
+					{
+						NodeData.Nodes[MovedNode.Children[i]].ParentIndex = NodeIndex;
+					}
+
+					if (MovedNodeParent.Children[i] == EndIndex)
+					{
+						MovedNodeParent.Children[i] = NodeIndex;
+					}
+				}
+			}
+		};
+
 		// Remove redundant nodes
 		// Remove from the back, since all removalbe children are further back than their parent in the node list and we want to remove bottom-up
 		int32 EndIndex = NodeData.Nodes.Num() - 1;
@@ -486,29 +515,23 @@ void FWaterQuadTree::Unlock(bool bPruneRedundantNodes)
 				// Delete all children (not strictly necessary, but now we don't leave any dangling/incorrect child pointers around)
 				FMemory::Memzero(&ParentNode.Children, sizeof(uint32) * 4);
 
-				if (NodeIndex != EndIndex)
-				{
-					// Swap to back. All the children of this node would have already been removed (or didn't exist to begin with), so don't care about those
-					NodeData.Nodes.SwapMemory(NodeIndex, EndIndex);
-
-					// Patch up the newly moved good node (parent and children)
-					FNode& MovedNode = NodeData.Nodes[NodeIndex];
-					FNode& MovedNodeParent = NodeData.Nodes[MovedNode.ParentIndex];
+				SwapRemove(NodeIndex, EndIndex);
 
-					for (int32 i = 0; i < 4; i++)
+				// Move back one step down
+				EndIndex--;
+			}
+			else if (!NodeData.Nodes[NodeIndex].HasMaterial && NodeData.Nodes[NodeIndex].HasCompleteSubtree && NodeData.Nodes[NodeIndex].IsSubtreeSameWaterBody)
+			{
+				for (int32 i = 0; i < 4; i++)
+				{
+					if (ParentNode.Children[i] == NodeIndex)
 					{
-						if (MovedNode.Children[i] > 0)
-						{
-							NodeData.Nodes[MovedNode.Children[i]].ParentIndex = NodeIndex;
-						}
-
-						if (MovedNodeParent.Children[i] == EndIndex)
-						{
-							MovedNodeParent.Children[i] = NodeIndex;
-						}
+						ParentNode.Children[i] = 0;
 					}
 				}
 
+				SwapRemove(NodeIndex, EndIndex);
+
 				// Move back one step down
 				EndIndex--;
 			}
diff --git a/Engine/Plugins/Experimental/Water/Source/Runtime/Public/WaterBodyComponent.h b/Engine/Plugins/Experimental/Water/Source/Runtime/Public/WaterBodyComponent.h
--- a/Engine/Plugins/Experimental/Water/Source/Runtime/Public/WaterBodyComponent.h
+++ b/Engine/Plugins/Experimental/Water/Source/Runtime/Public/WaterBodyComponent.h
@@ -527,6 +527,13 @@ protected:
 	UPROPERTY(Category = Rendering, EditAnywhere, BlueprintReadOnly)
 	UStaticMesh* WaterMeshOverride;
 
+	/** 
+	 * When this is set to true, the water mesh will always generate tiles for this water body. 
+	 * For example, this can be useful to generate water tiles even when the water material is invalid, for the case where "empty" water tiles are actually desirable.
+	 */
+	UPROPERTY(Category = Rendering, EditAnywhere, BlueprintReadOnly)
+	bool bAlwaysGenerateWaterMeshTiles = false;
+
 	/** Higher number is higher priority. If two water bodies overlap and they don't have a transition material specified, this will be used to determine which water body to use the material from. Valid range is -8192 to 8191 */
 	UPROPERTY(Category = Rendering, EditAnywhere, BlueprintReadOnly, meta = (ClampMin = "-8192", ClampMax = "8191"))
 	int32 OverlapMaterialPriority = 0;
diff --git a/Engine/Plugins/Experimental/Water/Source/Runtime/Public/WaterQuadTree.h b/Engine/Plugins/Experimental/Water/Source/Runtime/Public/WaterQuadTree.h
--- a/Engine/Plugins/Experimental/Water/Source/Runtime/Public/WaterQuadTree.h
+++ b/Engine/Plugins/Experimental/Water/Source/Runtime/Public/WaterQuadTree.h
@@ -197,7 +197,7 @@ private:
 	struct FNode
 	{
 
-		FNode() : WaterBodyIndex(0), TransitionWaterBodyIndex(0), ParentIndex(INVALID_PARENT), HasCompleteSubtree(1), IsSubtreeSameWaterBody(1) {}
+		FNode() : WaterBodyIndex(0), TransitionWaterBodyIndex(0), ParentIndex(INVALID_PARENT), HasCompleteSubtree(1), IsSubtreeSameWaterBody(1), HasMaterial(0) {}
 
 		/** If this node is allowed to be rendered, it means it can be rendered in place of all leaf nodes in its subtree. */
 		bool CanRender(int32 InDensityLevel, int32 InForceCollapseDensityLevel, const FWaterBodyRenderData& InWaterBodyRenderData) const;
@@ -241,7 +241,10 @@ private:
 		/** If all descendant nodes are from the same waterbody. We can safely collapse this even if HasCompleteSubtree is false */
 		uint32 IsSubtreeSameWaterBody : 1;
 
-		// 2 spare bits here in the bit field with ParentIndex
+		/** Cached value to avoid having to visit this node's FWaterBodyRenderData */
+		uint32 HasMaterial : 1;
+
+		// 1 spare bits here in the bit field with ParentIndex
 
 		/** Children, 0 means invalid */
 		uint32 Children[4] = { 0, 0, 0, 0 };
