Water: changed exclusion volume to exclude all by default.

diff --git a/Engine/Plugins/Experimental/Water/Source/Runtime/Private/WaterBodyExclusionVolume.cpp b/Engine/Plugins/Experimental/Water/Source/Runtime/Private/WaterBodyExclusionVolume.cpp
--- a/Engine/Plugins/Experimental/Water/Source/Runtime/Private/WaterBodyExclusionVolume.cpp
+++ b/Engine/Plugins/Experimental/Water/Source/Runtime/Private/WaterBodyExclusionVolume.cpp
@@ -50,7 +50,7 @@ void AWaterBodyExclusionVolume::UpdateOverlappingWaterBodies()
 	for (const FOverlapResult& Result : Overlaps)
 	{
 		AWaterBody* WaterBody = Result.OverlapObjectHandle.FetchActor<AWaterBody>();
-		if (WaterBody && (bIgnoreAllOverlappingWaterBodies || WaterBodiesToIgnore.Contains(WaterBody)))
+		if (WaterBody && (bExcludeAllOverlappingWaterBodies || WaterBodiesToExclude.Contains(WaterBody)))
 		{
 			UWaterBodyComponent* WaterBodyComponent = WaterBody->GetWaterBodyComponent();
 			NewOverlappingBodies.Add(WaterBodyComponent);
@@ -89,17 +89,28 @@ void AWaterBodyExclusionVolume::PostLoad()
 {
 	Super::PostLoad();
 
-#if WITH_EDITOR
+#if WITH_EDITORONLY_DATA
 	// Perform data deprecation: 
 	if (GetLinkerCustomVersion(FFortniteMainBranchObjectVersion::GUID) < FFortniteMainBranchObjectVersion::SupportMultipleWaterBodiesPerExclusionVolume)
 	{
 		if (WaterBodyToIgnore_DEPRECATED != nullptr)
 		{
-			WaterBodiesToIgnore.Add(WaterBodyToIgnore_DEPRECATED);
+			WaterBodiesToIgnore_DEPRECATED.Add(WaterBodyToIgnore_DEPRECATED);
 			WaterBodyToIgnore_DEPRECATED = nullptr;
 		}
 	}
-#endif // WITH_EDITOR
+
+	if (GetLinkerCustomVersion(FFortniteMainBranchObjectVersion::GUID) < FFortniteMainBranchObjectVersion::WaterExclusionVolumeExcludeAllDefault)
+	{
+		WaterBodiesToExclude = MoveTemp(WaterBodiesToIgnore_DEPRECATED);
+
+		// If the existing actor had selected specific water bodies to exclude then we set the new value to false by default.
+		if (!bIgnoreAllOverlappingWaterBodies_DEPRECATED && WaterBodiesToExclude.Num() > 0)
+		{
+			bExcludeAllOverlappingWaterBodies = false;
+		}
+	}
+#endif // WITH_EDITORONLY_DATA
 
 	UpdateOverlappingWaterBodies();
 }
diff --git a/Engine/Plugins/Experimental/Water/Source/Runtime/Public/WaterBodyComponent.h b/Engine/Plugins/Experimental/Water/Source/Runtime/Public/WaterBodyComponent.h
--- a/Engine/Plugins/Experimental/Water/Source/Runtime/Public/WaterBodyComponent.h
+++ b/Engine/Plugins/Experimental/Water/Source/Runtime/Public/WaterBodyComponent.h
@@ -551,10 +551,10 @@ protected:
 	UMaterialInstanceDynamic* WaterInfoMID;
 
 	/** Islands in this water body*/
-	UPROPERTY(Category = Water, EditInstanceOnly, AdvancedDisplay)
+	UPROPERTY(Category = Water, VisibleAnywhere, AdvancedDisplay)
 	TArray<TLazyObjectPtr<AWaterBodyIsland>> Islands;
 
-	UPROPERTY(Category = Water, EditInstanceOnly, AdvancedDisplay)
+	UPROPERTY(Category = Water, VisibleAnywhere, AdvancedDisplay)
 	TArray<TLazyObjectPtr<AWaterBodyExclusionVolume>> ExclusionVolumes;
 
 	UPROPERTY(Transient)
diff --git a/Engine/Plugins/Experimental/Water/Source/Runtime/Public/WaterBodyExclusionVolume.h b/Engine/Plugins/Experimental/Water/Source/Runtime/Public/WaterBodyExclusionVolume.h
--- a/Engine/Plugins/Experimental/Water/Source/Runtime/Public/WaterBodyExclusionVolume.h
+++ b/Engine/Plugins/Experimental/Water/Source/Runtime/Public/WaterBodyExclusionVolume.h
@@ -38,13 +38,19 @@ protected:
 public:
 	/** If checked, all water bodies overlapping with this exclusion volumes will be affected. */
 	UPROPERTY(EditAnywhere, Category = Water)
-	bool bIgnoreAllOverlappingWaterBodies = false;
+	bool bExcludeAllOverlappingWaterBodies = true;
 
 	/** List of water bodies that will be affected by this exclusion volume */
-	UPROPERTY(EditInstanceOnly, Category = Water, meta = (EditCondition = "!bIgnoreAllOverlappingWaterBodies"))
-	TArray<AWaterBody*> WaterBodiesToIgnore;
+	UPROPERTY(EditInstanceOnly, Category = Water, meta = (EditCondition = "!bExcludeAllOverlappingWaterBodies"))
+	TArray<AWaterBody*> WaterBodiesToExclude;
 
 #if WITH_EDITORONLY_DATA
+	UPROPERTY(meta = (DeprecationMessage = "Property renamed to bExcludeAllOverlapping"))
+	bool bIgnoreAllOverlappingWaterBodies_DEPRECATED = false;
+
+	UPROPERTY(meta = (DeprecationMessage = "Property renamed to WaterBodiesToExclude"))
+	TArray<AWaterBody*> WaterBodiesToIgnore_DEPRECATED;
+
 	UPROPERTY(meta = (DeprecatedProperty))
 	AWaterBody* WaterBodyToIgnore_DEPRECATED;
 
diff --git a/Engine/Source/Runtime/Core/Public/UObject/FortniteMainBranchObjectVersion.h b/Engine/Source/Runtime/Core/Public/UObject/FortniteMainBranchObjectVersion.h
--- a/Engine/Source/Runtime/Core/Public/UObject/FortniteMainBranchObjectVersion.h
+++ b/Engine/Source/Runtime/Core/Public/UObject/FortniteMainBranchObjectVersion.h
@@ -232,6 +232,9 @@ struct CORE_API FFortniteMainBranchObjectVersion
 		// Moved parameters affecting Skeleton pose rendering from the PoseWatch class to the PoseWatchPoseElement class.
 		PoseWatchMigrateSkeletonDrawParametersToPoseElement,
 
+		// Reset default value for Water exclusion volumes to make them more intuitive and support the "it just works" philosophy.
+		WaterExclusionVolumeExcludeAllDefault,
+
 		// -----<new versions can be added above this line>-------------------------------------------------
 		VersionPlusOne,
 		LatestVersion = VersionPlusOne - 1
