Clamped Falloff width to a value above 0 to prevent a div by 0  in the water brush materials.

diff --git a/Engine/Plugins/Experimental/Water/Source/Editor/Private/WaterBrushManager.cpp b/Engine/Plugins/Experimental/Water/Source/Editor/Private/WaterBrushManager.cpp
--- a/Engine/Plugins/Experimental/Water/Source/Editor/Private/WaterBrushManager.cpp
+++ b/Engine/Plugins/Experimental/Water/Source/Editor/Private/WaterBrushManager.cpp
@@ -888,10 +888,12 @@ void AWaterBrushManager::FalloffAndBlendMode(const FBrushActorRenderContext& Bru
 	float FalloffTangentValue = FMath::Clamp(FalloffSettings.FalloffAngle, 1.0f, 89.0f);
 	FalloffTangentValue = FMath::Tan(FMath::DegreesToRadians(FalloffTangentValue));
 
+	const float EdgeWidth = FMath::Max(FalloffSettings.FalloffWidth, 0.1f);
+
 	BrushActorRenderContext.MID->SetScalarParameterValue(FName(TEXT("FalloffTangent")), FalloffTangentValue);
 	BrushActorRenderContext.MID->SetScalarParameterValue(FName(TEXT("UseMeshDepth")), 1.0f);
 	BrushActorRenderContext.MID->SetScalarParameterValue(FName(TEXT("CapShapeInterior")), 1.0f);
-	BrushActorRenderContext.MID->SetScalarParameterValue(FName(TEXT("EdgeWidth")), FalloffSettings.FalloffWidth);
+	BrushActorRenderContext.MID->SetScalarParameterValue(FName(TEXT("EdgeWidth")), EdgeWidth);
 
 	float EdgeCenterOffsetValue = FalloffSettings.EdgeOffset - (SplineMeshExtension / 2.0f);
 	BrushActorRenderContext.MID->SetScalarParameterValue(FName(TEXT("EdgeCenterOffset")), EdgeCenterOffsetValue);
@@ -936,8 +938,9 @@ void AWaterBrushManager::ApplyWeightmapSettings(const FBrushRenderContext& Brush
 	check(::IsValid(WeightmapMID));
 	check(::IsValid(BrushActorRenderContext.CacheContainer));
 
+	const float GradientWidth = FMath::Max(WMSettings.FalloffWidth, 0.1f);
 	WeightmapMID->SetTextureParameterValue(FName(TEXT("CachedDistanceFieldHeight")), BrushActorRenderContext.CacheContainer->Cache.CacheRenderTarget);
-	WeightmapMID->SetScalarParameterValue(FName(TEXT("GradientWidth")), WMSettings.FalloffWidth);
+	WeightmapMID->SetScalarParameterValue(FName(TEXT("GradientWidth")), GradientWidth);
 
 	float EdgeOffsetValue = WMSettings.EdgeOffset - (SplineMeshExtension / 2.0f);
 	WeightmapMID->SetScalarParameterValue(FName(TEXT("EdgeOffset")), EdgeOffsetValue);
diff --git a/Engine/Plugins/Experimental/Water/Source/Runtime/Public/WaterBodyWeightmapSettings.h b/Engine/Plugins/Experimental/Water/Source/Runtime/Public/WaterBodyWeightmapSettings.h
--- a/Engine/Plugins/Experimental/Water/Source/Runtime/Public/WaterBodyWeightmapSettings.h
+++ b/Engine/Plugins/Experimental/Water/Source/Runtime/Public/WaterBodyWeightmapSettings.h
@@ -24,7 +24,7 @@ struct FWaterBodyWeightmapSettings
 		, FinalOpacity(1.0f)
 	{}
 
-	UPROPERTY(EditAnywhere, BlueprintReadWrite, Category= WaterBodyWeightmapSettings)
+	UPROPERTY(EditAnywhere, BlueprintReadWrite, Category= WaterBodyWeightmapSettings, meta = (ClampMin = "0.1"))
 	float FalloffWidth;
 
 	UPROPERTY(EditAnywhere, BlueprintReadWrite, Category = WaterBodyWeightmapSettings)
diff --git a/Engine/Plugins/Experimental/Water/Source/Runtime/Public/WaterFalloffSettings.h b/Engine/Plugins/Experimental/Water/Source/Runtime/Public/WaterFalloffSettings.h
--- a/Engine/Plugins/Experimental/Water/Source/Runtime/Public/WaterFalloffSettings.h
+++ b/Engine/Plugins/Experimental/Water/Source/Runtime/Public/WaterFalloffSettings.h
@@ -25,7 +25,7 @@ struct FWaterFalloffSettings
 	UPROPERTY(EditAnywhere, BlueprintReadWrite, Category = FalloffSettings)
 	float FalloffAngle;
 
-	UPROPERTY(EditAnywhere, BlueprintReadWrite, Category = FalloffSettings)
+	UPROPERTY(EditAnywhere, BlueprintReadWrite, Category = FalloffSettings, meta = (ClampMin = "0.1"))
 	float FalloffWidth;
 
 	UPROPERTY(EditAnywhere, BlueprintReadWrite, Category = FalloffSettings)
