ObjectPtr upgrade for engine plugins used from UnrealEditor -allmodules.  Automated upgrade process.

diff --git a/Engine/Plugins/Experimental/CommonUI/Source/CommonUI/Public/CommonUIRichTextData.h b/Engine/Plugins/Experimental/CommonUI/Source/CommonUI/Public/CommonUIRichTextData.h
--- a/Engine/Plugins/Experimental/CommonUI/Source/CommonUI/Public/CommonUIRichTextData.h
+++ b/Engine/Plugins/Experimental/CommonUI/Source/CommonUI/Public/CommonUIRichTextData.h
@@ -40,5 +40,5 @@ public:
 
 private:
 	UPROPERTY(EditDefaultsOnly, Category = "Inline Icons", meta = (RowType = "/Script/CommonUI.RichTextIconData"))
-	UDataTable* InlineIconSet;
+	TObjectPtr<UDataTable> InlineIconSet;
 };
\ No newline at end of file
diff --git a/Engine/Plugins/Experimental/CommonUI/Source/CommonUI/Public/CommonUISettings.h b/Engine/Plugins/Experimental/CommonUI/Source/CommonUI/Public/CommonUISettings.h
--- a/Engine/Plugins/Experimental/CommonUI/Source/CommonUI/Public/CommonUISettings.h
+++ b/Engine/Plugins/Experimental/CommonUI/Source/CommonUI/Public/CommonUISettings.h
@@ -82,14 +82,14 @@ private:
 	FGameplayTagContainer PlatformTraitContainer;
 
 	UPROPERTY(Transient)
-	UObject* DefaultImageResourceObjectInstance;
+	TObjectPtr<UObject> DefaultImageResourceObjectInstance;
 
 	UPROPERTY(Transient)
-	UMaterialInterface* DefaultThrobberMaterialInstance;
+	TObjectPtr<UMaterialInterface> DefaultThrobberMaterialInstance;
 
 	UPROPERTY(Transient)
 	FSlateBrush DefaultThrobberBrush;
 
 	UPROPERTY(Transient)
-	UCommonUIRichTextData* RichTextDataInstance;
+	TObjectPtr<UCommonUIRichTextData> RichTextDataInstance;
 };
\ No newline at end of file
diff --git a/Engine/Plugins/Experimental/CommonUI/Source/CommonUI/Public/CommonVideoPlayer.h b/Engine/Plugins/Experimental/CommonUI/Source/CommonUI/Public/CommonVideoPlayer.h
--- a/Engine/Plugins/Experimental/CommonUI/Source/CommonUI/Public/CommonVideoPlayer.h
+++ b/Engine/Plugins/Experimental/CommonUI/Source/CommonUI/Public/CommonVideoPlayer.h
@@ -66,19 +66,19 @@ private:
 
 private:
 	UPROPERTY(EditAnywhere, Category = VideoPlayer)
-	UMediaSource* Video;
+	TObjectPtr<UMediaSource> Video;
 
 	UPROPERTY(Transient)
-	UMediaPlayer* MediaPlayer;
+	TObjectPtr<UMediaPlayer> MediaPlayer;
 
 	UPROPERTY(Transient)
-	UMediaTexture* MediaTexture;
+	TObjectPtr<UMediaTexture> MediaTexture;
 
 	UPROPERTY(Transient)
-	UMaterial* VideoMaterial;
+	TObjectPtr<UMaterial> VideoMaterial;
 
 	UPROPERTY(Transient)
-	UMediaSoundComponent* SoundComponent;
+	TObjectPtr<UMediaSoundComponent> SoundComponent;
 
 	UPROPERTY(Transient)
 	FSlateBrush VideoBrush;
diff --git a/Engine/Plugins/Experimental/CommonUI/Source/CommonUI/Public/CommonWidgetCarouselNavBar.h b/Engine/Plugins/Experimental/CommonUI/Source/CommonUI/Public/CommonWidgetCarouselNavBar.h
--- a/Engine/Plugins/Experimental/CommonUI/Source/CommonUI/Public/CommonWidgetCarouselNavBar.h
+++ b/Engine/Plugins/Experimental/CommonUI/Source/CommonUI/Public/CommonWidgetCarouselNavBar.h
@@ -60,11 +60,11 @@ protected:
 	TSharedPtr<SHorizontalBox> MyContainer;
 	
 	UPROPERTY()
-	UCommonWidgetCarousel* LinkedCarousel;
+	TObjectPtr<UCommonWidgetCarousel> LinkedCarousel;
 
 	UPROPERTY()
-	UCommonButtonGroupBase* ButtonGroup;
+	TObjectPtr<UCommonButtonGroupBase> ButtonGroup;
 
 	UPROPERTY()
-	TArray<UCommonButtonBase*> Buttons;
+	TArray<TObjectPtr<UCommonButtonBase>> Buttons;
 };
diff --git a/Engine/Plugins/Experimental/CommonUI/Source/CommonUI/Public/Input/CommonBoundActionButton.h b/Engine/Plugins/Experimental/CommonUI/Source/CommonUI/Public/Input/CommonBoundActionButton.h
--- a/Engine/Plugins/Experimental/CommonUI/Source/CommonUI/Public/Input/CommonBoundActionButton.h
+++ b/Engine/Plugins/Experimental/CommonUI/Source/CommonUI/Public/Input/CommonBoundActionButton.h
@@ -28,7 +28,7 @@ protected:
 
 protected:
 	UPROPERTY(BlueprintReadOnly, meta = (BindWidget), Category = "Text Block")
-	UCommonTextBlock* Text_ActionName;
+	TObjectPtr<UCommonTextBlock> Text_ActionName;
 
 private:
 	FUIActionBindingHandle BindingHandle;
diff --git a/Engine/Plugins/Experimental/CommonUI/Source/CommonUI/Public/StyleSheet/CommonStyleSheet.h b/Engine/Plugins/Experimental/CommonUI/Source/CommonUI/Public/StyleSheet/CommonStyleSheet.h
--- a/Engine/Plugins/Experimental/CommonUI/Source/CommonUI/Public/StyleSheet/CommonStyleSheet.h
+++ b/Engine/Plugins/Experimental/CommonUI/Source/CommonUI/Public/StyleSheet/CommonStyleSheet.h
@@ -36,7 +36,7 @@ private:
 	TArray<TObjectPtr<UCommonStyleSheetTypeBase>> Properties;
 
 	UPROPERTY(EditDefaultsOnly, Category = "Import")
-	TArray<UCommonStyleSheet*> ImportedStyleSheets;
+	TArray<TObjectPtr<UCommonStyleSheet>> ImportedStyleSheets;
 
 	bool bIsApplying = false;
 };
diff --git a/Engine/Plugins/Experimental/CommonUI/Source/CommonUI/Public/Widgets/CommonActivatableWidgetContainer.h b/Engine/Plugins/Experimental/CommonUI/Source/CommonUI/Public/Widgets/CommonActivatableWidgetContainer.h
--- a/Engine/Plugins/Experimental/CommonUI/Source/CommonUI/Public/Widgets/CommonActivatableWidgetContainer.h
+++ b/Engine/Plugins/Experimental/CommonUI/Source/CommonUI/Public/Widgets/CommonActivatableWidgetContainer.h
@@ -108,10 +108,10 @@ protected:
 	float TransitionDuration = 0.4f;
 
 	UPROPERTY(Transient)
-	TArray<UCommonActivatableWidget*> WidgetList;
+	TArray<TObjectPtr<UCommonActivatableWidget>> WidgetList;
 
 	UPROPERTY(Transient)
-	UCommonActivatableWidget* DisplayedWidget;
+	TObjectPtr<UCommonActivatableWidget> DisplayedWidget;
 
 	UPROPERTY(Transient)
 	FUserWidgetPool GeneratedWidgetsPool;
@@ -195,7 +195,7 @@ private:
 	TSubclassOf<UCommonActivatableWidget> RootContentWidgetClass;
 
 	UPROPERTY(Transient)
-	UCommonActivatableWidget* RootContentWidget;
+	TObjectPtr<UCommonActivatableWidget> RootContentWidget;
 };
 
 //////////////////////////////////////////////////////////////////////////
diff --git a/Engine/Plugins/Experimental/Enterprise/DatasmithRuntime/Source/Private/DatasmithRuntimeAuxiliaryData.h b/Engine/Plugins/Experimental/Enterprise/DatasmithRuntime/Source/Private/DatasmithRuntimeAuxiliaryData.h
--- a/Engine/Plugins/Experimental/Enterprise/DatasmithRuntime/Source/Private/DatasmithRuntimeAuxiliaryData.h
+++ b/Engine/Plugins/Experimental/Enterprise/DatasmithRuntime/Source/Private/DatasmithRuntimeAuxiliaryData.h
@@ -74,7 +74,7 @@ public:
 private:
 	/** Array of user data stored with the asset */
 	UPROPERTY()
-	TArray< UAssetUserData* > AssetUserData;
+	TArray< TObjectPtr<UAssetUserData> > AssetUserData;
 };
 
 /** Asset user data that can be used with DatasmithRuntime on Actors and other objects  */
@@ -91,7 +91,7 @@ public:
 	}
 
 	UPROPERTY(VisibleAnywhere, Category = "DatasmithRuntime Internal")
-	UObject* Auxiliary;
+	TObjectPtr<UObject> Auxiliary;
 
 	TSet<uint64> Referencers;
 
diff --git a/Engine/Plugins/Experimental/GPULightmass/Source/GPULightmass/Public/GPULightmassSettings.h b/Engine/Plugins/Experimental/GPULightmass/Source/GPULightmass/Public/GPULightmassSettings.h
--- a/Engine/Plugins/Experimental/GPULightmass/Source/GPULightmass/Public/GPULightmassSettings.h
+++ b/Engine/Plugins/Experimental/GPULightmass/Source/GPULightmass/Public/GPULightmassSettings.h
@@ -146,7 +146,7 @@ class GPULIGHTMASS_API AGPULightmassSettingsActor : public AInfo
 
 public:
 	UPROPERTY(EditAnywhere, BlueprintReadWrite, Category = Settings)
-	UGPULightmassSettings* Settings;
+	TObjectPtr<UGPULightmassSettings> Settings;
 };
 
 UCLASS()
diff --git a/Engine/Plugins/Experimental/GameFeatures/Source/GameFeatures/Private/GameFeaturePluginStateMachine.h b/Engine/Plugins/Experimental/GameFeatures/Source/GameFeatures/Private/GameFeaturePluginStateMachine.h
--- a/Engine/Plugins/Experimental/GameFeatures/Source/GameFeatures/Private/GameFeaturePluginStateMachine.h
+++ b/Engine/Plugins/Experimental/GameFeatures/Source/GameFeatures/Private/GameFeaturePluginStateMachine.h
@@ -263,7 +263,7 @@ struct FGameFeaturePluginStateMachineProperties
 
 	/** The data asset describing this game feature */
 	UPROPERTY(Transient)
-	UGameFeatureData* GameFeatureData = nullptr;
+	TObjectPtr<UGameFeatureData> GameFeatureData = nullptr;
 
 	/** Callbacks for when the current state transition is cancelled */
 	DECLARE_MULTICAST_DELEGATE_OneParam(FOnTransitionCanceled, UGameFeaturePluginStateMachine* /*Machine*/);
diff --git a/Engine/Plugins/Experimental/GameFeatures/Source/GameFeatures/Public/GameFeatureData.h b/Engine/Plugins/Experimental/GameFeatures/Source/GameFeatures/Public/GameFeatureData.h
--- a/Engine/Plugins/Experimental/GameFeatures/Source/GameFeatures/Public/GameFeatureData.h
+++ b/Engine/Plugins/Experimental/GameFeatures/Source/GameFeatures/Public/GameFeatureData.h
@@ -52,7 +52,7 @@ protected:
 
 	/** List of actions to perform as this game feature is loaded/activated/deactivated/unloaded */
 	UPROPERTY(EditDefaultsOnly, Instanced, Category="Actions")
-	TArray<UGameFeatureAction*> Actions;
+	TArray<TObjectPtr<UGameFeatureAction>> Actions;
 
 	/** List of asset types to scan at startup */
 	UPROPERTY(EditAnywhere, Category="Game Feature | Asset Manager", meta=(TitleProperty="PrimaryAssetType"))
diff --git a/Engine/Plugins/Experimental/GameFeatures/Source/GameFeatures/Public/GameFeaturesSubsystem.h b/Engine/Plugins/Experimental/GameFeatures/Source/GameFeatures/Public/GameFeaturesSubsystem.h
--- a/Engine/Plugins/Experimental/GameFeatures/Source/GameFeatures/Public/GameFeaturesSubsystem.h
+++ b/Engine/Plugins/Experimental/GameFeatures/Source/GameFeatures/Public/GameFeaturesSubsystem.h
@@ -577,19 +577,19 @@ private:
 private:
 	/** The list of all game feature plugin state machine objects */
 	UPROPERTY(Transient)
-	TMap<FGameFeaturePluginIdentifier, UGameFeaturePluginStateMachine*> GameFeaturePluginStateMachines;
+	TMap<FGameFeaturePluginIdentifier, TObjectPtr<UGameFeaturePluginStateMachine>> GameFeaturePluginStateMachines;
 
 	/** Game feature plugin state machine objects that are being terminated. Used to prevent GC until termination is complete. */
 	UPROPERTY(Transient)
-	TArray<UGameFeaturePluginStateMachine*> TerminalGameFeaturePluginStateMachines;
+	TArray<TObjectPtr<UGameFeaturePluginStateMachine>> TerminalGameFeaturePluginStateMachines;
 
 	TMap<FString, FString> GameFeaturePluginNameToPathMap;
 
 	UPROPERTY()
-	TArray<UObject*> Observers;
+	TArray<TObjectPtr<UObject>> Observers;
 
 	UPROPERTY(Transient)
-	UGameFeaturesProjectPolicies* GameSpecificPolicies;
+	TObjectPtr<UGameFeaturesProjectPolicies> GameSpecificPolicies;
 
 	bool bInitializedPolicyManager = false;
 };
diff --git a/Engine/Plugins/Experimental/GameplayBehaviors/Source/GameplayBehaviorsModule/Public/AI/GameplayBehavior_BehaviorTree.h b/Engine/Plugins/Experimental/GameplayBehaviors/Source/GameplayBehaviorsModule/Public/AI/GameplayBehavior_BehaviorTree.h
--- a/Engine/Plugins/Experimental/GameplayBehaviors/Source/GameplayBehaviorsModule/Public/AI/GameplayBehavior_BehaviorTree.h
+++ b/Engine/Plugins/Experimental/GameplayBehaviors/Source/GameplayBehaviorsModule/Public/AI/GameplayBehavior_BehaviorTree.h
@@ -25,10 +25,10 @@ protected:
 	void OnTimerTick();
 
 	UPROPERTY()
-	UBehaviorTree* PreviousBT;
+	TObjectPtr<UBehaviorTree> PreviousBT;
 
 	UPROPERTY()
-	AAIController* AIController;
+	TObjectPtr<AAIController> AIController;
 
 	/** Indicates if BehaviorTree should run only once or in loop. */
 	UPROPERTY(EditAnywhere, Category = SmartObject)
diff --git a/Engine/Plugins/Experimental/GameplayBehaviors/Source/GameplayBehaviorsModule/Public/GameplayBehavior.h b/Engine/Plugins/Experimental/GameplayBehaviors/Source/GameplayBehaviorsModule/Public/GameplayBehavior.h
--- a/Engine/Plugins/Experimental/GameplayBehaviors/Source/GameplayBehaviorsModule/Public/GameplayBehavior.h
+++ b/Engine/Plugins/Experimental/GameplayBehaviors/Source/GameplayBehaviorsModule/Public/GameplayBehavior.h
@@ -163,11 +163,11 @@ protected:
 	 * Can be used as patrol points, investigation location, etc.
 	 */
 	UPROPERTY(EditAnywhere, BlueprintReadWrite, Category = GameplayBehavior)
-	TArray<AActor*> RelevantActors;
+	TArray<TObjectPtr<AActor>> RelevantActors;
 
 	/* SmartObject Actor Owner, can be null */
 	UPROPERTY(Transient)
-	AActor* TransientSmartObjectOwner = nullptr;
+	TObjectPtr<AActor> TransientSmartObjectOwner = nullptr;
 
 	/**
 	 * Used mostly as world context for IGameplayTaskOwnerInterface function.
@@ -175,10 +175,10 @@ protected:
 	 *	Set automatically as part of Trigger call
 	 */
 	UPROPERTY()
-	AActor* TransientAvatar = nullptr;
+	TObjectPtr<AActor> TransientAvatar = nullptr;
 
 private:
 	/** List of currently active tasks, do not modify directly */
 	UPROPERTY()
-	TArray<UGameplayTask*> ActiveTasks;
+	TArray<TObjectPtr<UGameplayTask>> ActiveTasks;
 };
diff --git a/Engine/Plugins/Experimental/GameplayBehaviors/Source/GameplayBehaviorsModule/Public/GameplayBehaviorSubsystem.h b/Engine/Plugins/Experimental/GameplayBehaviors/Source/GameplayBehaviorsModule/Public/GameplayBehaviorSubsystem.h
--- a/Engine/Plugins/Experimental/GameplayBehaviors/Source/GameplayBehaviorsModule/Public/GameplayBehaviorSubsystem.h
+++ b/Engine/Plugins/Experimental/GameplayBehaviors/Source/GameplayBehaviorsModule/Public/GameplayBehaviorSubsystem.h
@@ -17,7 +17,7 @@ struct FAgentGameplayBehaviors
 	GENERATED_BODY()
 
 	UPROPERTY()
-	TArray<UGameplayBehavior*> Behaviors;
+	TArray<TObjectPtr<UGameplayBehavior>> Behaviors;
 };
 
 UCLASS(config = Game, defaultconfig, Transient)
@@ -36,5 +36,5 @@ protected:
 	virtual bool TriggerBehaviorImpl(UGameplayBehavior& Behavior, AActor& Avatar, const UGameplayBehaviorConfig* Config, AActor* SmartObjectOwner = nullptr);
 
 	UPROPERTY()
-	TMap<AActor*, FAgentGameplayBehaviors> AgentGameplayBehaviors;
+	TMap<TObjectPtr<AActor>, FAgentGameplayBehaviors> AgentGameplayBehaviors;
 };
diff --git a/Engine/Plugins/Experimental/GameplayBehaviors/Source/GameplayBehaviorsModule/Public/GameplayBehavior_AnimationBased.h b/Engine/Plugins/Experimental/GameplayBehaviors/Source/GameplayBehaviorsModule/Public/GameplayBehavior_AnimationBased.h
--- a/Engine/Plugins/Experimental/GameplayBehaviors/Source/GameplayBehaviorsModule/Public/GameplayBehavior_AnimationBased.h
+++ b/Engine/Plugins/Experimental/GameplayBehaviors/Source/GameplayBehaviorsModule/Public/GameplayBehavior_AnimationBased.h
@@ -19,13 +19,13 @@ struct FMontagePlaybackData
 	GENERATED_BODY()
 
 	UPROPERTY()
-	AActor* Avatar;
+	TObjectPtr<AActor> Avatar;
 
 	UPROPERTY()
-	UAnimMontage* AnimMontage;
+	TObjectPtr<UAnimMontage> AnimMontage;
 
 	UPROPERTY()
-	UAbilitySystemComponent* AbilityComponent;
+	TObjectPtr<UAbilitySystemComponent> AbilityComponent;
 
 	float PlayRate;
 	FName SectionName;
diff --git a/Engine/Plugins/Experimental/GeometryCacheAbcFile/Source/GeometryCacheAbcFile/Private/GeometryCacheAbcFileActor.h b/Engine/Plugins/Experimental/GeometryCacheAbcFile/Source/GeometryCacheAbcFile/Private/GeometryCacheAbcFileActor.h
--- a/Engine/Plugins/Experimental/GeometryCacheAbcFile/Source/GeometryCacheAbcFile/Private/GeometryCacheAbcFileActor.h
+++ b/Engine/Plugins/Experimental/GeometryCacheAbcFile/Source/GeometryCacheAbcFile/Private/GeometryCacheAbcFileActor.h
@@ -18,5 +18,5 @@ class AGeometryCacheAbcFileActor : public AActor
 
 private:
 	UPROPERTY(Category = GeometryCacheAbcFileActor, VisibleAnywhere)
-	UGeometryCacheAbcFileComponent* GeometryCacheAbcFileComponent;
+	TObjectPtr<UGeometryCacheAbcFileComponent> GeometryCacheAbcFileComponent;
 };
diff --git a/Engine/Plugins/Experimental/GeometryCacheAbcFile/Source/GeometryCacheAbcFile/Public/GeometryCacheAbcFileComponent.h b/Engine/Plugins/Experimental/GeometryCacheAbcFile/Source/GeometryCacheAbcFile/Public/GeometryCacheAbcFileComponent.h
--- a/Engine/Plugins/Experimental/GeometryCacheAbcFile/Source/GeometryCacheAbcFile/Public/GeometryCacheAbcFileComponent.h
+++ b/Engine/Plugins/Experimental/GeometryCacheAbcFile/Source/GeometryCacheAbcFile/Public/GeometryCacheAbcFileComponent.h
@@ -57,5 +57,5 @@ protected:
 	void InitializeGeometryCache();
 
 	UPROPERTY(Transient)
-	UAbcImportSettings* AbcSettings;
+	TObjectPtr<UAbcImportSettings> AbcSettings;
 };
diff --git a/Engine/Plugins/Experimental/GizmoEdMode/Source/GizmoEdMode/Public/GizmoEdMode.h b/Engine/Plugins/Experimental/GizmoEdMode/Source/GizmoEdMode/Public/GizmoEdMode.h
--- a/Engine/Plugins/Experimental/GizmoEdMode/Source/GizmoEdMode/Public/GizmoEdMode.h
+++ b/Engine/Plugins/Experimental/GizmoEdMode/Source/GizmoEdMode/Public/GizmoEdMode.h
@@ -49,7 +49,7 @@ private:
 	IAssetEditorGizmoFactory* LastFactory = nullptr;
 
 	UPROPERTY()
-	TArray<UInteractiveGizmo*> InteractiveGizmos;
+	TArray<TObjectPtr<UInteractiveGizmo>> InteractiveGizmos;
 
 	FDelegateHandle WidgetModeChangedHandle;
 
diff --git a/Engine/Plugins/Experimental/GizmoEdMode/Source/LightGizmos/Private/DirectionalLightGizmo.h b/Engine/Plugins/Experimental/GizmoEdMode/Source/LightGizmos/Private/DirectionalLightGizmo.h
--- a/Engine/Plugins/Experimental/GizmoEdMode/Source/LightGizmos/Private/DirectionalLightGizmo.h
+++ b/Engine/Plugins/Experimental/GizmoEdMode/Source/LightGizmos/Private/DirectionalLightGizmo.h
@@ -98,22 +98,22 @@ private:
 
 	/** A transform proxy of the LightActor */
 	UPROPERTY()
-	USubTransformProxy* TransformProxy;
+	TObjectPtr<USubTransformProxy> TransformProxy;
 
 	UPROPERTY()
-	UWorld* World;
+	TObjectPtr<UWorld> World;
 
 	/** The internal actor used by the light gizmo */
 	UPROPERTY()
-	ADirectionalLightGizmoActor* GizmoActor;
+	TObjectPtr<ADirectionalLightGizmoActor> GizmoActor;
 
 	/** Used to properly render the handle gizmo. */
 	UPROPERTY()
-	UGizmoViewContext* GizmoViewContext;
+	TObjectPtr<UGizmoViewContext> GizmoViewContext;
 
 	/** The current target of the gizmo */
 	UPROPERTY()
-	ADirectionalLight* LightActor;
+	TObjectPtr<ADirectionalLight> LightActor;
 
 	UPROPERTY()
 	bool bIsHovering{ false };
@@ -141,7 +141,7 @@ private:
 	FVector RotationPlaneZ;
 
 	UPROPERTY()
-	UGizmoBaseComponent* HitComponent;
+	TObjectPtr<UGizmoBaseComponent> HitComponent;
 
 	UPROPERTY()
 	float ArrowLength{ 120.f };
diff --git a/Engine/Plugins/Experimental/GizmoEdMode/Source/LightGizmos/Private/PointLightGizmo.h b/Engine/Plugins/Experimental/GizmoEdMode/Source/LightGizmos/Private/PointLightGizmo.h
--- a/Engine/Plugins/Experimental/GizmoEdMode/Source/LightGizmos/Private/PointLightGizmo.h
+++ b/Engine/Plugins/Experimental/GizmoEdMode/Source/LightGizmos/Private/PointLightGizmo.h
@@ -59,18 +59,18 @@ private:
 
 	/** The Gizmo that is used to scale the attenuation */
 	UPROPERTY()
-	UScalableSphereGizmo* AttenuationGizmo;
+	TObjectPtr<UScalableSphereGizmo> AttenuationGizmo;
 
 	/** The target Point Light */
 	UPROPERTY()
-	APointLight* LightActor;
+	TObjectPtr<APointLight> LightActor;
 
 	UPROPERTY()
-	UWorld* World;
+	TObjectPtr<UWorld> World;
 
 	/** A transform proxy to represent the light actor in gizmos */
 	UPROPERTY()
-	USubTransformProxy* TransformProxy;
+	TObjectPtr<USubTransformProxy> TransformProxy;
 	
 	/**
 	 * Function that is called when the radius of the sphere is changed
diff --git a/Engine/Plugins/Experimental/GizmoEdMode/Source/LightGizmos/Private/SpotLightGizmo.h b/Engine/Plugins/Experimental/GizmoEdMode/Source/LightGizmos/Private/SpotLightGizmo.h
--- a/Engine/Plugins/Experimental/GizmoEdMode/Source/LightGizmos/Private/SpotLightGizmo.h
+++ b/Engine/Plugins/Experimental/GizmoEdMode/Source/LightGizmos/Private/SpotLightGizmo.h
@@ -90,22 +90,22 @@ private:
 
 	/** The current target light the gizmo is attached to*/
 	UPROPERTY()
-	ASpotLight* LightActor;
+	TObjectPtr<ASpotLight> LightActor;
 
 	UPROPERTY()
-	UWorld* World;
+	TObjectPtr<UWorld> World;
 
 	/** A transform proxy to use with other gizmos*/
 	UPROPERTY()
-	USubTransformProxy* TransformProxy;
+	TObjectPtr<USubTransformProxy> TransformProxy;
 
 	/** The gizmo to change the outer angle of the spotlight */
 	UPROPERTY()
-	UScalableConeGizmo* OuterAngleGizmo;
+	TObjectPtr<UScalableConeGizmo> OuterAngleGizmo;
 
 	/** The gizmo to change the inner angle of the spotlight */
 	UPROPERTY()
-	UScalableConeGizmo* InnerAngleGizmo;
+	TObjectPtr<UScalableConeGizmo> InnerAngleGizmo;
 
 	/** The internal gizmo actor that is used by the gizmo 
 	 *  We need a GizmoActor separate from the 2 UScalableConeGizmo's
@@ -113,11 +113,11 @@ private:
 	 *  both of them simultaneously
 	 */
 	UPROPERTY()
-	ASpotLightGizmoActor* GizmoActor;
+	TObjectPtr<ASpotLightGizmoActor> GizmoActor;
 
 	/** Used to properly render the handle gizmo. */
 	UPROPERTY()
-	UGizmoViewContext* GizmoViewContext;
+	TObjectPtr<UGizmoViewContext> GizmoViewContext;
 
 	UPROPERTY()
 	FVector DragStartWorldPosition;
diff --git a/Engine/Plugins/Experimental/ImagePlate/Source/ImagePlate/Public/ImagePlate.h b/Engine/Plugins/Experimental/ImagePlate/Source/ImagePlate/Public/ImagePlate.h
--- a/Engine/Plugins/Experimental/ImagePlate/Source/ImagePlate/Public/ImagePlate.h
+++ b/Engine/Plugins/Experimental/ImagePlate/Source/ImagePlate/Public/ImagePlate.h
@@ -25,5 +25,5 @@ public:
 protected:
 
 	UPROPERTY(Category=ImagePlate, VisibleAnywhere, BlueprintReadOnly, meta=(ExposeFunctionCategories="Mesh,Rendering,Physics,Components|StaticMesh"))
-	UImagePlateComponent* ImagePlate;
+	TObjectPtr<UImagePlateComponent> ImagePlate;
 };
\ No newline at end of file
diff --git a/Engine/Plugins/Experimental/ImagePlate/Source/ImagePlate/Public/ImagePlateComponent.h b/Engine/Plugins/Experimental/ImagePlate/Source/ImagePlate/Public/ImagePlateComponent.h
--- a/Engine/Plugins/Experimental/ImagePlate/Source/ImagePlate/Public/ImagePlateComponent.h
+++ b/Engine/Plugins/Experimental/ImagePlate/Source/ImagePlate/Public/ImagePlateComponent.h
@@ -21,7 +21,7 @@ struct IMAGEPLATE_API FImagePlateParameters
 
 	/** The material that the image plate is rendered with. */
 	UPROPERTY(EditAnywhere, BlueprintReadWrite, AdvancedDisplay, Category="Image Plate")
-	UMaterialInterface* Material;
+	TObjectPtr<UMaterialInterface> Material;
 
 	/** Name of a texture parameter inside the material to patch the render target texture to */
 	UPROPERTY(EditAnywhere, BlueprintReadWrite, AdvancedDisplay, Category="Image Plate")
@@ -41,11 +41,11 @@ struct IMAGEPLATE_API FImagePlateParameters
 
 	/** Transient texture that receives image frames */
 	UPROPERTY(transient, BlueprintReadOnly, Category="Image Plate")
-	UTexture* RenderTexture;
+	TObjectPtr<UTexture> RenderTexture;
 
 	/** Transient MID to hold the material with the render texture patched in */
 	UPROPERTY(transient)
-	UMaterialInstanceDynamic* DynamicMaterial;
+	TObjectPtr<UMaterialInstanceDynamic> DynamicMaterial;
 };
 
 /** 
@@ -141,6 +141,6 @@ private:
 
 #if WITH_EDITORONLY_DATA
 	UPROPERTY(transient)
-	UImagePlateFrustumComponent* EditorFrustum;
+	TObjectPtr<UImagePlateFrustumComponent> EditorFrustum;
 #endif
 };
diff --git a/Engine/Plugins/Experimental/Landmass/Source/Runtime/Public/BrushEffectsList.h b/Engine/Plugins/Experimental/Landmass/Source/Runtime/Public/BrushEffectsList.h
--- a/Engine/Plugins/Experimental/Landmass/Source/Runtime/Public/BrushEffectsList.h
+++ b/Engine/Plugins/Experimental/Landmass/Source/Runtime/Public/BrushEffectsList.h
@@ -69,7 +69,7 @@ struct FBrushEffectCurves
 	bool bUseCurveChannel;
 
 	UPROPERTY(EditAnywhere, BlueprintReadWrite, Category = BrushEffects)
-	UCurveFloat* ElevationCurveAsset;
+	TObjectPtr<UCurveFloat> ElevationCurveAsset;
 
 	UPROPERTY(EditAnywhere, BlueprintReadWrite, Category = BrushEffects)
 	float ChannelEdgeOffset;
@@ -103,7 +103,7 @@ struct FBrushEffectDisplacement
 	float DisplacementTiling;
 
 	UPROPERTY(EditAnywhere, BlueprintReadWrite, Category = BrushEffects)
-	UTexture2D* Texture;
+	TObjectPtr<UTexture2D> Texture;
 
 	UPROPERTY(EditAnywhere, BlueprintReadWrite, Category = BrushEffects)
 	float Midpoint;
diff --git a/Engine/Plugins/Experimental/LandscapePatch/Source/LandscapePatch/Public/LandscapeTexturePatch.h b/Engine/Plugins/Experimental/LandscapePatch/Source/LandscapePatch/Public/LandscapeTexturePatch.h
--- a/Engine/Plugins/Experimental/LandscapePatch/Source/LandscapePatch/Public/LandscapeTexturePatch.h
+++ b/Engine/Plugins/Experimental/LandscapePatch/Source/LandscapePatch/Public/LandscapeTexturePatch.h
@@ -195,7 +195,7 @@ protected:
 
 	// Needed mainly so that we can get at the resolution...
 	UPROPERTY()
-	ULandscapeTexturePatch* OwningPatch = nullptr;
+	TObjectPtr<ULandscapeTexturePatch> OwningPatch = nullptr;
 
 	// TODO: Like the similar flag for the height patch, this might not work once local merge works
 	// with landscape brushes...
diff --git a/Engine/Plugins/Experimental/ModularGameplay/Source/ModularGameplay/Public/Components/GameFrameworkComponentManager.h b/Engine/Plugins/Experimental/ModularGameplay/Source/ModularGameplay/Public/Components/GameFrameworkComponentManager.h
--- a/Engine/Plugins/Experimental/ModularGameplay/Source/ModularGameplay/Public/Components/GameFrameworkComponentManager.h
+++ b/Engine/Plugins/Experimental/ModularGameplay/Source/ModularGameplay/Public/Components/GameFrameworkComponentManager.h
@@ -75,7 +75,7 @@ struct FActorInitStateChangedParams
 
 	/** The actor owning the feature that changed */
 	UPROPERTY(BlueprintReadWrite, EditAnywhere, Category = Default)
-	AActor* OwningActor;
+	TObjectPtr<AActor> OwningActor;
 
 	/** Name of the feature that changed */
 	UPROPERTY(BlueprintReadWrite, EditAnywhere, Category = Default)
@@ -83,7 +83,7 @@ struct FActorInitStateChangedParams
 
 	/** The object (often a component) that implements the feature */
 	UPROPERTY(BlueprintReadWrite, EditAnywhere, Category = Default)
-	UObject* Implementer;
+	TObjectPtr<UObject> Implementer;
 
 	/** The new state of the feature */
 	UPROPERTY(BlueprintReadWrite, EditAnywhere, Category = Default)
diff --git a/Engine/Plugins/Experimental/PCG/Source/PCG/Private/Tests/Elements/PCGPropertyToParamDataElementTest.h b/Engine/Plugins/Experimental/PCG/Source/PCG/Private/Tests/Elements/PCGPropertyToParamDataElementTest.h
--- a/Engine/Plugins/Experimental/PCG/Source/PCG/Private/Tests/Elements/PCGPropertyToParamDataElementTest.h
+++ b/Engine/Plugins/Experimental/PCG/Source/PCG/Private/Tests/Elements/PCGPropertyToParamDataElementTest.h
@@ -71,7 +71,7 @@ public:
 	TSubclassOf<AActor> ClassProperty = AActor::StaticClass();
 
 	UPROPERTY()
-	UObject* ObjectProperty = nullptr;
+	TObjectPtr<UObject> ObjectProperty = nullptr;
 
 	UPROPERTY()
 	FVector2D Vector2Property;
diff --git a/Engine/Plugins/Experimental/PCG/Source/PCG/Public/Data/PCGRenderTargetData.h b/Engine/Plugins/Experimental/PCG/Source/PCG/Public/Data/PCGRenderTargetData.h
--- a/Engine/Plugins/Experimental/PCG/Source/PCG/Public/Data/PCGRenderTargetData.h
+++ b/Engine/Plugins/Experimental/PCG/Source/PCG/Public/Data/PCGRenderTargetData.h
@@ -23,5 +23,5 @@ public:
 	void Initialize(UTextureRenderTarget2D* InRenderTarget, const FTransform& InTransform);
 
 	UPROPERTY(BlueprintReadOnly, VisibleAnywhere, Category = Properties)
-	UTextureRenderTarget2D* RenderTarget = nullptr;
+	TObjectPtr<UTextureRenderTarget2D> RenderTarget = nullptr;
 };
diff --git a/Engine/Plugins/Experimental/PCG/Source/PCG/Public/Data/PCGTextureData.h b/Engine/Plugins/Experimental/PCG/Source/PCG/Public/Data/PCGTextureData.h
--- a/Engine/Plugins/Experimental/PCG/Source/PCG/Public/Data/PCGTextureData.h
+++ b/Engine/Plugins/Experimental/PCG/Source/PCG/Public/Data/PCGTextureData.h
@@ -101,5 +101,5 @@ public:
 	void Initialize(UTexture2D* InTexture, const FTransform& InTransform);
 
 	UPROPERTY(BlueprintReadOnly, VisibleAnywhere, Category = Properties)
-	UTexture2D* Texture = nullptr;
+	TObjectPtr<UTexture2D> Texture = nullptr;
 };
diff --git a/Engine/Plugins/Experimental/PCG/Source/PCG/Public/MeshSelectors/PCGMeshSelectorBase.h b/Engine/Plugins/Experimental/PCG/Source/PCG/Public/MeshSelectors/PCGMeshSelectorBase.h
--- a/Engine/Plugins/Experimental/PCG/Source/PCG/Public/MeshSelectors/PCGMeshSelectorBase.h
+++ b/Engine/Plugins/Experimental/PCG/Source/PCG/Public/MeshSelectors/PCGMeshSelectorBase.h
@@ -38,7 +38,7 @@ struct FPCGMeshInstanceList
 	bool bOverrideMaterials = false;
 
 	UPROPERTY(BlueprintReadWrite, EditAnywhere, Category = Settings)
-	TArray<UMaterialInterface*> MaterialOverrides;
+	TArray<TObjectPtr<UMaterialInterface>> MaterialOverrides;
 
 	UPROPERTY(BlueprintReadWrite, EditAnywhere, Category = Settings)
     TArray<FPCGPoint> Instances;
diff --git a/Engine/Plugins/Experimental/PCG/Source/PCG/Public/MeshSelectors/PCGMeshSelectorByAttribute.h b/Engine/Plugins/Experimental/PCG/Source/PCG/Public/MeshSelectors/PCGMeshSelectorByAttribute.h
--- a/Engine/Plugins/Experimental/PCG/Source/PCG/Public/MeshSelectors/PCGMeshSelectorByAttribute.h
+++ b/Engine/Plugins/Experimental/PCG/Source/PCG/Public/MeshSelectors/PCGMeshSelectorByAttribute.h
@@ -33,6 +33,6 @@ public:
 	bool bOverrideMaterials = false;
 
 	UPROPERTY(BlueprintReadWrite, EditAnywhere, Category = Settings)
-	TArray<UMaterialInterface*> MaterialOverrides;
+	TArray<TObjectPtr<UMaterialInterface>> MaterialOverrides;
 };
 
diff --git a/Engine/Plugins/Experimental/PCG/Source/PCG/Public/PCGComponent.h b/Engine/Plugins/Experimental/PCG/Source/PCG/Public/PCGComponent.h
--- a/Engine/Plugins/Experimental/PCG/Source/PCG/Public/PCGComponent.h
+++ b/Engine/Plugins/Experimental/PCG/Source/PCG/Public/PCGComponent.h
@@ -266,22 +266,22 @@ private:
 	FBox GetGridBounds(AActor* InActor) const;
 
 	UPROPERTY(Transient, NonPIEDuplicateTransient)
-	UPCGData* CachedPCGData = nullptr;
+	TObjectPtr<UPCGData> CachedPCGData = nullptr;
 
 	UPROPERTY(Transient, NonPIEDuplicateTransient)
-	UPCGData* CachedInputData = nullptr;
+	TObjectPtr<UPCGData> CachedInputData = nullptr;
 
 	UPROPERTY(Transient, NonPIEDuplicateTransient)
-	UPCGData* CachedActorData = nullptr;
+	TObjectPtr<UPCGData> CachedActorData = nullptr;
 
 	UPROPERTY(Transient, NonPIEDuplicateTransient)
-	UPCGData* CachedLandscapeData = nullptr;
+	TObjectPtr<UPCGData> CachedLandscapeData = nullptr;
 
 	UPROPERTY(Transient, NonPIEDuplicateTransient)
-	UPCGData* CachedLandscapeHeightData = nullptr;
+	TObjectPtr<UPCGData> CachedLandscapeHeightData = nullptr;
 
 	UPROPERTY(Transient, NonPIEDuplicateTransient)
-	TMap<AActor*, UPCGData*> CachedExclusionData;
+	TMap<TObjectPtr<AActor>, TObjectPtr<UPCGData>> CachedExclusionData;
 
 	// Cached excluded actors list is serialized because we can't get it at postload time
 	UPROPERTY()
@@ -333,7 +333,7 @@ private:
 	bool bActorToTagsMapPopulated = false;
 
 	UPROPERTY(Transient)
-	TMap<const UPCGNode*, FPCGDataCollection> InspectionCache;
+	TMap<TObjectPtr<const UPCGNode>, FPCGDataCollection> InspectionCache;
 #endif
 
 	FCriticalSection GeneratedResourcesLock;
diff --git a/Engine/Plugins/Experimental/PanoramicCapture/Source/PanoramicCapture/Private/SceneCapturer.h b/Engine/Plugins/Experimental/PanoramicCapture/Source/PanoramicCapture/Private/SceneCapturer.h
--- a/Engine/Plugins/Experimental/PanoramicCapture/Source/PanoramicCapture/Private/SceneCapturer.h
+++ b/Engine/Plugins/Experimental/PanoramicCapture/Source/PanoramicCapture/Private/SceneCapturer.h
@@ -29,7 +29,7 @@ struct FPostProcessVolumeData
 	GENERATED_BODY()
 	
 	UPROPERTY()
-	APostProcessVolume* Object;
+	TObjectPtr<APostProcessVolume> Object;
 
 	bool WasEnabled;
 };
@@ -176,7 +176,7 @@ public:
 	
 	// CaptureSceneComponent will be used as parent of capturecomponents to provide world location and rotation.
 	UPROPERTY(Transient)
-	USceneComponent* CaptureSceneComponent;
+	TObjectPtr<USceneComponent> CaptureSceneComponent;
 
 	bool GetComponentSteps( int32 Step, int32& CurrentHorizontalStep, int32& CurrentVerticalStep )
 	{
diff --git a/Engine/Plugins/Experimental/PanoramicCapture/Source/PanoramicCapture/Public/StereoCapturePawn.h b/Engine/Plugins/Experimental/PanoramicCapture/Source/PanoramicCapture/Public/StereoCapturePawn.h
--- a/Engine/Plugins/Experimental/PanoramicCapture/Source/PanoramicCapture/Public/StereoCapturePawn.h
+++ b/Engine/Plugins/Experimental/PanoramicCapture/Source/PanoramicCapture/Public/StereoCapturePawn.h
@@ -60,10 +60,10 @@ public:
     void UpdateStereoAtlas(UObject* WorldContextObject, struct FLatentActionInfo LatentInfo);
 
     UPROPERTY(VisibleAnywhere, BlueprintReadOnly, Category = Pawn)
-    UTexture2D* LeftEyeAtlas;
+    TObjectPtr<UTexture2D> LeftEyeAtlas;
 
     UPROPERTY(VisibleAnywhere, BlueprintReadOnly, Category = Pawn)
-    UTexture2D* RightEyeAtlas;
+    TObjectPtr<UTexture2D> RightEyeAtlas;
 
 
     FStereoCaptureDoneAction* StereoCaptureDoneAction;
diff --git a/Engine/Plugins/Experimental/RemoteSession/Source/RemoteSessionEditor/Private/Widgets/SRemoteSessionStream.h b/Engine/Plugins/Experimental/RemoteSession/Source/RemoteSessionEditor/Private/Widgets/SRemoteSessionStream.h
--- a/Engine/Plugins/Experimental/RemoteSession/Source/RemoteSessionEditor/Private/Widgets/SRemoteSessionStream.h
+++ b/Engine/Plugins/Experimental/RemoteSession/Source/RemoteSessionEditor/Private/Widgets/SRemoteSessionStream.h
@@ -73,7 +73,7 @@ public:
 	FVector2D Size;
 
 	UPROPERTY(EditAnywhere, AdvancedDisplay, Category = "Remote Session")
-	UTextureRenderTarget2D* RenderTarget;
+	TObjectPtr<UTextureRenderTarget2D> RenderTarget;
 
 	UPROPERTY(EditAnywhere, AdvancedDisplay, Category = "Remote Session")
 	int16 Port;
diff --git a/Engine/Plugins/Experimental/SampleToolsEditorMode/Source/Private/SampleTools/CreateActorSampleTool.h b/Engine/Plugins/Experimental/SampleToolsEditorMode/Source/Private/SampleTools/CreateActorSampleTool.h
--- a/Engine/Plugins/Experimental/SampleToolsEditorMode/Source/Private/SampleTools/CreateActorSampleTool.h
+++ b/Engine/Plugins/Experimental/SampleToolsEditorMode/Source/Private/SampleTools/CreateActorSampleTool.h
@@ -68,7 +68,7 @@ public:
 
 protected:
 	UPROPERTY()
-	UCreateActorSampleToolProperties* Properties;
+	TObjectPtr<UCreateActorSampleToolProperties> Properties;
 
 
 protected:
diff --git a/Engine/Plugins/Experimental/SampleToolsEditorMode/Source/Private/SampleTools/DrawCurveOnMeshSampleTool.h b/Engine/Plugins/Experimental/SampleToolsEditorMode/Source/Private/SampleTools/DrawCurveOnMeshSampleTool.h
--- a/Engine/Plugins/Experimental/SampleToolsEditorMode/Source/Private/SampleTools/DrawCurveOnMeshSampleTool.h
+++ b/Engine/Plugins/Experimental/SampleToolsEditorMode/Source/Private/SampleTools/DrawCurveOnMeshSampleTool.h
@@ -80,7 +80,7 @@ public:
 
 protected:
 	UPROPERTY()
-	UDrawCurveOnMeshSampleToolProperties* Settings;
+	TObjectPtr<UDrawCurveOnMeshSampleToolProperties> Settings;
 
 	UPROPERTY()
 	TArray<FVector> Positions;
diff --git a/Engine/Plugins/Experimental/SampleToolsEditorMode/Source/Private/SampleTools/MeasureDistanceSampleTool.h b/Engine/Plugins/Experimental/SampleToolsEditorMode/Source/Private/SampleTools/MeasureDistanceSampleTool.h
--- a/Engine/Plugins/Experimental/SampleToolsEditorMode/Source/Private/SampleTools/MeasureDistanceSampleTool.h
+++ b/Engine/Plugins/Experimental/SampleToolsEditorMode/Source/Private/SampleTools/MeasureDistanceSampleTool.h
@@ -84,7 +84,7 @@ public:
 protected:
 	/** Properties of the tool are stored here */
 	UPROPERTY()
-	UMeasureDistanceProperties* Properties;
+	TObjectPtr<UMeasureDistanceProperties> Properties;
 
 
 protected:
diff --git a/Engine/Plugins/Experimental/Shotgrid/Source/Shotgrid/Private/ShotgridEngine.h b/Engine/Plugins/Experimental/Shotgrid/Source/Shotgrid/Private/ShotgridEngine.h
--- a/Engine/Plugins/Experimental/Shotgrid/Source/Shotgrid/Private/ShotgridEngine.h
+++ b/Engine/Plugins/Experimental/Shotgrid/Source/Shotgrid/Private/ShotgridEngine.h
@@ -79,7 +79,7 @@ public:
 	TArray<FAssetData> SelectedAssets;
 
 	UPROPERTY(EditAnywhere, BlueprintReadWrite, Category = Python, meta = (Deprecated, DeprecationMessage = "SelectedActors is deprecated, use GetSelectedActors instead." ))
-	TArray<AActor*> SelectedActors;
+	TArray<TObjectPtr<AActor>> SelectedActors;
 
 	/** Selected actors to be used for Shotgrid commands */
 	UFUNCTION(BlueprintCallable, Category = Python)
diff --git a/Engine/Plugins/Experimental/Text3D/Source/Text3D/Private/Text3DEngineSubsystem.h b/Engine/Plugins/Experimental/Text3D/Source/Text3D/Private/Text3DEngineSubsystem.h
--- a/Engine/Plugins/Experimental/Text3D/Source/Text3D/Private/Text3DEngineSubsystem.h
+++ b/Engine/Plugins/Experimental/Text3D/Source/Text3D/Private/Text3DEngineSubsystem.h
@@ -60,7 +60,7 @@ public:
 	TSharedPtr<int32> GetCacheCounter();
 
 	UPROPERTY()
-	TMap<uint32, class UStaticMesh*> Glyphs;
+	TMap<uint32, TObjectPtr<class UStaticMesh>> Glyphs;
 
 private:
 	TSharedPtr<int32> CacheCounter;
@@ -91,7 +91,7 @@ public:
 	TSharedContourNode GetGlyphContours(uint32 GlyphIndex);
 
 	UPROPERTY()
-	class UFont* Font;
+	TObjectPtr<class UFont> Font;
 
 	UPROPERTY()
 	TMap<uint32, FCachedFontMeshes> Meshes;
@@ -122,7 +122,7 @@ public:
 	FCachedFontData& GetCachedFontData(class UFont* Font);
 
 	UPROPERTY()
-	class UMaterial* DefaultMaterial;
+	TObjectPtr<class UMaterial> DefaultMaterial;
 
 private:
 	bool CleanupTimerCallback(float DeltaTime);
diff --git a/Engine/Plugins/Experimental/Text3D/Source/Text3D/Public/Text3DComponent.h b/Engine/Plugins/Experimental/Text3D/Source/Text3D/Public/Text3DComponent.h
--- a/Engine/Plugins/Experimental/Text3D/Source/Text3D/Public/Text3DComponent.h
+++ b/Engine/Plugins/Experimental/Text3D/Source/Text3D/Public/Text3DComponent.h
@@ -72,23 +72,23 @@ public:
 
 	/** Material for the front part */
 	UPROPERTY(EditAnywhere, BlueprintReadWrite, BlueprintSetter = SetFrontMaterial, Category = "Text3D")
-	class UMaterialInterface* FrontMaterial;
+	TObjectPtr<class UMaterialInterface> FrontMaterial;
 
 	/** Material for the bevel part */
 	UPROPERTY(EditAnywhere, BlueprintReadWrite, BlueprintSetter = SetBevelMaterial, Category = "Text3D", Meta = (EditCondition = "!bOutline"))
-	class UMaterialInterface* BevelMaterial;
+	TObjectPtr<class UMaterialInterface> BevelMaterial;
 
 	/** Material for the extruded part */
 	UPROPERTY(EditAnywhere, BlueprintReadWrite, BlueprintSetter = SetExtrudeMaterial, Category = "Text3D")
-	class UMaterialInterface* ExtrudeMaterial;
+	TObjectPtr<class UMaterialInterface> ExtrudeMaterial;
 
 	/** Material for the back part */
 	UPROPERTY(EditAnywhere, BlueprintReadWrite, BlueprintSetter = SetBackMaterial, Category = "Text3D")
-	class UMaterialInterface* BackMaterial;
+	TObjectPtr<class UMaterialInterface> BackMaterial;
 
 	/** Text font */
 	UPROPERTY(EditAnywhere, BlueprintReadWrite, BlueprintSetter = SetFont, Category = "Text3D")
-	class UFont* Font;
+	TObjectPtr<class UFont> Font;
 
 	/** Horizontal text alignment */
 	UPROPERTY(EditAnywhere, BlueprintReadWrite, BlueprintSetter = SetHorizontalAlignment, Category = "Text3D")
@@ -275,7 +275,7 @@ protected:
 
 private:
 	UPROPERTY(Transient)
-	class USceneComponent* TextRoot;
+	TObjectPtr<class USceneComponent> TextRoot;
 
 	UPROPERTY(BlueprintAssignable, Category = Events, meta = (AllowPrivateAccess = true, DisplayName = "On Text Generated"))
 	FTextGenerated TextGeneratedDelegate;
@@ -294,10 +294,10 @@ private:
 	TArray<TSharedPtr<int32>> CachedCounterReferences;
 
 	UPROPERTY(Transient)
-	TArray<USceneComponent*> CharacterKernings;
+	TArray<TObjectPtr<USceneComponent>> CharacterKernings;
 
 	UPROPERTY(Transient)
-	TArray<UStaticMeshComponent*> CharacterMeshes;
+	TArray<TObjectPtr<UStaticMeshComponent>> CharacterMeshes;
 
 	/** Allocates, or shrinks existing components to match the input number. Returns false if nothing modified. */
 	bool AllocateGlyphs(int32 Num);
diff --git a/Engine/Plugins/Experimental/UIFramework/Source/Public/Layers/UIFCanvasLayerUserWidget.h b/Engine/Plugins/Experimental/UIFramework/Source/Public/Layers/UIFCanvasLayerUserWidget.h
--- a/Engine/Plugins/Experimental/UIFramework/Source/Public/Layers/UIFCanvasLayerUserWidget.h
+++ b/Engine/Plugins/Experimental/UIFramework/Source/Public/Layers/UIFCanvasLayerUserWidget.h
@@ -60,5 +60,5 @@ private:
 
 private:
 	UPROPERTY(BlueprintReadOnly, Category="UIFramework", meta = (BindWidget, AllowPrivateAccess))
-	UCanvasPanel* Canvas = nullptr;
+	TObjectPtr<UCanvasPanel> Canvas = nullptr;
 };
diff --git a/Engine/Plugins/Experimental/UIFramework/Source/Public/Layers/UIFSlotLayer.h b/Engine/Plugins/Experimental/UIFramework/Source/Public/Layers/UIFSlotLayer.h
--- a/Engine/Plugins/Experimental/UIFramework/Source/Public/Layers/UIFSlotLayer.h
+++ b/Engine/Plugins/Experimental/UIFramework/Source/Public/Layers/UIFSlotLayer.h
@@ -115,5 +115,5 @@ private:
 	FUIFSlotWidgetList WidgetList;
 
 	UPROPERTY(Transient)
-	TMap<FName, UUIFWidget*> LocalSlotContents;
+	TMap<FName, TObjectPtr<UUIFWidget>> LocalSlotContents;
 };
diff --git a/Engine/Plugins/Experimental/UIFramework/Source/Public/UIFPlayerComponent.h b/Engine/Plugins/Experimental/UIFramework/Source/Public/UIFPlayerComponent.h
--- a/Engine/Plugins/Experimental/UIFramework/Source/Public/UIFPlayerComponent.h
+++ b/Engine/Plugins/Experimental/UIFramework/Source/Public/UIFPlayerComponent.h
@@ -22,7 +22,7 @@ struct FUIFLayerEntry : public FFastArraySerializerItem
 	FUIFLayerEntry() = default;
 
 	UPROPERTY()
-	UUIFLayer* Layer = nullptr;
+	TObjectPtr<UUIFLayer> Layer = nullptr;
 
 	//~ In the LayerEntry instead of UILayer because it is needed for the correct initialization and cannot changed at runtime
 	UPROPERTY()
@@ -74,7 +74,7 @@ private:
 	TArray<FUIFLayerEntry> Entries;
 
 	UPROPERTY()
-	UActorComponent* OwnerComponent = nullptr;
+	TObjectPtr<UActorComponent> OwnerComponent = nullptr;
 };
 
 template<>
diff --git a/Engine/Plugins/Experimental/UIFramework/Source/Public/Widgets/UIFButton.h b/Engine/Plugins/Experimental/UIFramework/Source/Public/Widgets/UIFButton.h
--- a/Engine/Plugins/Experimental/UIFramework/Source/Public/Widgets/UIFButton.h
+++ b/Engine/Plugins/Experimental/UIFramework/Source/Public/Widgets/UIFButton.h
@@ -20,10 +20,10 @@ class UUIFButtonUserWidget : public UUserWidget
 
 public:
 	UPROPERTY(BlueprintReadOnly, Category = "UI Framework", meta = (BindWidget))
-	UTextBlock* TextBlock = nullptr;
+	TObjectPtr<UTextBlock> TextBlock = nullptr;
 
 	UPROPERTY(BlueprintReadOnly, Category = "UI Framework", meta = (BindWidget))
-	UButton* Button = nullptr;
+	TObjectPtr<UButton> Button = nullptr;
 };
 
 
diff --git a/Engine/Plugins/Experimental/UVEditor/Source/UVEditor/Public/UVEditorSubsystem.h b/Engine/Plugins/Experimental/UVEditor/Source/UVEditor/Public/UVEditorSubsystem.h
--- a/Engine/Plugins/Experimental/UVEditor/Source/UVEditor/Public/UVEditorSubsystem.h
+++ b/Engine/Plugins/Experimental/UVEditor/Source/UVEditor/Public/UVEditorSubsystem.h
@@ -64,7 +64,7 @@ protected:
 	 * added here in Initialize() as they are developed.
 	 */
 	UPROPERTY()
-	UToolTargetManager* ToolTargetManager = nullptr;
+	TObjectPtr<UToolTargetManager> ToolTargetManager = nullptr;
 
 	TMap<TObjectPtr<UObject>, TObjectPtr<UUVEditor>> OpenedEditorInstances;
 };
\ No newline at end of file
diff --git a/Engine/Plugins/Experimental/VirtualCamera/Source/VCamCore/Public/SimpleVirtualCamera.h b/Engine/Plugins/Experimental/VirtualCamera/Source/VCamCore/Public/SimpleVirtualCamera.h
--- a/Engine/Plugins/Experimental/VirtualCamera/Source/VCamCore/Public/SimpleVirtualCamera.h
+++ b/Engine/Plugins/Experimental/VirtualCamera/Source/VCamCore/Public/SimpleVirtualCamera.h
@@ -20,5 +20,5 @@ public:
 
 private:
 	UPROPERTY(EditAnywhere, Category = "VirtualCamera")
-	UVCamComponent* VCamComponent = nullptr;
+	TObjectPtr<UVCamComponent> VCamComponent = nullptr;
 };
diff --git a/Engine/Plugins/Experimental/VirtualCamera/Source/VCamCore/Public/VCamComponent.h b/Engine/Plugins/Experimental/VirtualCamera/Source/VCamCore/Public/VCamComponent.h
--- a/Engine/Plugins/Experimental/VirtualCamera/Source/VCamCore/Public/VCamComponent.h
+++ b/Engine/Plugins/Experimental/VirtualCamera/Source/VCamCore/Public/VCamComponent.h
@@ -252,7 +252,7 @@ public:
 
 	// List of Output Providers (executed in order)
 	UPROPERTY(EditAnywhere, Instanced, Category="VirtualCamera")
-	TArray<UVCamOutputProviderBase*> OutputProviders;
+	TArray<TObjectPtr<UVCamOutputProviderBase>> OutputProviders;
 
 	UFUNCTION(BlueprintCallable, Category="VirtualCamera")
 	void GetLiveLinkDataForCurrentFrame(FLiveLinkCameraBlueprintData& LiveLinkData);
@@ -366,7 +366,7 @@ private:
 
 	// Modifier Context object that can be accessed by the Modifier Stack
 	UPROPERTY(EditAnywhere, Instanced, Category = "VirtualCamera")
-	UVCamModifierContext* ModifierContext;
+	TObjectPtr<UVCamModifierContext> ModifierContext;
 
 	// List of Modifiers (executed in order)
 	UPROPERTY(EditAnywhere, Category = "VirtualCamera")
diff --git a/Engine/Plugins/Experimental/VirtualCamera/Source/VCamCore/Public/VCamOutputComposure.h b/Engine/Plugins/Experimental/VirtualCamera/Source/VCamCore/Public/VCamOutputComposure.h
--- a/Engine/Plugins/Experimental/VirtualCamera/Source/VCamCore/Public/VCamOutputComposure.h
+++ b/Engine/Plugins/Experimental/VirtualCamera/Source/VCamCore/Public/VCamOutputComposure.h
@@ -28,7 +28,7 @@ public:
 
 	// TextureRenderTarget2D asset that contains the final output
 	UPROPERTY(EditAnywhere, BlueprintReadWrite, Category = "Output")
-	UTextureRenderTarget2D* FinalOutputRenderTarget = nullptr;
+	TObjectPtr<UTextureRenderTarget2D> FinalOutputRenderTarget = nullptr;
 
 protected:
 	virtual void CreateUMG() override;
diff --git a/Engine/Plugins/Experimental/VirtualCamera/Source/VCamCore/Public/VCamOutputMediaOutput.h b/Engine/Plugins/Experimental/VirtualCamera/Source/VCamCore/Public/VCamOutputMediaOutput.h
--- a/Engine/Plugins/Experimental/VirtualCamera/Source/VCamCore/Public/VCamOutputMediaOutput.h
+++ b/Engine/Plugins/Experimental/VirtualCamera/Source/VCamCore/Public/VCamOutputMediaOutput.h
@@ -26,7 +26,7 @@ public:
 
 	// Media Output configuration asset
 	UPROPERTY(EditAnywhere, BlueprintReadWrite, Category = "Output")
-	UMediaOutput* OutputConfig;
+	TObjectPtr<UMediaOutput> OutputConfig;
 
 	// If using the output from a Composure Output Provider, specify it here
 	UPROPERTY(EditAnywhere, BlueprintReadWrite, Category = "Output")
@@ -34,7 +34,7 @@ public:
 
 protected:
 	UPROPERTY(Transient)
-	UMediaCapture* MediaCapture = nullptr;
+	TObjectPtr<UMediaCapture> MediaCapture = nullptr;
 
 private:
 	void StartCapturing();
diff --git a/Engine/Plugins/Experimental/VirtualCamera/Source/VCamCore/Public/VCamOutputProviderBase.h b/Engine/Plugins/Experimental/VirtualCamera/Source/VCamCore/Public/VCamOutputProviderBase.h
--- a/Engine/Plugins/Experimental/VirtualCamera/Source/VCamCore/Public/VCamOutputProviderBase.h
+++ b/Engine/Plugins/Experimental/VirtualCamera/Source/VCamCore/Public/VCamOutputProviderBase.h
@@ -135,7 +135,7 @@ protected:
 #endif
 
 	UPROPERTY(Transient)
-	UVPFullScreenUserWidget* UMGWidget = nullptr;
+	TObjectPtr<UVPFullScreenUserWidget> UMGWidget = nullptr;
 
 private:
 	bool IsOuterComponentEnabled() const;
diff --git a/Engine/Plugins/Experimental/VirtualCamera/Source/VCamCore/Public/VCamOutputRemoteSession.h b/Engine/Plugins/Experimental/VirtualCamera/Source/VCamCore/Public/VCamOutputRemoteSession.h
--- a/Engine/Plugins/Experimental/VirtualCamera/Source/VCamCore/Public/VCamOutputRemoteSession.h
+++ b/Engine/Plugins/Experimental/VirtualCamera/Source/VCamCore/Public/VCamOutputRemoteSession.h
@@ -44,10 +44,10 @@ public:
 protected:
 
 	UPROPERTY(Transient)
-	URemoteSessionMediaOutput* MediaOutput = nullptr;
+	TObjectPtr<URemoteSessionMediaOutput> MediaOutput = nullptr;
 
 	UPROPERTY(Transient)
-	URemoteSessionMediaCapture* MediaCapture = nullptr;
+	TObjectPtr<URemoteSessionMediaCapture> MediaCapture = nullptr;
 
 private:
 	TSharedPtr<IRemoteSessionUnmanagedRole> RemoteSessionHost;
diff --git a/Engine/Plugins/Experimental/VirtualCamera/Source/VCamCore/Public/VCamTypes.h b/Engine/Plugins/Experimental/VirtualCamera/Source/VCamCore/Public/VCamTypes.h
--- a/Engine/Plugins/Experimental/VirtualCamera/Source/VCamCore/Public/VCamTypes.h
+++ b/Engine/Plugins/Experimental/VirtualCamera/Source/VCamCore/Public/VCamTypes.h
@@ -28,7 +28,7 @@ struct FModifierStackEntry
 
 	// The current generated modifier instance
 	UPROPERTY(EditAnywhere, Instanced, Category="Modifier")
-    UVCamModifier* GeneratedModifier = nullptr;
+    TObjectPtr<UVCamModifier> GeneratedModifier = nullptr;
 
 #if WITH_EDITOR
 	// GUID used in the editor to identify specific stack entries during editor operations
diff --git a/Engine/Plugins/Experimental/VirtualCamera/Source/VirtualCamera/Public/LevelSequencePlaybackController.h b/Engine/Plugins/Experimental/VirtualCamera/Source/VirtualCamera/Public/LevelSequencePlaybackController.h
--- a/Engine/Plugins/Experimental/VirtualCamera/Source/VirtualCamera/Public/LevelSequencePlaybackController.h
+++ b/Engine/Plugins/Experimental/VirtualCamera/Source/VirtualCamera/Public/LevelSequencePlaybackController.h
@@ -174,7 +174,7 @@ protected:
 
 	/** The sequence to play back */
 	UPROPERTY(Transient)
-	ULevelSequence* ActiveLevelSequence;
+	TObjectPtr<ULevelSequence> ActiveLevelSequence;
 
 #if WITH_EDITORONLY_DATA
 	/** Weak reference to Sequencer associated with the active LevelSequence */
diff --git a/Engine/Plugins/Experimental/VirtualCamera/Source/VirtualCamera/Public/VirtualCameraActor.h b/Engine/Plugins/Experimental/VirtualCamera/Source/VirtualCamera/Public/VirtualCameraActor.h
--- a/Engine/Plugins/Experimental/VirtualCamera/Source/VirtualCamera/Public/VirtualCameraActor.h
+++ b/Engine/Plugins/Experimental/VirtualCamera/Source/VirtualCamera/Public/VirtualCameraActor.h
@@ -46,16 +46,16 @@ public:
 	~AVirtualCameraActor();
 
 	UPROPERTY(VisibleAnywhere, BlueprintReadWrite, Category = "VirtualCamera | Component")
-	USceneCaptureComponent2D* SceneCaptureComponent;
+	TObjectPtr<USceneCaptureComponent2D> SceneCaptureComponent;
 
 	UPROPERTY(EditAnywhere, BlueprintReadOnly, Category = "VirtualCamera | Movement")
 	FLiveLinkSubjectRepresentation LiveLinkSubject;
 
 	UPROPERTY(VisibleAnywhere, BlueprintReadWrite, Category = "VirtualCamera | Movement")
-	UVirtualCameraMovement* MovementComponent;
+	TObjectPtr<UVirtualCameraMovement> MovementComponent;
 
 	UPROPERTY(Transient, EditDefaultsOnly, Category = "VirtualCamera | MediaOutput")
-	URemoteSessionMediaOutput* MediaOutput;
+	TObjectPtr<URemoteSessionMediaOutput> MediaOutput;
 
 	UPROPERTY(EditAnywhere, BlueprintReadWrite, Category = "VirtualCamera | UMG")
 	TSubclassOf<UUserWidget> CameraUMGClass;
@@ -69,16 +69,16 @@ public:
 protected:
 
 	UPROPERTY(EditDefaultsOnly, Category = "VirtualCamera | UMG")
-	UVPFullScreenUserWidget* CameraScreenWidget;
+	TObjectPtr<UVPFullScreenUserWidget> CameraScreenWidget;
 
 	UPROPERTY(Transient, EditDefaultsOnly, Category = "VirtualCamera | MediaOutput")
-	URemoteSessionMediaCapture* MediaCapture;
+	TObjectPtr<URemoteSessionMediaCapture> MediaCapture;
 
 	UPROPERTY(Transient)
-	UWorld* ActorWorld;
+	TObjectPtr<UWorld> ActorWorld;
 
 	UPROPERTY(Transient)
-	AActor* PreviousViewTarget;
+	TObjectPtr<AActor> PreviousViewTarget;
 
 	/** Should focus plane be shown on all touch focus events */
 	UPROPERTY(BlueprintReadOnly, Category = "VirtualCamera | Focus")
@@ -196,7 +196,7 @@ private:
 private:
 
 	UPROPERTY(BlueprintGetter = GetCineCameraComponent, Category = "VirtualCamera | Component")
-	UCineCameraComponent* StreamedCamera;
+	TObjectPtr<UCineCameraComponent> StreamedCamera;
 	bool bIsStreaming;
 	TUniquePtr<FVirtualCameraViewportSettings> ViewportSettingsBackup;
 	FHitResult LastViewportTouchResult;
diff --git a/Engine/Plugins/Experimental/VirtualCamera/Source/VirtualCamera/Public/VirtualCameraConcertCameraComponent.h b/Engine/Plugins/Experimental/VirtualCamera/Source/VirtualCamera/Public/VirtualCameraConcertCameraComponent.h
--- a/Engine/Plugins/Experimental/VirtualCamera/Source/VirtualCamera/Public/VirtualCameraConcertCameraComponent.h
+++ b/Engine/Plugins/Experimental/VirtualCamera/Source/VirtualCamera/Public/VirtualCameraConcertCameraComponent.h
@@ -43,7 +43,7 @@ public:
 
 	/** The camera component that will be tracked. */
 	UPROPERTY(BlueprintReadWrite, EditAnywhere, Category = "VirtualCamera")
-	UCineCameraComponent* ComponentToTrack;
+	TObjectPtr<UCineCameraComponent> ComponentToTrack;
 
 public:
 	virtual void BeginPlay() override;
diff --git a/Engine/Plugins/Experimental/VirtualCamera/Source/VirtualCamera/Public/VirtualCameraMovementComponent.h b/Engine/Plugins/Experimental/VirtualCamera/Source/VirtualCamera/Public/VirtualCameraMovementComponent.h
--- a/Engine/Plugins/Experimental/VirtualCamera/Source/VirtualCamera/Public/VirtualCameraMovementComponent.h
+++ b/Engine/Plugins/Experimental/VirtualCamera/Source/VirtualCamera/Public/VirtualCameraMovementComponent.h
@@ -204,13 +204,13 @@ private:
 	 * @see UpdatedComponent
 	 */
 	UPROPERTY(Transient, DuplicateTransient)
-	USceneComponent* RootUpdatedComponent;
+	TObjectPtr<USceneComponent> RootUpdatedComponent;
 
 	/**
 	 * RootUpdatedComponent, cast as a UPrimitiveComponent. May be invalid if RootUpdatedComponent was null or not a UPrimitiveComponent.
 	 */
 	UPROPERTY(Transient, DuplicateTransient)
-	UPrimitiveComponent* RootUpdatedPrimitive;
+	TObjectPtr<UPrimitiveComponent> RootUpdatedPrimitive;
 
 	/**
 	 * Applies relative axis filtering based on locks for location.
diff --git a/Engine/Plugins/Experimental/VirtualCamera/Source/VirtualCamera/Public/VirtualCameraPawnBase.h b/Engine/Plugins/Experimental/VirtualCamera/Source/VirtualCamera/Public/VirtualCameraPawnBase.h
--- a/Engine/Plugins/Experimental/VirtualCamera/Source/VirtualCamera/Public/VirtualCameraPawnBase.h
+++ b/Engine/Plugins/Experimental/VirtualCamera/Source/VirtualCamera/Public/VirtualCameraPawnBase.h
@@ -30,7 +30,7 @@ public:
 
 	/** Cinematic camera used for view */
 	UPROPERTY(VisibleAnywhere, BlueprintReadWrite, Category = "Cinematic Camera")
-	UVirtualCameraCineCameraComponent* CineCamera;
+	TObjectPtr<UVirtualCameraCineCameraComponent> CineCamera;
 
 	/** Movement Component to handle the motion input for the camera */
 	UVirtualCameraMovementComponent* MovementComponent;
diff --git a/Engine/Plugins/Experimental/VirtualCamera/Source/VirtualCamera/Public/VirtualCameraPlayerControllerBase.h b/Engine/Plugins/Experimental/VirtualCamera/Source/VirtualCamera/Public/VirtualCameraPlayerControllerBase.h
--- a/Engine/Plugins/Experimental/VirtualCamera/Source/VirtualCamera/Public/VirtualCameraPlayerControllerBase.h
+++ b/Engine/Plugins/Experimental/VirtualCamera/Source/VirtualCamera/Public/VirtualCameraPlayerControllerBase.h
@@ -94,7 +94,7 @@ public:
 	FVirtualCameraResetOffsetsDelegate OnOffsetReset;
 
 	UPROPERTY(transient, BlueprintReadOnly, Category = "VirtualCamera")
-	AVPRootActor* RootActor;
+	TObjectPtr<AVPRootActor> RootActor;
 
 	/**
 	 * Overridable function to allow user to get tracker data from blueprints.
@@ -182,11 +182,11 @@ protected:
 
 	/** Controller for level sequence playback */
 	UPROPERTY(Transient)
-	ULevelSequencePlaybackController* LevelSequencePlaybackController;
+	TObjectPtr<ULevelSequencePlaybackController> LevelSequencePlaybackController;
 	
 	/** Target camera that is spawned or possessed on begin play for the sequence controller */
 	UPROPERTY(Transient)
-	ACineCameraActor* TargetCameraActor;
+	TObjectPtr<ACineCameraActor> TargetCameraActor;
 
 	/** Cached value for IsVirtualCameraControlledByRemoteSession() */
 	UPROPERTY(Transient, BlueprintReadOnly, Category = "VirtualCamera")
diff --git a/Engine/Plugins/Experimental/VirtualCamera/Source/VirtualCamera/Public/VirtualCameraSubsystem.h b/Engine/Plugins/Experimental/VirtualCamera/Source/VirtualCamera/Public/VirtualCameraSubsystem.h
--- a/Engine/Plugins/Experimental/VirtualCamera/Source/VirtualCamera/Public/VirtualCameraSubsystem.h
+++ b/Engine/Plugins/Experimental/VirtualCamera/Source/VirtualCamera/Public/VirtualCameraSubsystem.h
@@ -44,7 +44,7 @@ public:
 	void SetVirtualCameraController(TScriptInterface<IVirtualCameraController> VirtualCamera);
 
 	UPROPERTY(BlueprintReadOnly, Category = "VirtualCamera")
-	ULevelSequencePlaybackController* SequencePlaybackController;
+	TObjectPtr<ULevelSequencePlaybackController> SequencePlaybackController;
 
 	UPROPERTY(BlueprintAssignable, Category = "VirtualCamera | Streaming")
 	FOnStreamStarted OnStreamStartedDelegate;
diff --git a/Engine/Plugins/Experimental/VirtualHeightfieldMesh/Source/VirtualHeightfieldMesh/Private/HeightfieldMinMaxTextureMaterialExpression.h b/Engine/Plugins/Experimental/VirtualHeightfieldMesh/Source/VirtualHeightfieldMesh/Private/HeightfieldMinMaxTextureMaterialExpression.h
--- a/Engine/Plugins/Experimental/VirtualHeightfieldMesh/Source/VirtualHeightfieldMesh/Private/HeightfieldMinMaxTextureMaterialExpression.h
+++ b/Engine/Plugins/Experimental/VirtualHeightfieldMesh/Source/VirtualHeightfieldMesh/Private/HeightfieldMinMaxTextureMaterialExpression.h
@@ -17,7 +17,7 @@ class UMaterialExpressionHeightfieldMinMaxTexture : public UMaterialExpression
 	GENERATED_UCLASS_BODY()
 
 	UPROPERTY(EditAnywhere, BlueprintReadWrite, Category = MaterialExpressionTexture)
-	UHeightfieldMinMaxTexture* MinMaxTexture;
+	TObjectPtr<UHeightfieldMinMaxTexture> MinMaxTexture;
 
 	UPROPERTY(EditAnywhere, BlueprintReadWrite, Category = MaterialExpressionTexture)
 	TEnumAsByte<enum EMaterialSamplerType> SamplerType;
diff --git a/Engine/Plugins/Experimental/VirtualHeightfieldMesh/Source/VirtualHeightfieldMesh/Public/HeightfieldMinMaxTexture.h b/Engine/Plugins/Experimental/VirtualHeightfieldMesh/Source/VirtualHeightfieldMesh/Public/HeightfieldMinMaxTexture.h
--- a/Engine/Plugins/Experimental/VirtualHeightfieldMesh/Source/VirtualHeightfieldMesh/Public/HeightfieldMinMaxTexture.h
+++ b/Engine/Plugins/Experimental/VirtualHeightfieldMesh/Source/VirtualHeightfieldMesh/Public/HeightfieldMinMaxTexture.h
@@ -30,15 +30,15 @@ public:
 
 	/** The Height MinMax texture. */
 	UPROPERTY(VisibleAnywhere, BlueprintReadOnly, Category = Texture, meta = (DisplayName = "Height MinMax Texture"))
-	class UTexture2D* Texture = nullptr;
+	TObjectPtr<class UTexture2D> Texture = nullptr;
 
 	/** A LodBias texture derived from the Height MinMax texture*/
 	UPROPERTY(VisibleAnywhere, BlueprintReadOnly, Category = Texture, meta = (DisplayName = "LodBias Texture"))
-	class UTexture2D* LodBiasTexture = nullptr;
+	TObjectPtr<class UTexture2D> LodBiasTexture = nullptr;
 
 	/** A LodBias MinMax texture derived from the LodBias texture*/
 	UPROPERTY(VisibleAnywhere, BlueprintReadOnly, Category = Texture, meta = (DisplayName = "LodBias MinMax Texture"))
-	class UTexture2D* LodBiasMinMaxTexture = nullptr;
+	TObjectPtr<class UTexture2D> LodBiasMinMaxTexture = nullptr;
 
 protected:
 	/** The number of mip levels to clone for CPU access. */
diff --git a/Engine/Plugins/Experimental/VirtualHeightfieldMesh/Source/VirtualHeightfieldMesh/Public/VirtualHeightfieldMeshActor.h b/Engine/Plugins/Experimental/VirtualHeightfieldMesh/Source/VirtualHeightfieldMesh/Public/VirtualHeightfieldMeshActor.h
--- a/Engine/Plugins/Experimental/VirtualHeightfieldMesh/Source/VirtualHeightfieldMesh/Public/VirtualHeightfieldMeshActor.h
+++ b/Engine/Plugins/Experimental/VirtualHeightfieldMesh/Source/VirtualHeightfieldMesh/Public/VirtualHeightfieldMeshActor.h
@@ -14,12 +14,12 @@ class AVirtualHeightfieldMesh : public AActor
 private:
 	/** Component for rendering the big mesh. */
 	UPROPERTY(VisibleAnywhere, BlueprintReadOnly, Category = VirtualTexture, meta = (AllowPrivateAccess = "true"))
-	class UVirtualHeightfieldMeshComponent* VirtualHeightfieldMeshComponent;
+	TObjectPtr<class UVirtualHeightfieldMeshComponent> VirtualHeightfieldMeshComponent;
 
 #if WITH_EDITORONLY_DATA
 	/** Box for visualizing virtual texture extents. */
 	UPROPERTY(Transient)
-	class UBoxComponent* Box = nullptr;
+	TObjectPtr<class UBoxComponent> Box = nullptr;
 #endif // WITH_EDITORONLY_DATA
 
 protected:
diff --git a/Engine/Plugins/Experimental/VirtualHeightfieldMesh/Source/VirtualHeightfieldMesh/Public/VirtualHeightfieldMeshComponent.h b/Engine/Plugins/Experimental/VirtualHeightfieldMesh/Source/VirtualHeightfieldMesh/Public/VirtualHeightfieldMeshComponent.h
--- a/Engine/Plugins/Experimental/VirtualHeightfieldMesh/Source/VirtualHeightfieldMesh/Public/VirtualHeightfieldMeshComponent.h
+++ b/Engine/Plugins/Experimental/VirtualHeightfieldMesh/Source/VirtualHeightfieldMesh/Public/VirtualHeightfieldMeshComponent.h
@@ -24,11 +24,11 @@ protected:
 	
 	/** UObject ref resolved from VirtualTexture weak ref. */
 	UPROPERTY(Transient)
-	ARuntimeVirtualTextureVolume* VirtualTextureRef;
+	TObjectPtr<ARuntimeVirtualTextureVolume> VirtualTextureRef;
 
 	/** Placeholder for details customization image. */
 	UPROPERTY(VisibleAnywhere, Transient, Category = Heightfield)
-	UObject* VirtualTextureThumbnail = nullptr;
+	TObjectPtr<UObject> VirtualTextureThumbnail = nullptr;
 
 	/** Placeholder for details customization button. */
 	UPROPERTY(VisibleAnywhere, Transient, Category = Heightfield)
@@ -36,7 +36,7 @@ protected:
 
 	/** Texture object containing minimum and maximum height values. */
 	UPROPERTY(EditAnywhere, BlueprintReadOnly, Category = HeightfieldBuild, meta = (DisplayName = "MinMax Texture"))
-	UHeightfieldMinMaxTexture* MinMaxTexture = nullptr;
+	TObjectPtr<UHeightfieldMinMaxTexture> MinMaxTexture = nullptr;
 
 	/** Number of levels to build in the MinMax Texture. A default value of 0 will build all levels from the heightfield. */
 	UPROPERTY(EditAnywhere, BlueprintReadOnly, Category = HeightfieldBuild, meta = (DisplayName = "MinMax Build Levels", UIMin = "0"))
@@ -48,7 +48,7 @@ protected:
 
 	/** The material to apply. */
 	UPROPERTY(EditAnywhere, Category = Rendering)
-	UMaterialInterface* Material = nullptr;
+	TObjectPtr<UMaterialInterface> Material = nullptr;
 
 	/** Target screen size for a LOD 0 tile. A smaller value uniformly increases the geometry resolution on screen. */
 	UPROPERTY(EditAnywhere, Category = Rendering, meta = (DisplayName = "LOD 0 Screen Size", ClampMin = "0.1", UIMin = "0.1"))
diff --git a/Engine/Plugins/Experimental/VirtualProductionUtilities/Source/VPUtilities/Public/VPFullScreenUserWidget.h b/Engine/Plugins/Experimental/VirtualProductionUtilities/Source/VPUtilities/Public/VPFullScreenUserWidget.h
--- a/Engine/Plugins/Experimental/VirtualProductionUtilities/Source/VPUtilities/Public/VPFullScreenUserWidget.h
+++ b/Engine/Plugins/Experimental/VirtualProductionUtilities/Source/VPUtilities/Public/VPFullScreenUserWidget.h
@@ -103,7 +103,7 @@ public:
 	 * OpacityFromTexture [Scalar]
 	 */
 	UPROPERTY(EditAnywhere, Category = PostProcess)
-	UMaterialInterface* PostProcessMaterial;
+	TObjectPtr<UMaterialInterface> PostProcessMaterial;
 
 	/** Tint color and opacity for this component. */
 	UPROPERTY(EditAnywhere, Category = PostProcess)
@@ -143,11 +143,11 @@ public:
 
 	/** List of composure layers that are expecting to use the WidgetRenderTarget. */
 	UPROPERTY(EditAnywhere, Category= PostProcess)
-    TArray<ACompositingElement*> ComposureLayerTargets;
+    TArray<TObjectPtr<ACompositingElement>> ComposureLayerTargets;
 
 	/** The target to which the user widget is rendered. */
 	UPROPERTY(Transient)
-    UTextureRenderTarget2D* WidgetRenderTarget;
+    TObjectPtr<UTextureRenderTarget2D> WidgetRenderTarget;
 
 #if WITH_EDITOR
 	/** If set, use this viewport instead of GetFirstActiveLevelViewport() */
@@ -156,11 +156,11 @@ public:
 private:
 	/** Post process component used to add the material to the post process chain. */
 	UPROPERTY(Transient)
-	UPostProcessComponent* PostProcessComponent;
+	TObjectPtr<UPostProcessComponent> PostProcessComponent;
 
 	/** The dynamic instance of the material that the render target is attached to. */
 	UPROPERTY(Transient)
-	UMaterialInstanceDynamic* PostProcessMaterialInstance;
+	TObjectPtr<UMaterialInstanceDynamic> PostProcessMaterialInstance;
 
 	/** The slate window that contains the user widget content. */
 	TSharedPtr<SVirtualWindow> SlateWindow;
@@ -264,7 +264,7 @@ public:
 private:
 	/** The User Widget object displayed and managed by this component */
 	UPROPERTY(Transient, DuplicateTransient)
-	UUserWidget* Widget;
+	TObjectPtr<UUserWidget> Widget;
 
 	/** The world the widget is attached to. */
 	TWeakObjectPtr<UWorld> World;
diff --git a/Engine/Plugins/Experimental/VirtualProductionUtilities/Source/VPUtilities/Public/VPFullScreenUserWidgetActor.h b/Engine/Plugins/Experimental/VirtualProductionUtilities/Source/VPUtilities/Public/VPFullScreenUserWidgetActor.h
--- a/Engine/Plugins/Experimental/VirtualProductionUtilities/Source/VPUtilities/Public/VPFullScreenUserWidgetActor.h
+++ b/Engine/Plugins/Experimental/VirtualProductionUtilities/Source/VPUtilities/Public/VPFullScreenUserWidgetActor.h
@@ -42,7 +42,7 @@ private:
 protected:
 	/** */
 	UPROPERTY(VisibleAnywhere, Instanced, NoClear, Category = "User Interface", meta = (ShowOnlyInnerProperties))
-	UVPFullScreenUserWidget* ScreenUserWidget;
+	TObjectPtr<UVPFullScreenUserWidget> ScreenUserWidget;
 
 #if WITH_EDITORONLY_DATA
 	/** Display requested and will be executed on the first frame because we can't call BP function in the loading phase */
diff --git a/Engine/Plugins/Experimental/VirtualProductionUtilities/Source/VPUtilities/Public/VPGameMode.h b/Engine/Plugins/Experimental/VirtualProductionUtilities/Source/VPUtilities/Public/VPGameMode.h
--- a/Engine/Plugins/Experimental/VirtualProductionUtilities/Source/VPUtilities/Public/VPGameMode.h
+++ b/Engine/Plugins/Experimental/VirtualProductionUtilities/Source/VPUtilities/Public/VPGameMode.h
@@ -22,7 +22,7 @@ class VPUTILITIES_API AVPGameMode : public AGameMode
 
 private:
 	UPROPERTY(transient, BlueprintGetter = "GetRootActor", Category = "Virtual Production")
-	AVPRootActor* RootActor;
+	TObjectPtr<AVPRootActor> RootActor;
 
 public:
 	UFUNCTION(BlueprintGetter, Category = "Virtual Production")
diff --git a/Engine/Plugins/Experimental/VirtualProductionUtilities/Source/VPUtilities/Public/VPRootActor.h b/Engine/Plugins/Experimental/VirtualProductionUtilities/Source/VPUtilities/Public/VPRootActor.h
--- a/Engine/Plugins/Experimental/VirtualProductionUtilities/Source/VPUtilities/Public/VPRootActor.h
+++ b/Engine/Plugins/Experimental/VirtualProductionUtilities/Source/VPUtilities/Public/VPRootActor.h
@@ -30,11 +30,11 @@ public:
 public:
 	/** Actor used to show the size of the scene. */
 	UPROPERTY(EditDefaultsOnly, Category = "Virtual Production")
-	UMeshComponent* RealWorldSceneRepresentation;
+	TObjectPtr<UMeshComponent> RealWorldSceneRepresentation;
 
 	/** Actor used to show the size of the scene. */
 	UPROPERTY(EditAnywhere, Category = "Virtual Production")
-	ACineCameraActor* CinematicCamera;
+	TObjectPtr<ACineCameraActor> CinematicCamera;
 
 private:
 	UPROPERTY()
@@ -44,11 +44,11 @@ private:
 protected:
 	/** Billboard used to see the scene in the editor */
 	UPROPERTY()
-	UBillboardComponent* SpriteComponent;
+	TObjectPtr<UBillboardComponent> SpriteComponent;
 
 	/** Arrow used to see the orientation of the scene in the editor */
 	UPROPERTY()
-	UArrowComponent* ArrowComponent;
+	TObjectPtr<UArrowComponent> ArrowComponent;
 
 public:
 	UPROPERTY(EditAnywhere, Category = "Virtual Production")
diff --git a/Engine/Plugins/Experimental/VirtualProductionUtilities/Source/VPUtilitiesEditor/Private/VPCameraUIBase.h b/Engine/Plugins/Experimental/VirtualProductionUtilities/Source/VPUtilitiesEditor/Private/VPCameraUIBase.h
--- a/Engine/Plugins/Experimental/VirtualProductionUtilities/Source/VPUtilitiesEditor/Private/VPCameraUIBase.h
+++ b/Engine/Plugins/Experimental/VirtualProductionUtilities/Source/VPUtilitiesEditor/Private/VPCameraUIBase.h
@@ -18,10 +18,10 @@ class UVPCameraUIBase : public UUserWidget
 
 protected:
 	UPROPERTY(Transient)
-	ACameraActor* SelectedCamera;
+	TObjectPtr<ACameraActor> SelectedCamera;
 
 	UPROPERTY(BlueprintReadOnly, Transient, Category = "Camera")
-	UCameraComponent* SelectedCameraComponent;
+	TObjectPtr<UCameraComponent> SelectedCameraComponent;
 
 protected:
 	virtual bool Initialize() override;
diff --git a/Engine/Plugins/Experimental/VirtualProductionUtilities/Source/VPUtilitiesEditor/Private/VPUIBase.h b/Engine/Plugins/Experimental/VirtualProductionUtilities/Source/VPUtilitiesEditor/Private/VPUIBase.h
--- a/Engine/Plugins/Experimental/VirtualProductionUtilities/Source/VPUtilitiesEditor/Private/VPUIBase.h
+++ b/Engine/Plugins/Experimental/VirtualProductionUtilities/Source/VPUtilitiesEditor/Private/VPUIBase.h
@@ -21,7 +21,7 @@ class UVPUIBase : public UEditorUtilityWidget
 
 public:
 	UPROPERTY(BlueprintReadOnly, Transient, Category = "Virtual Production")
-	AActor* SelectedActor;
+	TObjectPtr<AActor> SelectedActor;
 	
 	UFUNCTION(BlueprintCallable, Category = "VirtualProduction")
 	void AppendVirtualProductionLog(FString NewMessage);
diff --git a/Engine/Plugins/Experimental/VirtualProductionUtilities/Source/VPUtilitiesEditor/Public/VPScoutingSubsystem.h b/Engine/Plugins/Experimental/VirtualProductionUtilities/Source/VPUtilitiesEditor/Public/VPScoutingSubsystem.h
--- a/Engine/Plugins/Experimental/VirtualProductionUtilities/Source/VPUtilitiesEditor/Public/VPScoutingSubsystem.h
+++ b/Engine/Plugins/Experimental/VirtualProductionUtilities/Source/VPUtilitiesEditor/Public/VPScoutingSubsystem.h
@@ -86,11 +86,11 @@ public:
 
 	/** Subsystems can't have any Blueprint implementations, so we attach this class for any BP logic that we to provide. */
 	UPROPERTY(Transient, BlueprintReadOnly, Category = "Virtual Production")
-	UVPScoutingSubsystemHelpersBase* VPSubsystemHelpers;
+	TObjectPtr<UVPScoutingSubsystemHelpersBase> VPSubsystemHelpers;
 
 	/** GestureManager that manage some user input in VR editor. */
 	UPROPERTY(Transient, BlueprintReadOnly, Category = "Virtual Production")
-	UVPScoutingSubsystemGestureManagerBase* GestureManager;
+	TObjectPtr<UVPScoutingSubsystemGestureManagerBase> GestureManager;
 
 	/** bool to keep track of whether the settings menu panel in the main menu is open*/
 	UPROPERTY(EditAnywhere, BlueprintReadWrite, Category = "Main Menu")
diff --git a/Engine/Plugins/Experimental/Water/Source/Editor/Private/JumpFloodComponent2D.h b/Engine/Plugins/Experimental/Water/Source/Editor/Private/JumpFloodComponent2D.h
--- a/Engine/Plugins/Experimental/Water/Source/Editor/Private/JumpFloodComponent2D.h
+++ b/Engine/Plugins/Experimental/Water/Source/Editor/Private/JumpFloodComponent2D.h
@@ -40,13 +40,13 @@ public:
 
 public:
 	UPROPERTY(EditInstanceOnly, AdvancedDisplay, BlueprintReadWrite, meta = (Category = "Materials"))
-	UMaterialInterface* JumpStepMaterial = nullptr;
+	TObjectPtr<UMaterialInterface> JumpStepMaterial = nullptr;
 
 	UPROPERTY(EditInstanceOnly, AdvancedDisplay, BlueprintReadWrite, meta = (Category = "Materials"))
-	UMaterialInterface* FindEdgesMaterial = nullptr;
+	TObjectPtr<UMaterialInterface> FindEdgesMaterial = nullptr;
 
 	UPROPERTY(EditInstanceOnly, AdvancedDisplay, BlueprintReadWrite, meta = (Category = "Materials"))
-	UMaterialInterface* BlurEdgesMaterial = nullptr;
+	TObjectPtr<UMaterialInterface> BlurEdgesMaterial = nullptr;
 
 	UPROPERTY(EditAnywhere, BlueprintReadWrite, AdvancedDisplay, meta=(Category="Default"))
 	bool UseBlur = false;
@@ -56,20 +56,20 @@ public:
 
 private:
 	UPROPERTY(VisibleAnywhere, AdvancedDisplay, meta = (Category = "Default"))
-	UTextureRenderTarget2D* RTA;
+	TObjectPtr<UTextureRenderTarget2D> RTA;
 
 	UPROPERTY(VisibleAnywhere, AdvancedDisplay, meta = (Category = "Default"))
-	UTextureRenderTarget2D* RTB;
+	TObjectPtr<UTextureRenderTarget2D> RTB;
 
 	// Transient properties (exposed only for debugging reasons) :
 	UPROPERTY(VisibleAnywhere, AdvancedDisplay, Transient, meta = (Category = "Debug"))
-	UMaterialInstanceDynamic* JumpStepMID;
+	TObjectPtr<UMaterialInstanceDynamic> JumpStepMID;
 
 	UPROPERTY(VisibleAnywhere, AdvancedDisplay, Transient, meta = (Category = "Debug"))
-	UMaterialInstanceDynamic* FindEdgesMID;
+	TObjectPtr<UMaterialInstanceDynamic> FindEdgesMID;
 
 	UPROPERTY(VisibleAnywhere, AdvancedDisplay, Transient, meta = (Category = "Debug"))
-	UMaterialInstanceDynamic* BlurEdgesMID;
+	TObjectPtr<UMaterialInstanceDynamic> BlurEdgesMID;
 
 	UPROPERTY(VisibleAnywhere, AdvancedDisplay, Transient, meta = (Category = "Debug"))
 	int32 RequiredPasses = 0;
diff --git a/Engine/Plugins/Experimental/Water/Source/Editor/Private/WaterBrushCacheContainer.h b/Engine/Plugins/Experimental/Water/Source/Editor/Private/WaterBrushCacheContainer.h
--- a/Engine/Plugins/Experimental/Water/Source/Editor/Private/WaterBrushCacheContainer.h
+++ b/Engine/Plugins/Experimental/Water/Source/Editor/Private/WaterBrushCacheContainer.h
@@ -16,7 +16,7 @@ struct FWaterBodyBrushCache
 	GENERATED_USTRUCT_BODY()
 
 	UPROPERTY(EditAnywhere, BlueprintReadWrite, Category = "Cache")
-	UTextureRenderTarget2D* CacheRenderTarget = nullptr;
+	TObjectPtr<UTextureRenderTarget2D> CacheRenderTarget = nullptr;
 
 	UPROPERTY(EditAnywhere, BlueprintReadWrite, Category = "Cache")
 	bool CacheIsValid = false;
diff --git a/Engine/Plugins/Experimental/Water/Source/Editor/Private/WaterBrushManager.h b/Engine/Plugins/Experimental/Water/Source/Editor/Private/WaterBrushManager.h
--- a/Engine/Plugins/Experimental/Water/Source/Editor/Private/WaterBrushManager.h
+++ b/Engine/Plugins/Experimental/Water/Source/Editor/Private/WaterBrushManager.h
@@ -33,122 +33,122 @@ public:
 	GENERATED_BODY()
 	
  	UPROPERTY()
-	USceneCaptureComponent2D* SceneCaptureComponent2D = nullptr;
+	TObjectPtr<USceneCaptureComponent2D> SceneCaptureComponent2D = nullptr;
 
 	UPROPERTY()
-	UJumpFloodComponent2D* JumpFloodComponent2D = nullptr;
+	TObjectPtr<UJumpFloodComponent2D> JumpFloodComponent2D = nullptr;
 
 	// RTs
 	UPROPERTY(VisibleAnywhere, Transient, meta = (Category = "Render Targets"))
-	UTextureRenderTarget2D* HeightmapRTA = nullptr;
+	TObjectPtr<UTextureRenderTarget2D> HeightmapRTA = nullptr;
 
 	UPROPERTY(VisibleAnywhere, Transient, meta=(Category="Render Targets"))
-	UTextureRenderTarget2D* HeightmapRTB = nullptr;
+	TObjectPtr<UTextureRenderTarget2D> HeightmapRTB = nullptr;
 
 	UPROPERTY(VisibleAnywhere, Transient, meta=(Category="Render Targets"))
-	UTextureRenderTarget2D* JumpFloodRTA;
+	TObjectPtr<UTextureRenderTarget2D> JumpFloodRTA;
 
 	UPROPERTY(VisibleAnywhere, Transient, meta=(Category="Render Targets"))
-	UTextureRenderTarget2D* JumpFloodRTB = nullptr;
+	TObjectPtr<UTextureRenderTarget2D> JumpFloodRTB = nullptr;
 
 	UPROPERTY(VisibleAnywhere, Transient, meta=(Category="Render Targets"))
-	UTextureRenderTarget2D* DepthAndShapeRT = nullptr;
+	TObjectPtr<UTextureRenderTarget2D> DepthAndShapeRT = nullptr;
 
 	UPROPERTY(VisibleAnywhere, Transient, meta=(Category="Render Targets"))
-	UTextureRenderTarget2D* WaterDepthAndVelocityRT = nullptr;
+	TObjectPtr<UTextureRenderTarget2D> WaterDepthAndVelocityRT = nullptr;
 
 	UPROPERTY(VisibleAnywhere, Transient, meta=(Category="Render Targets"))
-	UTextureRenderTarget2D* CombinedVelocityAndHeightRTA = nullptr;
+	TObjectPtr<UTextureRenderTarget2D> CombinedVelocityAndHeightRTA = nullptr;
 
 	UPROPERTY(VisibleAnywhere, Transient, meta=(Category="Render Targets"))
-	UTextureRenderTarget2D* CombinedVelocityAndHeightRTB = nullptr;
+	TObjectPtr<UTextureRenderTarget2D> CombinedVelocityAndHeightRTB = nullptr;
 
 	UPROPERTY(VisibleAnywhere, Transient, meta=(Category="Render Targets"))
-	UTextureRenderTarget2D* LandscapeRTRef = nullptr;
+	TObjectPtr<UTextureRenderTarget2D> LandscapeRTRef = nullptr;
 
 	UPROPERTY(VisibleAnywhere, Transient, meta=(Category="Render Targets"))
-	UTextureRenderTarget2D* WeightmapRTA = nullptr;
+	TObjectPtr<UTextureRenderTarget2D> WeightmapRTA = nullptr;
 	
 	UPROPERTY(VisibleAnywhere, Transient, meta=(Category="Render Targets"))
-	UTextureRenderTarget2D* WeightmapRTB = nullptr;
+	TObjectPtr<UTextureRenderTarget2D> WeightmapRTB = nullptr;
 	// RTs End
 
 	// Brush materials
 	UPROPERTY(EditAnywhere, AdvancedDisplay, BlueprintReadWrite, meta = (Category = "Brush Materials"))
-	UMaterialInterface* BrushAngleFalloffMaterial = nullptr;
+	TObjectPtr<UMaterialInterface> BrushAngleFalloffMaterial = nullptr;
 
 	UPROPERTY(EditAnywhere, AdvancedDisplay, BlueprintReadWrite, meta = (Category = "Brush Materials"))
-	UMaterialInterface* BrushWidthFalloffMaterial = nullptr;
+	TObjectPtr<UMaterialInterface> BrushWidthFalloffMaterial = nullptr;
 
 	UPROPERTY(EditAnywhere, AdvancedDisplay, BlueprintReadWrite, meta = (Category = "Brush Materials"))
-	UMaterialInterface* DistanceFieldCacheMaterial = nullptr;
+	TObjectPtr<UMaterialInterface> DistanceFieldCacheMaterial = nullptr;
 
 	UPROPERTY(EditAnywhere, AdvancedDisplay, BlueprintReadWrite, meta = (Category = "Brush Materials"))
-	UMaterialInterface* RenderRiverSplineDepthMaterial = nullptr;
+	TObjectPtr<UMaterialInterface> RenderRiverSplineDepthMaterial = nullptr;
 
 	// TODO [jonathan.bard] : rename to DebugDistanceFieldMaterial and make it work : 
 	UPROPERTY(EditAnywhere, AdvancedDisplay, BlueprintReadWrite, meta = (Category = "Brush Materials"))
-	UMaterialInterface* DebugDistanceFieldMaterial = nullptr;
+	TObjectPtr<UMaterialInterface> DebugDistanceFieldMaterial = nullptr;
 
 	UPROPERTY(EditAnywhere, AdvancedDisplay, BlueprintReadWrite, meta = (Category = "Brush Materials"))
-	UMaterialInterface* WeightmapMaterial = nullptr;
+	TObjectPtr<UMaterialInterface> WeightmapMaterial = nullptr;
 
 	UPROPERTY(EditAnywhere, AdvancedDisplay, BlueprintReadWrite, meta = (Category = "Brush Materials"))
-	UMaterialInterface* DrawCanvasMaterial = nullptr;
+	TObjectPtr<UMaterialInterface> DrawCanvasMaterial = nullptr;
 
 	UPROPERTY(EditAnywhere, AdvancedDisplay, BlueprintReadWrite, meta = (Category = "Brush Materials"))
-	UMaterialInterface* CompositeWaterBodyTextureMaterial = nullptr;
+	TObjectPtr<UMaterialInterface> CompositeWaterBodyTextureMaterial = nullptr;
 
 	UPROPERTY(EditAnywhere, AdvancedDisplay, BlueprintReadWrite, meta = (Category = "Brush Materials"))
-	UMaterialInterface* IslandFalloffMaterial = nullptr;
+	TObjectPtr<UMaterialInterface> IslandFalloffMaterial = nullptr;
 
 	UPROPERTY(EditAnywhere, AdvancedDisplay, BlueprintReadWrite, meta = (Category = "Brush Materials"))
-	UMaterialInterface* FinalizeVelocityHeightMaterial = nullptr;
+	TObjectPtr<UMaterialInterface> FinalizeVelocityHeightMaterial = nullptr;
 
 	UPROPERTY(EditAnywhere, AdvancedDisplay, BlueprintReadWrite, meta = (Category = "Brush Materials"))
-	UMaterialInterface* JumpStepMaterial = nullptr;
+	TObjectPtr<UMaterialInterface> JumpStepMaterial = nullptr;
 
 	UPROPERTY(EditAnywhere, AdvancedDisplay, BlueprintReadWrite, meta = (Category = "Brush Materials"))
-	UMaterialInterface* FindEdgesMaterial = nullptr;
+	TObjectPtr<UMaterialInterface> FindEdgesMaterial = nullptr;
 
 	UPROPERTY(EditAnywhere, AdvancedDisplay, BlueprintReadWrite, meta = (Category = "Brush Materials"))
-	UMaterialInterface* BlurEdgesMaterial = nullptr;
+	TObjectPtr<UMaterialInterface> BlurEdgesMaterial = nullptr;
 	// Brush materials end
 
 	// MIDs
 	UPROPERTY(EditInstanceOnly, AdvancedDisplay, BlueprintReadWrite, Transient, meta=(Category="Debug MIDs"))
-	UMaterialInstanceDynamic* BrushAngleFalloffMID = nullptr;
+	TObjectPtr<UMaterialInstanceDynamic> BrushAngleFalloffMID = nullptr;
 
 	UPROPERTY(EditInstanceOnly, AdvancedDisplay, BlueprintReadWrite, Transient, meta=(Category="Debug MIDs"))
-	UMaterialInstanceDynamic* BrushWidthFalloffMID = nullptr;
+	TObjectPtr<UMaterialInstanceDynamic> BrushWidthFalloffMID = nullptr;
 
 	UPROPERTY(EditInstanceOnly, AdvancedDisplay, BlueprintReadWrite, Transient, meta=(Category="Debug MIDs"))
-	UMaterialInstanceDynamic* DistanceFieldCacheMID = nullptr;
+	TObjectPtr<UMaterialInstanceDynamic> DistanceFieldCacheMID = nullptr;
 
 	UPROPERTY(EditInstanceOnly, AdvancedDisplay, BlueprintReadWrite, Transient, meta = (Category = "Debug MIDs"))
-	TArray<UMaterialInstanceDynamic*> RiverSplineMIDs;
+	TArray<TObjectPtr<UMaterialInstanceDynamic>> RiverSplineMIDs;
 
 	UPROPERTY(EditInstanceOnly, AdvancedDisplay, BlueprintReadWrite, Transient, meta = (Category = "Debug MIDs"))
-	UMaterialInstanceDynamic* DebugDistanceFieldMID = nullptr;
+	TObjectPtr<UMaterialInstanceDynamic> DebugDistanceFieldMID = nullptr;
 
 	UPROPERTY(EditInstanceOnly, AdvancedDisplay, BlueprintReadWrite, Transient, meta = (Category = "Debug MIDs"))
-	UMaterialInstanceDynamic* WeightmapMID = nullptr;
+	TObjectPtr<UMaterialInstanceDynamic> WeightmapMID = nullptr;
 
 	UPROPERTY(EditInstanceOnly, AdvancedDisplay, BlueprintReadWrite, Transient, meta = (Category = "Debug MIDs"))
-	UMaterialInstanceDynamic* DrawCanvasMID = nullptr;
+	TObjectPtr<UMaterialInstanceDynamic> DrawCanvasMID = nullptr;
 
 	UPROPERTY(EditInstanceOnly, AdvancedDisplay, BlueprintReadWrite, Transient, meta = (Category = "Debug MIDs"))
-	UMaterialInstanceDynamic* CompositeWaterBodyTextureMID = nullptr;
+	TObjectPtr<UMaterialInstanceDynamic> CompositeWaterBodyTextureMID = nullptr;
 
 	UPROPERTY(EditInstanceOnly, AdvancedDisplay, BlueprintReadWrite, Transient, meta = (Category = "Debug MIDs"))
-	UMaterialInstanceDynamic* IslandFalloffMID = nullptr;
+	TObjectPtr<UMaterialInstanceDynamic> IslandFalloffMID = nullptr;
 
 	UPROPERTY(EditInstanceOnly, AdvancedDisplay, BlueprintReadWrite, Transient, meta = (Category = "Debug MIDs"))
-	UMaterialInstanceDynamic* FinalizeVelocityHeightMID = nullptr;
+	TObjectPtr<UMaterialInstanceDynamic> FinalizeVelocityHeightMID = nullptr;
 	// MIDs End
 
 	UPROPERTY(VisibleAnywhere, BlueprintReadWrite, Transient, meta=(Category="Debug"))
-	TMap<UCurveFloat*,FWaterBodyBrushCache> BrushCurveRTCache;
+	TMap<TObjectPtr<UCurveFloat>,FWaterBodyBrushCache> BrushCurveRTCache;
 	
 	UPROPERTY(VisibleAnywhere, BlueprintReadOnly, meta=(Category="Debug"))
 	FVector WorldSize;
diff --git a/Engine/Plugins/Experimental/Water/Source/Editor/Private/WaterLandscapeBrush.h b/Engine/Plugins/Experimental/Water/Source/Editor/Private/WaterLandscapeBrush.h
--- a/Engine/Plugins/Experimental/Water/Source/Editor/Private/WaterLandscapeBrush.h
+++ b/Engine/Plugins/Experimental/Water/Source/Editor/Private/WaterLandscapeBrush.h
@@ -105,7 +105,7 @@ public:
 
 #if WITH_EDITORONLY_DATA
 	UPROPERTY(Transient)
-	UBillboardComponent* ActorIcon;
+	TObjectPtr<UBillboardComponent> ActorIcon;
 #endif // WITH_EDITORONLY_DATA
 
 protected:
@@ -142,7 +142,7 @@ private:
 #endif // WITH_EDITOR
 
 	UPROPERTY(Transient, DuplicateTransient, VisibleAnywhere, AdvancedDisplay, meta = (Category = "Debug"))
-	TMap<TWeakObjectPtr<AActor>, UObject*> Cache;
+	TMap<TWeakObjectPtr<AActor>, TObjectPtr<UObject>> Cache;
 
 	bool bRenderTargetsDirty = false;
 };
\ No newline at end of file
diff --git a/Engine/Plugins/Experimental/Water/Source/Editor/Public/CausticsGeneratorActor.h b/Engine/Plugins/Experimental/Water/Source/Editor/Public/CausticsGeneratorActor.h
--- a/Engine/Plugins/Experimental/Water/Source/Editor/Public/CausticsGeneratorActor.h
+++ b/Engine/Plugins/Experimental/Water/Source/Editor/Public/CausticsGeneratorActor.h
@@ -25,7 +25,7 @@ public:
 	//UHierarchicalInstancedStaticMeshComponent* CausticParticlesGridHISMC;
 
 	UPROPERTY(BlueprintReadWrite, NonTransactional, meta = (Category = "Default"))
-	USceneComponent* DefaultSceneRoot;
+	TObjectPtr<USceneComponent> DefaultSceneRoot;
 
 
 	UFUNCTION(BlueprintNativeEvent, CallInEditor, BlueprintCallable, Category = "Tick")
diff --git a/Engine/Plugins/Experimental/Water/Source/Editor/Public/WaterEditorSettings.h b/Engine/Plugins/Experimental/Water/Source/Editor/Public/WaterEditorSettings.h
--- a/Engine/Plugins/Experimental/Water/Source/Editor/Public/WaterEditorSettings.h
+++ b/Engine/Plugins/Experimental/Water/Source/Editor/Public/WaterEditorSettings.h
@@ -121,7 +121,7 @@ struct WATEREDITOR_API FWaterBodyLakeDefaults : public FWaterBodyDefaults
 	FWaterBrushActorDefaults BrushDefaults;
 
 	UPROPERTY(EditAnywhere, Instanced, Category = Wave)
-	UWaterWavesBase* WaterWaves = nullptr;
+	TObjectPtr<UWaterWavesBase> WaterWaves = nullptr;
 };
 
 
@@ -136,7 +136,7 @@ struct WATEREDITOR_API FWaterBodyOceanDefaults : public FWaterBodyDefaults
 	FWaterBrushActorDefaults BrushDefaults;
 
 	UPROPERTY(EditAnywhere, Instanced, Category = Wave)
-	UWaterWavesBase* WaterWaves = nullptr;
+	TObjectPtr<UWaterWavesBase> WaterWaves = nullptr;
 };
 
 
diff --git a/Engine/Plugins/Experimental/Water/Source/Editor/Public/WaterEditorSubsystem.h b/Engine/Plugins/Experimental/Water/Source/Editor/Public/WaterEditorSubsystem.h
--- a/Engine/Plugins/Experimental/Water/Source/Editor/Public/WaterEditorSubsystem.h
+++ b/Engine/Plugins/Experimental/Water/Source/Editor/Public/WaterEditorSubsystem.h
@@ -35,12 +35,12 @@ public:
 	virtual UTexture2D* GetErrorSprite() const override { return ErrorSprite; }
 private:
 	UPROPERTY()
-	UMaterialParameterCollection* LandscapeMaterialParameterCollection;
+	TObjectPtr<UMaterialParameterCollection> LandscapeMaterialParameterCollection;
 
 	UPROPERTY()
-	TMap<UClass*, UTexture2D*> WaterActorSprites;
+	TMap<TObjectPtr<UClass>, TObjectPtr<UTexture2D>> WaterActorSprites;
 	UPROPERTY()
-	UTexture2D* DefaultWaterActorSprite;
+	TObjectPtr<UTexture2D> DefaultWaterActorSprite;
 	UPROPERTY()
-	UTexture2D* ErrorSprite;
+	TObjectPtr<UTexture2D> ErrorSprite;
 };
diff --git a/Engine/Plugins/Experimental/Water/Source/Runtime/Public/BuoyancyComponent.h b/Engine/Plugins/Experimental/Water/Source/Runtime/Public/BuoyancyComponent.h
--- a/Engine/Plugins/Experimental/Water/Source/Runtime/Public/BuoyancyComponent.h
+++ b/Engine/Plugins/Experimental/Water/Source/Runtime/Public/BuoyancyComponent.h
@@ -109,11 +109,11 @@ protected:
 	void ComputePontoonCoefficients();
 
 	UPROPERTY(Transient)
-	TArray<UWaterBodyComponent*> CurrentWaterBodyComponents;
+	TArray<TObjectPtr<UWaterBodyComponent>> CurrentWaterBodyComponents;
 
 	// Primitive component that will be used for physics simulation.
 	UPROPERTY()
-	UPrimitiveComponent* SimulatingComponent;
+	TObjectPtr<UPrimitiveComponent> SimulatingComponent;
 
 	// async data
 
diff --git a/Engine/Plugins/Experimental/Water/Source/Runtime/Public/BuoyancyManager.h b/Engine/Plugins/Experimental/Water/Source/Runtime/Public/BuoyancyManager.h
--- a/Engine/Plugins/Experimental/Water/Source/Runtime/Public/BuoyancyManager.h
+++ b/Engine/Plugins/Experimental/Water/Source/Runtime/Public/BuoyancyManager.h
@@ -56,7 +56,7 @@ private:
 	TSet<const UPrimitiveComponent*> PhysicsInitializedSimulatingComponents;
 
 	UPROPERTY()
-	TArray<UBuoyancyComponent*> BuoyancyComponents;
+	TArray<TObjectPtr<UBuoyancyComponent>> BuoyancyComponents;
 	TArray<UBuoyancyComponent*> BuoyancyComponentsActive;
 	
 	// List of buoyancy components to defer registration until this buoyancy manager is fully initialized.
diff --git a/Engine/Plugins/Experimental/Water/Source/Runtime/Public/BuoyancyTypes.h b/Engine/Plugins/Experimental/Water/Source/Runtime/Public/BuoyancyTypes.h
--- a/Engine/Plugins/Experimental/Water/Source/Runtime/Public/BuoyancyTypes.h
+++ b/Engine/Plugins/Experimental/Water/Source/Runtime/Public/BuoyancyTypes.h
@@ -89,7 +89,7 @@ struct FSphericalPontoon
 	uint8 bUseCenterSocket : 1;
 
 	UPROPERTY(Transient, BlueprintReadOnly, Category = Buoyancy)
-	UWaterBodyComponent* CurrentWaterBodyComponent;
+	TObjectPtr<UWaterBodyComponent> CurrentWaterBodyComponent;
 
 	FSolverSafeWaterBodyData* SolverWaterBody;
 
diff --git a/Engine/Plugins/Experimental/Water/Source/Runtime/Public/GerstnerWaterWaves.h b/Engine/Plugins/Experimental/Water/Source/Runtime/Public/GerstnerWaterWaves.h
--- a/Engine/Plugins/Experimental/Water/Source/Runtime/Public/GerstnerWaterWaves.h
+++ b/Engine/Plugins/Experimental/Water/Source/Runtime/Public/GerstnerWaterWaves.h
@@ -187,7 +187,7 @@ public:
 	virtual float GetWaveAttenuationFactor(const FVector& InPosition, float InWaterDepth, float InTargetWaveMaskDepth) const override;
 
 	UPROPERTY(EditAnywhere, BlueprintReadOnly, Instanced, Category = Waves)
-	UGerstnerWaterWaveGeneratorBase* GerstnerWaveGenerator;
+	TObjectPtr<UGerstnerWaterWaveGeneratorBase> GerstnerWaveGenerator;
 
 protected:
 	UPROPERTY(BlueprintReadOnly, Category = "Wave")
diff --git a/Engine/Plugins/Experimental/Water/Source/Runtime/Public/LakeCollisionComponent.h b/Engine/Plugins/Experimental/Water/Source/Runtime/Public/LakeCollisionComponent.h
--- a/Engine/Plugins/Experimental/Water/Source/Runtime/Public/LakeCollisionComponent.h
+++ b/Engine/Plugins/Experimental/Water/Source/Runtime/Public/LakeCollisionComponent.h
@@ -35,7 +35,7 @@ protected:
 	void CreateLakeBodySetupIfNeeded();
 private:
 	UPROPERTY(NonPIEDuplicateTransient)
-	class UBodySetup* CachedBodySetup;
+	TObjectPtr<class UBodySetup> CachedBodySetup;
 
 	UPROPERTY()
 	FVector BoxExtent;
diff --git a/Engine/Plugins/Experimental/Water/Source/Runtime/Public/NiagaraDataInterfaceWater.h b/Engine/Plugins/Experimental/Water/Source/Runtime/Public/NiagaraDataInterfaceWater.h
--- a/Engine/Plugins/Experimental/Water/Source/Runtime/Public/NiagaraDataInterfaceWater.h
+++ b/Engine/Plugins/Experimental/Water/Source/Runtime/Public/NiagaraDataInterfaceWater.h
@@ -41,5 +41,5 @@ public:
 	void SetWaterBodyComponent(UWaterBodyComponent* InWaterBodyComponent) { SourceBodyComponent = InWaterBodyComponent; }
 private:
 	UPROPERTY(EditAnywhere, Category = "Water") 
-	UWaterBodyComponent* SourceBodyComponent;
+	TObjectPtr<UWaterBodyComponent> SourceBodyComponent;
 };
diff --git a/Engine/Plugins/Experimental/Water/Source/Runtime/Public/OceanCollisionComponent.h b/Engine/Plugins/Experimental/Water/Source/Runtime/Public/OceanCollisionComponent.h
--- a/Engine/Plugins/Experimental/Water/Source/Runtime/Public/OceanCollisionComponent.h
+++ b/Engine/Plugins/Experimental/Water/Source/Runtime/Public/OceanCollisionComponent.h
@@ -38,7 +38,7 @@ protected:
 	void CreateOceanBodySetupIfNeeded();
 private:
 	UPROPERTY(NonPIEDuplicateTransient)
-	class UBodySetup* CachedBodySetup;
+	TObjectPtr<class UBodySetup> CachedBodySetup;
 
 	FBox BoundingBox;
 };
diff --git a/Engine/Plugins/Experimental/Water/Source/Runtime/Public/WaterBodyActor.h b/Engine/Plugins/Experimental/Water/Source/Runtime/Public/WaterBodyActor.h
--- a/Engine/Plugins/Experimental/Water/Source/Runtime/Public/WaterBodyActor.h
+++ b/Engine/Plugins/Experimental/Water/Source/Runtime/Public/WaterBodyActor.h
@@ -105,13 +105,13 @@ protected:
 
 	/** The spline data attached to this water type. */
 	UPROPERTY(Category = Water, VisibleAnywhere, BlueprintReadOnly, meta = (AllowPrivateAccess = "true"))
-	UWaterSplineComponent* SplineComp;
+	TObjectPtr<UWaterSplineComponent> SplineComp;
 
 	UPROPERTY(Instanced)
-	UWaterSplineMetadata* WaterSplineMetadata;
+	TObjectPtr<UWaterSplineMetadata> WaterSplineMetadata;
 
 	UPROPERTY(Category = Water, VisibleAnywhere, BlueprintReadOnly, meta = (ExposeFunctionCategories = "Water,Wave,Rendering,Terrain,Navigation,Physics,Collision,Debug", AllowPrivateAccess = "true"))
-	UWaterBodyComponent* WaterBodyComponent;
+	TObjectPtr<UWaterBodyComponent> WaterBodyComponent;
 
 	/** Unique Id for accessing (wave, ... ) data in GPU buffers */
 	UPROPERTY(Transient, DuplicateTransient, NonTransactional, BlueprintReadOnly, Category = Water, meta = (AllowPrivateAccess = "true"))
@@ -122,7 +122,7 @@ protected:
 
 	// #todo_water: This should be moved to the component when component subobjects are supported
 	UPROPERTY(EditAnywhere, BlueprintReadOnly, Instanced, Category = Wave, DisplayName = "Waves Source", meta = (Tooltip = ""))
-	UWaterWavesBase* WaterWaves = nullptr;
+	TObjectPtr<UWaterWavesBase> WaterWaves = nullptr;
 
 #if WITH_EDITOR
 	virtual void PostEditMove(bool bFinished) override;
@@ -185,7 +185,7 @@ protected:
 	friend class UWaterBodyComponent;
 #if WITH_EDITORONLY_DATA
 	UPROPERTY()
-	UPhysicalMaterial* PhysicalMaterial_DEPRECATED;
+	TObjectPtr<UPhysicalMaterial> PhysicalMaterial_DEPRECATED;
 
 	UPROPERTY()
 	float TargetWaveMaskDepth_DEPRECATED;
@@ -203,10 +203,10 @@ protected:
 	FWaterCurveSettings CurveSettings_DEPRECATED;
 
 	UPROPERTY()
-	UMaterialInterface* WaterMaterial_DEPRECATED;
+	TObjectPtr<UMaterialInterface> WaterMaterial_DEPRECATED;
 
 	UPROPERTY()
-	UMaterialInterface* UnderwaterPostProcessMaterial_DEPRECATED;
+	TObjectPtr<UMaterialInterface> UnderwaterPostProcessMaterial_DEPRECATED;
 
 	UPROPERTY()
 	FLandmassTerrainCarvingSettings TerrainCarvingSettings_DEPRECATED;
@@ -227,7 +227,7 @@ protected:
 	bool bOverrideWaterMesh_DEPRECATED;
 
 	UPROPERTY()
-	UStaticMesh* WaterMeshOverride_DEPRECATED;
+	TObjectPtr<UStaticMesh> WaterMeshOverride_DEPRECATED;
 
 	UPROPERTY()
 	int32 OverlapMaterialPriority_DEPRECATED = 0;
@@ -236,10 +236,10 @@ protected:
 	FName CollisionProfileName_DEPRECATED;
 
 	UPROPERTY()
-	UMaterialInstanceDynamic* WaterMID_DEPRECATED;
+	TObjectPtr<UMaterialInstanceDynamic> WaterMID_DEPRECATED;
 
 	UPROPERTY()
-	UMaterialInstanceDynamic* UnderwaterPostProcessMID_DEPRECATED;
+	TObjectPtr<UMaterialInstanceDynamic> UnderwaterPostProcessMID_DEPRECATED;
 
 	UPROPERTY()
 	TArray<TLazyObjectPtr<AWaterBodyIsland>> Islands_DEPRECATED;
diff --git a/Engine/Plugins/Experimental/Water/Source/Runtime/Public/WaterBodyComponent.h b/Engine/Plugins/Experimental/Water/Source/Runtime/Public/WaterBodyComponent.h
--- a/Engine/Plugins/Experimental/Water/Source/Runtime/Public/WaterBodyComponent.h
+++ b/Engine/Plugins/Experimental/Water/Source/Runtime/Public/WaterBodyComponent.h
@@ -67,7 +67,7 @@ struct FUnderwaterPostProcessSettings
 
 	/** This is the parent post process material for the PostProcessSettings */
 	UPROPERTY()
-	UMaterialInterface* UnderwaterPostProcessMaterial_DEPRECATED;
+	TObjectPtr<UMaterialInterface> UnderwaterPostProcessMaterial_DEPRECATED;
 };
 
 // ----------------------------------------------------------------------------------
@@ -468,7 +468,7 @@ public:
 	static const FName GroundZMinParamName;
 
 	UPROPERTY(EditDefaultsOnly, Category = Collision, meta = (EditCondition = "bGenerateCollisions"))
-	UPhysicalMaterial* PhysicalMaterial;
+	TObjectPtr<UPhysicalMaterial> PhysicalMaterial;
 
 	/** Water depth at which waves start being attenuated. */
 	UPROPERTY(EditAnywhere, BlueprintReadWrite, Category = Wave, DisplayName = "Wave Attenuation Water Depth", meta = (UIMin = 0, ClampMin = 0, UIMax = 10000.0))
@@ -487,17 +487,17 @@ public:
 	FWaterCurveSettings CurveSettings;
 
 	UPROPERTY(Category = Rendering, EditAnywhere, BlueprintReadOnly)
-	UMaterialInterface* WaterMaterial;
+	TObjectPtr<UMaterialInterface> WaterMaterial;
 
 	UPROPERTY(Category = HLOD, EditAnywhere, BlueprintReadOnly, meta = (DisplayName = "Water HLOD Material"))
-	UMaterialInterface* WaterHLODMaterial;
+	TObjectPtr<UMaterialInterface> WaterHLODMaterial;
 
 	/** Post process material to apply when the camera goes underwater (only available when bGenerateCollisions is true because collisions are needed to detect if it's under water). */
 	UPROPERTY(Category = Rendering, EditAnywhere, BlueprintReadOnly, meta = (EditCondition = "bGenerateCollisions", DisplayAfter = "WaterMaterial"))
-	UMaterialInterface* UnderwaterPostProcessMaterial;
+	TObjectPtr<UMaterialInterface> UnderwaterPostProcessMaterial;
 
 	UPROPERTY(Category = Rendering, EditAnywhere, BlueprintReadOnly, meta = (DisplayAfter = "WaterMaterial"))
-	UMaterialInterface* WaterInfoMaterial;
+	TObjectPtr<UMaterialInterface> WaterInfoMaterial;
 	
 	UPROPERTY(Category = Terrain, EditAnywhere, BlueprintReadWrite)
 	FWaterBodyHeightmapSettings WaterHeightmapSettings;
@@ -525,7 +525,7 @@ protected:
 	int32 WaterBodyIndex = INDEX_NONE;
 
 	UPROPERTY(Category = Rendering, EditAnywhere, BlueprintReadOnly)
-	UStaticMesh* WaterMeshOverride;
+	TObjectPtr<UStaticMesh> WaterMeshOverride;
 
 	/** 
 	 * When this is set to true, the water mesh will always generate tiles for this water body. 
@@ -539,16 +539,16 @@ protected:
 	int32 OverlapMaterialPriority = 0;
 
 	UPROPERTY(Transient)
-	UWaterSplineMetadata* WaterSplineMetadata;
+	TObjectPtr<UWaterSplineMetadata> WaterSplineMetadata;
 
 	UPROPERTY(Category = Debug, VisibleInstanceOnly, Transient, NonPIEDuplicateTransient, TextExportTransient, meta = (DisplayAfter = "WaterMaterial"))
-	UMaterialInstanceDynamic* WaterMID;
+	TObjectPtr<UMaterialInstanceDynamic> WaterMID;
 
 	UPROPERTY(Category = Debug, VisibleInstanceOnly, Transient, NonPIEDuplicateTransient, TextExportTransient, meta = (DisplayAfter = "UnderwaterPostProcessMaterial"))
-	UMaterialInstanceDynamic* UnderwaterPostProcessMID;
+	TObjectPtr<UMaterialInstanceDynamic> UnderwaterPostProcessMID;
 	
 	UPROPERTY(Category = Debug, VisibleInstanceOnly, Transient, NonPIEDuplicateTransient, TextExportTransient, meta = (DisplayAfter = "WaterInfoMaterial"))
-	UMaterialInstanceDynamic* WaterInfoMID;
+	TObjectPtr<UMaterialInstanceDynamic> WaterInfoMID;
 
 	/** Islands in this water body*/
 	UPROPERTY(Category = Water, VisibleAnywhere, AdvancedDisplay)
diff --git a/Engine/Plugins/Experimental/Water/Source/Runtime/Public/WaterBodyCustomActor.h b/Engine/Plugins/Experimental/Water/Source/Runtime/Public/WaterBodyCustomActor.h
--- a/Engine/Plugins/Experimental/Water/Source/Runtime/Public/WaterBodyCustomActor.h
+++ b/Engine/Plugins/Experimental/Water/Source/Runtime/Public/WaterBodyCustomActor.h
@@ -18,7 +18,7 @@ class UDEPRECATED_CustomMeshGenerator : public UDEPRECATED_WaterBodyGenerator
 
 public:
 	UPROPERTY(NonPIEDuplicateTransient)
-	UStaticMeshComponent* MeshComp;
+	TObjectPtr<UStaticMeshComponent> MeshComp;
 };
 
 // ----------------------------------------------------------------------------------
@@ -32,6 +32,6 @@ protected:
 
 #if WITH_EDITORONLY_DATA
 	UPROPERTY(NonPIEDuplicateTransient)
-	UDEPRECATED_CustomMeshGenerator* CustomGenerator_DEPRECATED;
+	TObjectPtr<UDEPRECATED_CustomMeshGenerator> CustomGenerator_DEPRECATED;
 #endif
 };
\ No newline at end of file
diff --git a/Engine/Plugins/Experimental/Water/Source/Runtime/Public/WaterBodyCustomComponent.h b/Engine/Plugins/Experimental/Water/Source/Runtime/Public/WaterBodyCustomComponent.h
--- a/Engine/Plugins/Experimental/Water/Source/Runtime/Public/WaterBodyCustomComponent.h
+++ b/Engine/Plugins/Experimental/Water/Source/Runtime/Public/WaterBodyCustomComponent.h
@@ -40,5 +40,5 @@ protected:
 
 protected:
 	UPROPERTY(NonPIEDuplicateTransient)
-	UStaticMeshComponent* MeshComp;
+	TObjectPtr<UStaticMeshComponent> MeshComp;
 };
\ No newline at end of file
diff --git a/Engine/Plugins/Experimental/Water/Source/Runtime/Public/WaterBodyExclusionVolume.h b/Engine/Plugins/Experimental/Water/Source/Runtime/Public/WaterBodyExclusionVolume.h
--- a/Engine/Plugins/Experimental/Water/Source/Runtime/Public/WaterBodyExclusionVolume.h
+++ b/Engine/Plugins/Experimental/Water/Source/Runtime/Public/WaterBodyExclusionVolume.h
@@ -42,19 +42,19 @@ public:
 
 	/** List of water bodies that will be affected by this exclusion volume */
 	UPROPERTY(EditInstanceOnly, Category = Water, meta = (EditCondition = "!bExcludeAllOverlappingWaterBodies"))
-	TArray<AWaterBody*> WaterBodiesToExclude;
+	TArray<TObjectPtr<AWaterBody>> WaterBodiesToExclude;
 
 #if WITH_EDITORONLY_DATA
 	UPROPERTY(meta = (DeprecationMessage = "Property renamed to bExcludeAllOverlapping"))
 	bool bIgnoreAllOverlappingWaterBodies_DEPRECATED = false;
 
 	UPROPERTY(meta = (DeprecationMessage = "Property renamed to WaterBodiesToExclude"))
-	TArray<AWaterBody*> WaterBodiesToIgnore_DEPRECATED;
+	TArray<TObjectPtr<AWaterBody>> WaterBodiesToIgnore_DEPRECATED;
 
 	UPROPERTY(meta = (DeprecatedProperty))
-	AWaterBody* WaterBodyToIgnore_DEPRECATED;
+	TObjectPtr<AWaterBody> WaterBodyToIgnore_DEPRECATED;
 
 	UPROPERTY(Transient)
-	class UBillboardComponent* ActorIcon;
+	TObjectPtr<class UBillboardComponent> ActorIcon;
 #endif // WITH_EDITORONLY_DATA
 };
diff --git a/Engine/Plugins/Experimental/Water/Source/Runtime/Public/WaterBodyIslandActor.h b/Engine/Plugins/Experimental/Water/Source/Runtime/Public/WaterBodyIslandActor.h
--- a/Engine/Plugins/Experimental/Water/Source/Runtime/Public/WaterBodyIslandActor.h
+++ b/Engine/Plugins/Experimental/Water/Source/Runtime/Public/WaterBodyIslandActor.h
@@ -61,7 +61,7 @@ public:
 	TMap<FName, FWaterBodyWeightmapSettings> WaterWeightmapSettings;
 
 	UPROPERTY(Transient)
-	UBillboardComponent* ActorIcon;
+	TObjectPtr<UBillboardComponent> ActorIcon;
 #endif
 
 protected:
@@ -82,5 +82,5 @@ protected:
 	 * The spline data attached to this water type.
 	 */
 	UPROPERTY(VisibleAnywhere, Category = Water, BlueprintReadOnly, meta = (AllowPrivateAccess = "true"))
-	UWaterSplineComponent* SplineComp;
+	TObjectPtr<UWaterSplineComponent> SplineComp;
 };
