Water: Fixed an issue with ground max z calculations when multiple landscapes are present in the same world with drastically different z locations.

diff --git a/Engine/Plugins/Experimental/Water/Source/Runtime/Private/WaterZoneActor.cpp b/Engine/Plugins/Experimental/Water/Source/Runtime/Private/WaterZoneActor.cpp
--- a/Engine/Plugins/Experimental/Water/Source/Runtime/Private/WaterZoneActor.cpp
+++ b/Engine/Plugins/Experimental/Water/Source/Runtime/Private/WaterZoneActor.cpp
@@ -282,12 +282,17 @@ bool AWaterZone::UpdateWaterInfoTexture()
 		float GroundZMax = TNumericLimits<float>::Lowest();
 
 		TArray<TWeakObjectPtr<AActor>> GroundActors;
+		const FBox WaterZoneBounds(GetActorLocation() - FVector(GetZoneExtent() / 2., 0.), GetActorLocation() + FVector(GetZoneExtent() / 2., 0));
 		for (ALandscapeProxy* LandscapeProxy : TActorRange<ALandscapeProxy>(World))
 		{
 			const FBox LandscapeBox = LandscapeProxy->GetComponentsBoundingBox();
-			GroundZMin = FMath::Min(GroundZMin, LandscapeBox.Min.Z);
-			GroundZMax = FMath::Max(GroundZMax, LandscapeBox.Max.Z);
-			GroundActors.Add(LandscapeProxy);
+			// Only consider landscapes which this zone intersects with in XY
+			if (WaterZoneBounds.IntersectXY(LandscapeBox))
+			{
+				GroundZMin = FMath::Min(GroundZMin, LandscapeBox.Min.Z);
+				GroundZMax = FMath::Max(GroundZMax, LandscapeBox.Max.Z);
+				GroundActors.Add(LandscapeProxy);
+			}
 		}
 
 		// Check all the ground actors have complete shader maps before we try to render them into the water info texture
