Water: Water LOD system. Replaces water mesh tiles at a distance with simplified, non-tessellated meshes. This allows the Water Info texture to localized and higher precision.

diff --git a/Engine/Plugins/Experimental/Water/Source/Editor/Private/WaterLandscapeBrush.cpp b/Engine/Plugins/Experimental/Water/Source/Editor/Private/WaterLandscapeBrush.cpp
--- a/Engine/Plugins/Experimental/Water/Source/Editor/Private/WaterLandscapeBrush.cpp
+++ b/Engine/Plugins/Experimental/Water/Source/Editor/Private/WaterLandscapeBrush.cpp
@@ -205,7 +205,7 @@ void AWaterLandscapeBrush::OnActorChanged(AActor* Actor, bool bWeightmapSettings
 	{
 		if (UWaterSubsystem* WaterSubsystem = UWaterSubsystem::GetWaterSubsystem(GetWorld()))
 		{
-			WaterSubsystem->MarkAllWaterZonesForRebuild();
+			WaterSubsystem->MarkAllWaterZonesForRebuild(EWaterZoneRebuildFlags::UpdateWaterMesh);
 		}
 	}
 }
diff --git a/Engine/Plugins/Experimental/Water/Source/Runtime/Private/WaterBodyActor.cpp b/Engine/Plugins/Experimental/Water/Source/Runtime/Private/WaterBodyActor.cpp
--- a/Engine/Plugins/Experimental/Water/Source/Runtime/Private/WaterBodyActor.cpp
+++ b/Engine/Plugins/Experimental/Water/Source/Runtime/Private/WaterBodyActor.cpp
@@ -554,11 +554,6 @@ void AWaterBody::PostRegisterAllComponents()
 		{
 			WaterBodyComponent->WaterBodyIndex = Manager->AddWaterBodyComponent(WaterBodyComponent);
 			WaterBodyIndex = WaterBodyComponent->WaterBodyIndex;
-
-			// By this point, the water body will have already generated a scene proxy.
-			// However, the existing scene proxy doesn't have the correct water body index as we have just been assigned one here from the water body manager.
-			// We need to regenerate the mesh and scene proxy so that the vertex colors contain the updated index.
-			WaterBodyComponent->UpdateWaterBodyRenderData();
 		}
 
 		// At this point, the water mesh actor should be ready and we can setup the MID accordingly : 
diff --git a/Engine/Plugins/Experimental/Water/Source/Runtime/Private/WaterBodyComponent.cpp b/Engine/Plugins/Experimental/Water/Source/Runtime/Private/WaterBodyComponent.cpp
--- a/Engine/Plugins/Experimental/Water/Source/Runtime/Private/WaterBodyComponent.cpp
+++ b/Engine/Plugins/Experimental/Water/Source/Runtime/Private/WaterBodyComponent.cpp
@@ -36,6 +36,9 @@
 #include "Logging/TokenizedMessage.h"
 #include "WaterBodySceneProxy.h"
 #include "Engine/StaticMesh.h"
+#include "DynamicMesh/DynamicMesh3.h"
+#include "DynamicMesh/DynamicMeshAABBTree3.h"
+#include "Operations/MeshPlaneCut.h"
 
 #if WITH_EDITOR
 #include "WaterIconHelper.h"
@@ -128,7 +131,13 @@ void UWaterBodyComponent::OnHiddenInGameChanged()
 
 FPrimitiveSceneProxy* UWaterBodyComponent::CreateSceneProxy()
 {
-	return new FWaterBodySceneProxy(this);
+	if (AffectsWaterMesh() && GetWaterZone())
+	{
+		const FVector TessellatedMeshCenter = GetWaterZone()->GetTessellatedWaterMeshCenter();
+		const FVector TessellatedMeshExtent =  GetWaterZone()->GetTessellatedWaterMeshExtent();
+		return new FWaterBodySceneProxy(this, FBox2D(FVector2D(TessellatedMeshCenter - TessellatedMeshExtent/ 2), FVector2D(TessellatedMeshCenter + TessellatedMeshExtent/ 2)));
+	}
+	return nullptr;
 }
 
 void UWaterBodyComponent::GetUsedMaterials(TArray<UMaterialInterface*>& OutMaterialInterfaces, bool bGetDebugMaterials) const
@@ -137,6 +146,27 @@ void UWaterBodyComponent::GetUsedMaterials(TArray<UMaterialInterface*>& OutMater
 	{
 		OutMaterialInterfaces.Add(WaterInfoMID);
 	}
+	if (WaterLODMID)
+	{
+		OutMaterialInterfaces.Add(WaterLODMID);
+	}
+}
+
+
+void UWaterBodyComponent::OnTessellatedWaterMeshBoundsChanged()
+{
+	// When the water info bounds changes, we need to push new bounds to the MIDs so they can address the texture correctly.
+	UpdateMaterialInstances();
+}
+
+void UWaterBodyComponent::PushTessellatedWaterMeshBoundsToProxy(const FBox2D& TessellatedWaterMeshBounds)
+{
+	OnTessellatedWaterMeshBoundsChanged();
+
+	if (SceneProxy)
+	{
+		static_cast<FWaterBodySceneProxy*>(SceneProxy)->OnTessellatedWaterMeshBoundsChanged_GameThread(TessellatedWaterMeshBounds);
+	}
 }
 
 bool UWaterBodyComponent::IsFlatSurface() const
@@ -265,6 +295,12 @@ UMaterialInstanceDynamic* UWaterBodyComponent::GetWaterInfoMaterialInstance()
 	return WaterInfoMID;
 }
 
+UMaterialInstanceDynamic* UWaterBodyComponent::GetWaterLODMaterialInstance()
+{
+	CreateOrUpdateWaterLODMID();
+	return WaterLODMID;
+}
+
 void UWaterBodyComponent::SetUnderwaterPostProcessMaterial(UMaterialInterface* InMaterial)
 {
 	UnderwaterPostProcessMaterial = InMaterial;
@@ -340,6 +376,27 @@ void UWaterBodyComponent::UpdateExclusionVolumes()
 	}
 }
 
+void UWaterBodyComponent::UpdateWaterZones()
+{
+	if (UWorld* World = GetWorld())
+	{
+		AWaterZone* FoundZone = FindWaterZone();
+
+		if (OwningWaterZone != nullptr)
+		{
+			OwningWaterZone->RemoveWaterBodyComponent(this);
+			OwningWaterZone = nullptr;
+		}
+
+		if (FoundZone != nullptr)
+		{
+			FoundZone->AddWaterBodyComponent(this);
+		}
+
+		OwningWaterZone = FoundZone;
+	}
+}
+
 FPostProcessVolumeProperties UWaterBodyComponent::GetPostProcessProperties() const
 {
 	FPostProcessVolumeProperties Ret;
@@ -684,8 +741,6 @@ void UWaterBodyComponent::OnRegister()
 
 	check(WaterSplineMetadata);
 
-	CreateOrUpdateWaterInfoMID();
-
 #if WITH_EDITOR
 	RegisterOnChangeWaterSplineData(/*bRegister = */true); 
 
@@ -745,6 +800,7 @@ void UWaterBodyComponent::UpdateMaterialInstances()
 {
 	CreateOrUpdateWaterMID();
 	CreateOrUpdateWaterInfoMID();
+	CreateOrUpdateWaterLODMID();
 	CreateOrUpdateUnderwaterPostProcessMID();
 }
 
@@ -810,6 +866,16 @@ void UWaterBodyComponent::CreateOrUpdateWaterInfoMID()
 	}
 }
 
+void UWaterBodyComponent::CreateOrUpdateWaterLODMID()
+{
+	if (GetWorld())
+	{
+		WaterLODMID = FWaterUtils::GetOrCreateTransientMID(WaterLODMID, TEXT("WaterLODMID"), WaterLODMaterial, GetTransientMIDFlags());
+
+		SetDynamicParametersOnMID(WaterLODMID);
+	}
+}
+
 void UWaterBodyComponent::PrepareCurrentPostProcessSettings()
 {
 	// Prepare the transient settings that are actually used by the post-process system : 
@@ -1203,6 +1269,8 @@ void UWaterBodyComponent::UpdateAll(bool bShapeOrPositionChanged)
 			UpdateIslands();
 
 			UpdateExclusionVolumes();
+
+			UpdateWaterZones();
 		}
 
 		// Finally, generate the body once again, this time with the updated list of exclusion volumes
@@ -1298,6 +1366,9 @@ void UWaterBodyComponent::PostLoad()
 #if WITH_EDITOR
 	RegisterOnUpdateWavesData(GetWaterWaves(), /* bRegister = */true);
 #endif // WITH_EDITOR
+
+	// Register the component with the water zone since the component is the authority
+	UpdateWaterZones();
 }
 
 void UWaterBodyComponent::DeprecateData()
@@ -1374,9 +1445,17 @@ void UWaterBodyComponent::GetResourceSizeEx(FResourceSizeEx& CumulativeResourceS
 {
 	Super::GetResourceSizeEx(CumulativeResourceSize);
 
+	uint32 SectionsSize = 0;
+	for (const FWaterBodyMeshSection& Section : WaterBodyMeshSections)
+	{
+		SectionsSize += Section.GetAllocatedSize();
+	}
+
 	// Account for all non-editor data properties :
 	CumulativeResourceSize.AddDedicatedSystemMemoryBytes(
-		WaterBodyMeshIndices.GetAllocatedSize()
+		WaterBodyMeshSections.GetAllocatedSize() 
+		+ SectionsSize
+		+ WaterBodyMeshIndices.GetAllocatedSize()
 		+ WaterBodyMeshVertices.GetAllocatedSize()
 		+ DilatedWaterBodyMeshIndices.GetAllocatedSize()
 		+ DilatedWaterBodyMeshVertices.GetAllocatedSize());
@@ -1407,14 +1486,14 @@ bool UWaterBodyComponent::SetDynamicParametersOnMID(UMaterialInstanceDynamic* In
 		check(WaterMeshComponent);
 
 		// Location should be the bottom left of the zone
-		const FVector2D ZoneExtent = WaterZone->GetZoneExtent();
-		const FVector2D WaterAreaLocation = FVector2D(WaterZone->GetActorLocation()) - (ZoneExtent / 2.f);
+		const FVector2D WaterInfoExtent = FVector2D(WaterZone->GetTessellatedWaterMeshExtent());
+		const FVector2D WaterInfoLocation = FVector2D(WaterZone->GetTessellatedWaterMeshCenter()) - (WaterInfoExtent / 2.f);
 
 		FVector4 WaterArea;
-		WaterArea.X = WaterAreaLocation.X;
-		WaterArea.Y = WaterAreaLocation.Y;
-		WaterArea.Z = ZoneExtent.X;
-		WaterArea.W = ZoneExtent.Y;
+		WaterArea.X = WaterInfoLocation.X;
+		WaterArea.Y = WaterInfoLocation.Y;
+		WaterArea.Z = WaterInfoExtent.X;
+		WaterArea.W = WaterInfoExtent.Y;
 		InMID->SetDoubleVectorParameterValue(WaterAreaParamName, WaterArea);
 
 		const FVector2f WaterHeightExtents = WaterZone->GetWaterHeightExtents();
@@ -1518,7 +1597,7 @@ float UWaterBodyComponent::GetMaxWaveHeight() const
 	return (HasWaves() ? GetWaterWaves()->GetMaxWaveHeight() : 0.0f) + MaxWaveHeightOffset;
 }
 
-void UWaterBodyComponent::CopySharedCollisionSettingsToComponent(UPrimitiveComponent * InComponent)
+void UWaterBodyComponent::CopySharedCollisionSettingsToComponent(UPrimitiveComponent* InComponent)
 {
 	InComponent->SetCollisionEnabled(GetCollisionEnabled());
 	InComponent->SetNotifyRigidBodyCollision(BodyInstance.bNotifyRigidBodyCollision);
@@ -1527,7 +1606,7 @@ void UWaterBodyComponent::CopySharedCollisionSettingsToComponent(UPrimitiveCompo
 	InComponent->SetGenerateOverlapEvents(GetGenerateOverlapEvents());
 }
 
-void UWaterBodyComponent::CopySharedNavigationSettingsToComponent(UPrimitiveComponent * InComponent)
+void UWaterBodyComponent::CopySharedNavigationSettingsToComponent(UPrimitiveComponent* InComponent)
 {
 	InComponent->SetCanEverAffectNavigation(CanEverAffectNavigation());
 	InComponent->SetCustomNavigableGeometry(HasCustomNavigableGeometry());
@@ -1564,12 +1643,12 @@ UWaterWavesBase* UWaterBodyComponent::GetWaterWaves() const
 	return nullptr;
 }
 
-AWaterZone* UWaterBodyComponent::GetWaterZone() const
+AWaterZone* UWaterBodyComponent::FindWaterZone() const
 {
-	TRACE_CPUPROFILER_EVENT_SCOPE(UWaterBodyComponent::GetWaterZone);
+	TRACE_CPUPROFILER_EVENT_SCOPE(UWaterBodyComponent::FindWaterZone);
 
 	TArray<AWaterZone*, TInlineAllocator<4>> ViableZones;
-	
+
 	if (const UWorld* World = GetWorld())
 	{
 		const ULevel* PreferredLevel = GetTypedOuter<ULevel>();
@@ -1578,12 +1657,15 @@ AWaterZone* UWaterBodyComponent::GetWaterZone() const
 			for (AWaterZone* WaterZone : TActorRange<AWaterZone>(World, AWaterZone::StaticClass(), EActorIteratorFlags::SkipPendingKill))
 			{
 				// WaterZone->GetZoneExtents returns the full extent of the zone but BoxSphereBounds expects a half-extent.
-				const FBoxSphereBounds WaterZoneBounds(WaterZone->GetActorLocation(), FVector(WaterZone->GetZoneExtent() / 2.f, 8192.f), 0.f);
-				const FBoxSphereBounds ComponentBounds = CalcBounds(GetComponentTransform());
-
-				// Only consider WaterZones which this component overlaps but prefer choosing water zones
-				// which are part of the same outered level.
-				if (FBoxSphereBounds::BoxesIntersect(ComponentBounds, WaterZoneBounds))
+				const FVector2D WaterZoneLocation = FVector2D(WaterZone->GetActorLocation());
+				const FVector2D WaterZoneHalfExtent = WaterZone->GetZoneExtent() / 2.0;
+				const FBox2D WaterZoneBounds(WaterZoneLocation - WaterZoneHalfExtent, WaterZoneLocation + WaterZoneHalfExtent);
+				const FBox ComponentBounds3d = CalcBounds(GetComponentTransform()).GetBox();
+				const FBox2D ComponentBounds = FBox2D(FVector2D(ComponentBounds3d.Min), FVector2D(ComponentBounds3d.Max));
+
+				// Only consider WaterZones which this component overlaps 
+				// but prefer choosing water zones which are part of the same outered level.
+				if (ComponentBounds.Intersect(WaterZoneBounds))
 				{
 					if (WaterZone->GetTypedOuter<ULevel>() == PreferredLevel)
 					{
@@ -1598,24 +1680,367 @@ AWaterZone* UWaterBodyComponent::GetWaterZone() const
 	return ViableZones.Num() > 0 ? ViableZones[0] : nullptr;
 }
 
-void UWaterBodyComponent::UpdateWaterBodyRenderData()
+static inline FColor PackFlowData(float VelocityMagnitude, float DirectionAngle, float MaxVelocity)
+{
+	check((DirectionAngle >= 0.f) && (DirectionAngle <= TWO_PI));
+
+	const float NormalizedMagnitude = FMath::Clamp(VelocityMagnitude, 0.f, MaxVelocity) / MaxVelocity;
+	const float NormalizedAngle = DirectionAngle / TWO_PI;
+	const float MappedMagnitude = NormalizedMagnitude * TNumericLimits<uint16>::Max();
+	const float MappedAngle = NormalizedAngle * TNumericLimits<uint16>::Max();
+
+	const uint16 ResultMag = (uint16)MappedMagnitude;
+	const uint16 ResultAngle = (uint16)MappedAngle;
+	
+	FColor Result;
+	Result.R = (ResultMag >> 8) & 0xFF;
+	Result.G = (ResultMag >> 0) & 0xFF;
+	Result.B = (ResultAngle >> 8) & 0xFF;
+	Result.A = (ResultAngle >> 0) & 0xFF;
+
+	return Result;
+}
+
+
+void UWaterBodyComponent::RebuildWaterBodyInfoMesh()
 {
-	if (AffectsWaterInfo())
+	using namespace UE::Geometry;
+
+	WaterBodyMeshVertices.Empty();
+	WaterBodyMeshIndices.Empty();
+	DilatedWaterBodyMeshVertices.Empty();
+	DilatedWaterBodyMeshIndices.Empty();
+
+	FDynamicMesh3 WaterBodyMesh(EMeshComponents::VertexColors);
+	FDynamicMesh3 WaterBodyDilatedMesh(EMeshComponents::None);
+
+	GenerateWaterBodyMesh(WaterBodyMesh, &WaterBodyDilatedMesh);
+	
+	const float MaxFlowVelocity = FWaterUtils::GetWaterMaxFlowVelocity(false);
+	WaterBodyMeshVertices.Reserve(WaterBodyMesh.VertexCount());
+	WaterBodyMeshIndices.Reserve(3 * WaterBodyMesh.TriangleCount());
+	for (int VertexIndex : WaterBodyMesh.VertexIndicesItr())
 	{
-		GenerateWaterBodyMesh();
-		MarkRenderStateDirty();
+		FVertexInfo VertexInfo;
+		const bool bWantColors = true;
+		const bool bWantNormals = false;
+		const bool bWantUVs = false;
+		WaterBodyMesh.GetVertex(VertexIndex, VertexInfo, bWantNormals, bWantColors, bWantUVs);
+
+		const FVector3f LocalSpacePosition = FVector3f(VertexInfo.Position);
+		FDynamicMeshVertex MeshVertex(LocalSpacePosition);
+
+		const float FlowSpeed = VertexInfo.Color.X;
+		const float FlowDirection = VertexInfo.Color.Y;
+		MeshVertex.Color = PackFlowData(FlowSpeed, FlowDirection, MaxFlowVelocity);
+
+		WaterBodyMeshVertices.Add(MeshVertex);
+	}
+	for (FIndex3i Triangle : WaterBodyMesh.TrianglesItr())
+	{
+		WaterBodyMeshIndices.Append({ (uint32)Triangle.A, (uint32)Triangle.B, (uint32)Triangle.C });
+	}
+
+	DilatedWaterBodyMeshIndices.Reserve(3 * WaterBodyDilatedMesh.TriangleCount());
+	DilatedWaterBodyMeshVertices.Reserve(WaterBodyDilatedMesh.VertexCount());
+	for (FVector3d Vertex : WaterBodyDilatedMesh.VerticesItr())
+	{
+		DilatedWaterBodyMeshVertices.Emplace(FVector3f(Vertex));
+	}
+	for (FIndex3i Triangle : WaterBodyDilatedMesh.TrianglesItr())
+	{
+		DilatedWaterBodyMeshIndices.Append( {(uint32)Triangle.A, (uint32)Triangle.B, (uint32)Triangle.C });
+	}
+}
+
+void UWaterBodyComponent::RebuildWaterBodyLODSections()
+{
+	using namespace UE::Geometry;
+
+	WaterBodyMeshSections.Empty();
+
+	TRACE_CPUPROFILER_EVENT_SCOPE("ComputeWaterBodyMeshSections");
+
+	const AWaterZone* WaterZone = GetWaterZone();
+	check(WaterZone);
+
+	const UWaterMeshComponent* WaterMesh = WaterZone->GetWaterMeshComponent();
+	check(WaterMesh);
+
+	const float WaterLODSectionSize = WaterZone->GetNonTessellatedLODSectionSize();
+	const FVector2i WaterZoneExtentInLODSections = FVector2i(FMath::DivideAndRoundUp(WaterZone->GetZoneExtent(), FVector2d(WaterLODSectionSize)));
+
+	const float GridTileSize = WaterMesh->GetTileSize();
+	const FVector2d GridPosition = FVector2d(WaterMesh->GetComponentLocation().GridSnap(GridTileSize));
+	const FVector2d MeshOrigin = GridPosition - (GridTileSize * FVector2d(WaterMesh->GetExtentInTiles()));
+
+	FDynamicMesh3 EditMesh(EMeshComponents::VertexColors);
 
-		// Compact the arrays since they are not going to change unless everything is rebuilt
-		WaterBodyMeshVertices.Shrink();
-		WaterBodyMeshIndices.Shrink();
-		DilatedWaterBodyMeshVertices.Shrink();
-		DilatedWaterBodyMeshIndices.Shrink();
+	GenerateWaterBodyMesh(EditMesh);
+
+	// Transform everything to world space. This is necessary to maintain an axis-aligned grid. The water mesh is guaranteed to always be axis-aligned in worldspace
+	// so we transform into that coordinate system to guarantee the calculations occur with axis-aligned section tiles.
+	for (int32 VID : EditMesh.VertexIndicesItr())
+	{
+		FVertexInfo VertexInfo = EditMesh.GetVertexInfo(VID);
+		VertexInfo.Position = GetComponentTransform().TransformPosition(VertexInfo.Position);
+		EditMesh.SetVertex(VID, VertexInfo.Position);
+	}
+
+
+	// Slice the mesh along the grid defined by the Water Mesh
+	{
+		TRACE_CPUPROFILER_EVENT_SCOPE("ComputeWaterBodyMeshSections::PartitionWaterBodyMesh");
+
+		// We start at an index of -1 and go to +1 here to ensure the bordering sections don't contain triangles that could extend outside the bounds of the cell.
+		for (int32 SectionIndexX = -1; SectionIndexX < WaterZoneExtentInLODSections.X + 1; ++SectionIndexX)
+		{
+			const FVector3d SlicePos = FVector3d(MeshOrigin + SectionIndexX * FVector2d(WaterLODSectionSize, 0.0), 0.0);
+			const FVector3d SliceNormal(-1.0, 0.0, 0.0);
+
+			FMeshPlaneCut Cut(&EditMesh, SlicePos, SliceNormal);
+			if (!ensure(Cut.CutWithoutDelete(false)))
+			{
+				continue;
+			}
+		}
+
+		for (int32 SectionIndexY = -1; SectionIndexY < WaterZoneExtentInLODSections.Y + 1; ++SectionIndexY)
+		{
+			const FVector3d SlicePos = FVector3d(MeshOrigin + SectionIndexY * FVector2d(0.0, WaterLODSectionSize), 0.0);
+			const FVector3d SliceNormal(0.0, -1.0, 0.0);
+
+			FMeshPlaneCut Cut(&EditMesh, SlicePos, SliceNormal);
+			if (!ensure(Cut.CutWithoutDelete(false)))
+			{
+				continue;
+			}
+		}
+	}
+
+	FDynamicMeshAABBTree3 AABBTree(&EditMesh, true);
+
+	const uint32 NumberOfSections = WaterZoneExtentInLODSections.X * WaterZoneExtentInLODSections.Y;
+	TArray<FWaterBodyMeshSection> Sections;
+	Sections.SetNum(NumberOfSections);
+
+	// Need to access this before the ParallelFor since Foreground Worker threads don't count as gamethread.
+	const float MaxFlowVelocity = FWaterUtils::GetWaterMaxFlowVelocity(false);
+	ParallelFor(NumberOfSections, [&](int32 SectionIndex)
+	{
+		TRACE_CPUPROFILER_EVENT_SCOPE("ComputeWaterBodyMeshSections::BuildMeshSection");
+		const int32 SectionIndexX = SectionIndex % WaterZoneExtentInLODSections.X;
+		const int32 SectionIndexY = SectionIndex / WaterZoneExtentInLODSections.X;
+
+		const FVector2d SectionPos = MeshOrigin + FVector2d(SectionIndexX, SectionIndexY) * WaterLODSectionSize;
+
+		const FAxisAlignedBox2d GridRectXY(SectionPos, SectionPos + WaterLODSectionSize);
+
+		FDynamicMeshAABBTree3::FTreeTraversal Traversal;
+		
+		int SelectAllDepth = TNumericLimits<int>::Max();
+		int CurrentDepth = -1;
+		
+		TArray<int> Triangles;
+
+		Traversal.NextBoxF =
+			[&GridRectXY, &SelectAllDepth, &CurrentDepth](const FAxisAlignedBox3d& Box, int Depth)
+		{
+			CurrentDepth = Depth;
+			if (Depth > SelectAllDepth)
+			{
+				// We are deeper than the depth whose AABB was first detected to be contained in the RectangleXY,
+				// descend and collect all leaf triangles
+				return true;
+			}
+			
+			SelectAllDepth = TNumericLimits<int>::Max();
+			
+			const FAxisAlignedBox2d BoxXY(FVector2d(Box.Min.X, Box.Min.Y), FVector2d(Box.Max.X, Box.Max.Y));
+			if (GridRectXY.Intersects(BoxXY))
+			{
+				if (GridRectXY.Contains(BoxXY))
+				{
+					SelectAllDepth = Depth;
+				}
+				
+				return true;		
+			}
+			return false;
+		};
+		
+
+		Traversal.NextTriangleF =
+			[&GridRectXY, &SelectAllDepth, &CurrentDepth, &Triangles, &EditMesh]
+			(int TriangleID)
+		{
+			if (CurrentDepth >= SelectAllDepth)
+			{
+				// This TriangleID is entirely contained in the selection rectangle so we can skip intersection testing
+				Triangles.Add(TriangleID);
+				return;
+			}
+
+			FAxisAlignedBox2d ExpandedGridRect(GridRectXY);
+			// Expand the rect by 1 unit to add an floating point error margin. 
+			ExpandedGridRect.Expand(1);
+		
+			// this triangle is either on the edge or intersecting with the edge but since we did the plane cuts we know it can't be crossing over into the next cell
+			FIndex3i Triangle = EditMesh.GetTriangle(TriangleID);
+			FVector3d VertexA;
+			FVector3d VertexB;
+			FVector3d VertexC;
+			EditMesh.GetTriVertices(TriangleID, VertexA, VertexB, VertexC);
+			if ((ExpandedGridRect.Contains(FVector2d(VertexA))) &&
+				(ExpandedGridRect.Contains(FVector2d(VertexB))) &&
+				(ExpandedGridRect.Contains(FVector2d(VertexC))))
+			{
+				Triangles.Add(TriangleID);
+			}
+		};
+			
+
+		AABBTree.DoTraversal(Traversal);
+
+		if (Triangles.Num() == 0)
+		{
+			return;
+		}
+
+		TArray<FDynamicMeshVertex> Vertices;
+		TArray<uint32> Indices;
+
+		bool bHasBoundaryTriangle = false;
+		if (GetWaterBodyType() != EWaterBodyType::River)
+		{
+			for (int32 TriangleID : Triangles)
+			{
+				if (EditMesh.IsBoundaryTriangle(TriangleID))
+				{
+					bHasBoundaryTriangle = true;
+					break;
+				}
+			}
+		}
+
+		// Simplification can occur for lakes/rivers since they have flat geometry and no vertex color data.
+		if (!bHasBoundaryTriangle && (GetWaterBodyType() != EWaterBodyType::River))
+		{
+			const float MeshZ = GetConstantSurfaceZ();
+			uint32 TriIndices[4];
+
+			auto AddQuadVertex = [&](FVector2D Position, int QuadIndex)
+			{
+				// Transform back into local space when rebuilding the mesh.
+				const FVector3f LocalSpacePosition = FVector3f(GetComponentTransform().InverseTransformPosition(FVector(Position, MeshZ)));
+				FDynamicMeshVertex MeshVertex(LocalSpacePosition);
+				MeshVertex.Color = FColor(0);
+				TriIndices[QuadIndex] = Vertices.Add(MeshVertex);
+			};
+
+			AddQuadVertex(GridRectXY.Min, 0);
+			AddQuadVertex(GridRectXY.Max, 1);
+			AddQuadVertex(FVector2D(GridRectXY.Min.X, GridRectXY.Max.Y), 2);
+			AddQuadVertex(FVector2D(GridRectXY.Max.X, GridRectXY.Min.Y), 3);
+
+			Indices.Append({ TriIndices[0], TriIndices[2], TriIndices[1], TriIndices[0], TriIndices[1], TriIndices[3] });
+		}
+		else
+		{
+			TMap<uint32, uint32> VertexMap;
+			for (int32 TriangleID : Triangles)
+			{
+				FIndex3i Triangle = EditMesh.GetTriangle(TriangleID);
+				
+				TArray<uint32, TInlineAllocator<3>> TriIndices;
+				for (int32 TriIndex = 0; TriIndex < 3; ++TriIndex)
+				{
+					if (uint32* VertexIndexPtr = VertexMap.Find(Triangle[TriIndex]))
+					{
+						TriIndices.Add(*VertexIndexPtr);
+					}
+					else
+					{
+						FVertexInfo VertexInfo;
+						const bool bWantColors = true;
+						const bool bWantNormals = false;
+						const bool bWantUVs = false;
+						EditMesh.GetVertex(Triangle[TriIndex], VertexInfo, bWantNormals, bWantColors, bWantUVs);
+
+						// Transform back into local space when rebuilding the mesh.
+						const FVector3f LocalSpacePosition = FVector3f(GetComponentTransform().InverseTransformPosition(VertexInfo.Position));
+						FDynamicMeshVertex MeshVertex(LocalSpacePosition);
+						MeshVertex.Color = PackFlowData(VertexInfo.Color.X, VertexInfo.Color.Y, MaxFlowVelocity);
+
+						uint32 NewIndex = Vertices.Add(MeshVertex);
+						VertexMap.Add(Triangle[TriIndex], NewIndex);
+						TriIndices.Add(NewIndex);
+					}
+				}
+
+				Indices.Append(TriIndices);
+			}
+		}
+
+		if (Vertices.Num() != 0)
+		{
+			Sections[SectionIndex].Vertices = MoveTemp(Vertices);
+			Sections[SectionIndex].Indices = MoveTemp(Indices);
+			Sections[SectionIndex].SectionBounds = FBox2D(FVector2D(SectionPos), FVector2D(SectionPos + WaterLODSectionSize));
+		}
+	});
+
+	for (FWaterBodyMeshSection& Section : Sections)
+	{
+		if (Section.IsValid())
+		{
+			WaterBodyMeshSections.Emplace(MoveTemp(Section));
+		}
+	}
+}
+
+
+void UWaterBodyComponent::UpdateWaterBodyRenderData()
+{
+	const bool bAffectsWaterInfo = AffectsWaterInfo();
+	const bool bAffectsWaterMesh = AffectsWaterMesh();
+
+	if (bAffectsWaterInfo)
+	{
+		RebuildWaterBodyInfoMesh();
 
 		if (AWaterZone* WaterZone = GetWaterZone())
 		{
 			WaterZone->MarkForRebuild(EWaterZoneRebuildFlags::UpdateWaterInfoTexture);
 		}
 	}
+
+	if (bAffectsWaterMesh)
+	{
+		checkf(bAffectsWaterInfo, TEXT("Invalid water body! WaterBody affects water mesh but doesn't affect water info"));
+
+		UpdateNonTessellatedMeshSections();
+	}
+
+	if (bAffectsWaterInfo || bAffectsWaterMesh)
+	{
+		MarkRenderStateDirty();
+	}
+}
+
+void UWaterBodyComponent::UpdateNonTessellatedMeshSections()
+{
+	if (AWaterZone* WaterZone = GetWaterZone(); WaterZone && WaterZone->IsNonTessellatedLODMeshEnabled())
+	{
+		RebuildWaterBodyLODSections();
+	}
+	else
+	{
+		// Clear the mesh sections if LOD is disabled
+		WaterBodyMeshSections.Empty();
+	}
+
+	MarkRenderStateDirty();
 }
 
 #if WITH_EDITOR
diff --git a/Engine/Plugins/Experimental/Water/Source/Runtime/Private/WaterBodyLakeComponent.cpp b/Engine/Plugins/Experimental/Water/Source/Runtime/Private/WaterBodyLakeComponent.cpp
--- a/Engine/Plugins/Experimental/Water/Source/Runtime/Private/WaterBodyLakeComponent.cpp
+++ b/Engine/Plugins/Experimental/Water/Source/Runtime/Private/WaterBodyLakeComponent.cpp
@@ -49,19 +49,16 @@ TArray<UPrimitiveComponent*> UWaterBodyLakeComponent::GetStandardRenderableCompo
 	return Result;
 }
 
-void UWaterBodyLakeComponent::GenerateWaterBodyMesh()
+bool UWaterBodyLakeComponent::GenerateWaterBodyMesh(UE::Geometry::FDynamicMesh3& OutMesh, UE::Geometry::FDynamicMesh3* OutDilatedMesh) const
 {
-	TRACE_CPUPROFILER_EVENT_SCOPE(GenerateLakeMesh);
-
 	using namespace UE::Geometry;
 
-	WaterBodyMeshVertices.Empty();
-	WaterBodyMeshIndices.Empty();
+	TRACE_CPUPROFILER_EVENT_SCOPE(GenerateLakeMesh);
 
 	const UWaterSplineComponent* SplineComp = GetWaterSpline();
 	if (SplineComp == nullptr || SplineComp->GetNumberOfSplineSegments() < 3)
 	{
-		return;
+		return false;
 	}
 
 	FPolygon2d LakePoly;
@@ -82,61 +79,45 @@ void UWaterBodyLakeComponent::GenerateWaterBodyMesh()
 
 	if (Triangulation.Triangles.Num() == 0)
 	{
-		return;
+		return false;
 	}
 
 
-	// This FDynamicMesh3 will only be used to compute the inset region for shape dilation
-	FDynamicMesh3 LakeMesh(EMeshComponents::None);
 	for (const FVector2d& Vertex : Triangulation.Vertices)
 	{
-		// push the set of undilated vertices to the persistent mesh
-		FDynamicMeshVertex MeshVertex(FVector3f(Vertex.X, Vertex.Y, 0.f));
-		MeshVertex.Color = FColor::Black;
-		MeshVertex.TextureCoordinate[0].X = WaterBodyIndex;
-		WaterBodyMeshVertices.Add(MeshVertex);
+		FVertexInfo MeshVertex(FVector3d(Vertex.X, Vertex.Y, 0.f));
+		MeshVertex.Color = FVector3f(0.0);
+		MeshVertex.bHaveC = true;
 
-		LakeMesh.AppendVertex(FVector3d(Vertex, 0.0));
+		OutMesh.AppendVertex(MeshVertex);
 	}
 
 	for (const FIndex3i& Triangle : Triangulation.Triangles)
 	{
-		WaterBodyMeshIndices.Append({ (uint32)Triangle.A, (uint32)Triangle.B, (uint32)Triangle.C });
-		LakeMesh.AppendTriangle(Triangle);
+		OutMesh.AppendTriangle(Triangle);
 	}
 
-	if (ShapeDilation > 0.f)
+	if (ShapeDilation > 0.f && OutDilatedMesh)
 	{
 		// Inset the mesh by -ShapeDilation to effectively expand the mesh
-		FInsetMeshRegion Inset(&LakeMesh);
+		OutDilatedMesh->Copy(OutMesh);
+		FInsetMeshRegion Inset(OutDilatedMesh);
 		Inset.InsetDistance = -1 * ShapeDilation / 2.f;
 
-		Inset.Triangles.Reserve(LakeMesh.TriangleCount());
-		for (int32 Idx : LakeMesh.TriangleIndicesItr())
+		Inset.Triangles.Reserve(OutDilatedMesh->TriangleCount());
+		for (int32 Idx : OutDilatedMesh->TriangleIndicesItr())
 		{
 			Inset.Triangles.Add(Idx);
 		}
 		
-		if (Inset.Apply())
-		{
-			for (const FVector3d& Vertex : LakeMesh.GetVerticesBuffer())
-			{
-				// push the set of dilated vertices to the persistent mesh
-				FDynamicMeshVertex MeshVertex(FVector3f(Vertex.X, Vertex.Y, 0.f));
-				MeshVertex.Color = FColor::Black;
-				DilatedWaterBodyMeshVertices.Add(MeshVertex);
-			}
-
-			for (const FIndex3i& Triangle : LakeMesh.GetTrianglesBuffer())
-			{
-				DilatedWaterBodyMeshIndices.Append({ (uint32)Triangle.A, (uint32)Triangle.B, (uint32)Triangle.C });
-			}
-		}
-		else
+		if (!Inset.Apply())
 		{
 			UE_LOG(LogWater, Warning, TEXT("Failed to apply mesh inset for shape dilation (%s"), *GetOwner()->GetActorNameOrLabel());
+			return false;
 		}
 	}
+
+	return true;
 }
 
 FBoxSphereBounds UWaterBodyLakeComponent::CalcBounds(const FTransform& LocalToWorld) const
diff --git a/Engine/Plugins/Experimental/Water/Source/Runtime/Private/WaterBodyOceanComponent.cpp b/Engine/Plugins/Experimental/Water/Source/Runtime/Private/WaterBodyOceanComponent.cpp
--- a/Engine/Plugins/Experimental/Water/Source/Runtime/Private/WaterBodyOceanComponent.cpp
+++ b/Engine/Plugins/Experimental/Water/Source/Runtime/Private/WaterBodyOceanComponent.cpp
@@ -127,110 +127,259 @@ void UWaterBodyOceanComponent::Reset()
 	CollisionHullSets.Reset();
 }
 
-void UWaterBodyOceanComponent::GenerateWaterBodyMesh()
+// Adds a quad to a dynamic mesh reusing passed vertex indices if they are not INDEX_NONE
+static void AddAABBQuadToDynamicMesh(UE::Geometry::FDynamicMesh3& EditMesh, FVector3d Min, FVector3d Max, TArray<int32, TInlineAllocator<4>>& InOutVertices)
+{
+	using namespace UE::Geometry;
+
+	int32 VertexAIndex = InOutVertices[0];
+	int32 VertexBIndex = InOutVertices[1];
+	int32 VertexCIndex = InOutVertices[2];
+	int32 VertexDIndex = InOutVertices[3];
+
+	if (VertexAIndex == INDEX_NONE)
+	{ 
+		FVertexInfo VertexA;
+		VertexA.Position = Min;
+		VertexA.Color = FVector3f(0.0);
+		VertexA.bHaveC = true;
+		VertexAIndex = EditMesh.AppendVertex(VertexA);
+	}
+	
+	if (VertexBIndex == INDEX_NONE)
+	{
+		FVertexInfo VertexB;
+		VertexB.Position = FVector3d(Max.X, Min.Y, 0);
+		VertexB.Color = FVector3f(0.0);
+		VertexB.bHaveC = true;
+		VertexBIndex = EditMesh.AppendVertex(VertexB);
+	}
+
+	if (VertexCIndex == INDEX_NONE)
+	{
+		FVertexInfo VertexC;
+		VertexC.Position = Max;
+		VertexC.Color = FVector3f(0.0);
+		VertexC.bHaveC = true;
+		VertexCIndex = EditMesh.AppendVertex(VertexC);
+	}
+
+	if (VertexDIndex == INDEX_NONE)
+	{
+		FVertexInfo VertexD;
+		VertexD.Position = FVector3d(Min.X, Max.Y, 0);
+		VertexD.Color = FVector3f(0.0);
+		VertexD.bHaveC = true;
+		VertexDIndex = EditMesh.AppendVertex(VertexD);
+	}
+
+	const int TriOne = EditMesh.AppendTriangle(VertexAIndex, VertexDIndex, VertexBIndex);
+	const int TriTwo = EditMesh.AppendTriangle(VertexBIndex, VertexDIndex, VertexCIndex);
+	check(TriOne != INDEX_NONE && TriTwo != INDEX_NONE);
+
+	InOutVertices = {VertexAIndex, VertexBIndex, VertexCIndex, VertexDIndex};
+}
+
+bool UWaterBodyOceanComponent::GenerateWaterBodyMesh(UE::Geometry::FDynamicMesh3& OutMesh, UE::Geometry::FDynamicMesh3* OutDilatedMesh) const
 {
 	TRACE_CPUPROFILER_EVENT_SCOPE(GenerateOceanMesh);
 
 	using namespace UE::Geometry;
 
-	WaterBodyMeshVertices.Empty();
-	WaterBodyMeshIndices.Empty();
-
 	const UWaterSplineComponent* SplineComp = GetWaterSpline();
+	const FVector OceanLocation = GetComponentLocation();
 	
 	if (SplineComp->GetNumberOfSplineSegments() < 3)
 	{
-		return;
+		return false;
 	}
 
 	FPolygon2d Island;
 	TArray<FVector> PolyLineVertices;
 	SplineComp->ConvertSplineToPolyLine(ESplineCoordinateSpace::Local, FMath::Square(10.f), PolyLineVertices);
-	
+
+	FAxisAlignedBox2d IslandBounds = FAxisAlignedBox2d(FVector2d(0.f), FVector2d(0.f));
 	// Construct a 2D polygon describing the central island
 	for (int32 i = PolyLineVertices.Num() - 2; i >= 0; --i) // skip the last vertex since it's the same as the first vertex
 	{
 		Island.AppendVertex(FVector2D(PolyLineVertices[i]));
+		IslandBounds.Contain(FVector2d(PolyLineVertices[i]));
 	}
+	// Expand the island slightly so we aren't intersecting with the spline
+	IslandBounds.Expand(1);
+
+	// #todo: account for scale
+	const FAxisAlignedBox2d OceanBounds = FAxisAlignedBox2d(-FVector2d(VisualExtents) / 2.0, FVector2d(VisualExtents) / 2.0);
+	FPolygon2d IslandBoundingPolygon = FPolygon2d::MakeRectangle(FVector2d(IslandBounds.Center()), 2 * IslandBounds.Extents().X, 2 * IslandBounds.Extents().Y);
 
-	FVector OceanLocation = GetComponentLocation();
-	FPolygon2d OceanBoundingPolygon = FPolygon2d::MakeRectangle(FVector2d(0, 0), VisualExtents.X, VisualExtents.Y);
 
 	FConstrainedDelaunay2d Triangulation;
 	Triangulation.FillRule = FConstrainedDelaunay2d::EFillRule::Positive;
-	Triangulation.Add(OceanBoundingPolygon);
+	Triangulation.Add(IslandBoundingPolygon);
+
 	if (!Island.IsClockwise())
 	{
 		Island.Reverse();
 	}
 	Triangulation.Add(Island);
-	Triangulation.Triangulate();
+	if (!Triangulation.Triangulate())
+	{
+		UE_LOG(LogWater, Warning, TEXT("Failed to triangulate ocean (%s). Spline might not be fully contained by VisualExtents."), *GetOwner()->GetActorNameOrLabel());
+		return false;
+	}
 
 	if (Triangulation.Triangles.Num() == 0)
 	{
-		return;
+		return false;
 	}
+	
+	int32 IslandBottomLeft = INDEX_NONE;
+	FVector3d IslandBottomLeftVertex;
+
+	int32 IslandBottomRight = INDEX_NONE;
+	FVector3d IslandBottomRightVertex;
+	
+	int32 IslandTopRight = INDEX_NONE;
+	FVector3d IslandTopRightVertex;
+
+	int32 IslandTopLeft = INDEX_NONE;
+	FVector3d IslandTopLeftVertex;
 
-	// This FDynamicMesh3 will only be used to compute the inset region for shape dilation
-	FDynamicMesh3 OceanMesh(EMeshComponents::None);
 	for (const FVector2d& Vertex : Triangulation.Vertices)
 	{
-		// push the set of undilated vertices to the persistent mesh
-		FDynamicMeshVertex MeshVertex(FVector3f(Vertex.X, Vertex.Y, 0.f));
-		MeshVertex.Color = FColor::Black;
-		MeshVertex.TextureCoordinate[0].X = WaterBodyIndex;
-		WaterBodyMeshVertices.Add(MeshVertex);
+		FVertexInfo VertexInfo;
+		VertexInfo.Position = FVector3d(Vertex, 0.f);
+		VertexInfo.Color = FVector3f(0.0);
+		VertexInfo.bHaveC = true;
+
+		int32 Index = OutMesh.AppendVertex(VertexInfo);
 
-		OceanMesh.AppendVertex(FVector3d(Vertex, 0.0));
+		// Collect the corner vertices of the island bounding box so we can stitch the outer quads to them.
+		if ((IslandBottomLeft == INDEX_NONE) || (VertexInfo.Position.X < IslandBottomLeftVertex.X || VertexInfo.Position.Y < IslandBottomLeftVertex.Y))
+		{
+			IslandBottomLeft = Index;
+			IslandBottomLeftVertex = VertexInfo.Position;
+		}
+		if ((IslandBottomRight == INDEX_NONE) || (VertexInfo.Position.X > IslandBottomRightVertex.X || VertexInfo.Position.Y < IslandBottomRightVertex.Y))
+		{
+			IslandBottomRight = Index;
+			IslandBottomRightVertex = VertexInfo.Position;
+		}
+		if ((IslandTopRight == INDEX_NONE) || (VertexInfo.Position.X > IslandTopRightVertex.X || VertexInfo.Position.Y > IslandTopRightVertex.Y))
+		{
+			IslandTopRight = Index;
+			IslandTopRightVertex = VertexInfo.Position;
+		}
+		if ((IslandTopLeft == INDEX_NONE) || (VertexInfo.Position.X < IslandTopLeftVertex.X || VertexInfo.Position.Y > IslandTopLeftVertex.Y))
+		{
+			IslandTopLeft = Index;
+			IslandTopLeftVertex = VertexInfo.Position;
+		}
 	}
+	check(IslandBottomLeft != INDEX_NONE &&
+		  IslandBottomRight != INDEX_NONE &&
+		  IslandTopRight != INDEX_NONE &&
+		  IslandTopLeft != INDEX_NONE);
 
 	for (const FIndex3i& Triangle : Triangulation.Triangles)
 	{
-		WaterBodyMeshIndices.Append({ (uint32)Triangle.A, (uint32)Triangle.B, (uint32)Triangle.C });
-		OceanMesh.AppendTriangle(Triangle);
+		OutMesh.AppendTriangle(Triangle);
 	}
 
+	// Add the bounding quads:
 
-	if (ShapeDilation > 0.f)
+	// Each quad has four vertices laid out like so
+	// D --- C
+	// |  \  |
+	// A --- B
+	
+	TArray<int32, TInlineAllocator<4>> Vertices({ INDEX_NONE, INDEX_NONE, INDEX_NONE, INDEX_NONE });
+
+	// Bottom left quad
+	Vertices[2] = IslandBottomLeft;
+	AddAABBQuadToDynamicMesh(OutMesh, FVector3d(OceanBounds.Min, 0), FVector3d(IslandBounds.Min, 0), Vertices);
+	// We need the B and C vertices of the first quad to attach to the last quad
+	int32 FirstQuadB = Vertices[1];
+	int32 FirstQuadC = Vertices[2];
+
+	// Left middle quad
+	Vertices[0] = Vertices[3];
+	Vertices[1] = Vertices[2];
+	Vertices[2] = IslandTopLeft;
+	Vertices[3] = INDEX_NONE;
+	AddAABBQuadToDynamicMesh(OutMesh, FVector3d(OceanBounds.Min.X, IslandBounds.Min.Y, 0), FVector3d(IslandBounds.Min.X, IslandBounds.Max.Y, 0), Vertices);
+
+	// Top left quad
+	Vertices[0] = Vertices[3];
+	Vertices[1] = Vertices[2];
+	Vertices[2] = INDEX_NONE;
+	Vertices[3] = INDEX_NONE;
+	AddAABBQuadToDynamicMesh(OutMesh, FVector3d(OceanBounds.Min.X, IslandBounds.Max.Y, 0), FVector3d(IslandBounds.Min.X, OceanBounds.Max.Y, 0), Vertices);
+
+	// Top middle quad
+	Vertices[0] = Vertices[1];
+	Vertices[3] = Vertices[2];
+	Vertices[1] = IslandTopRight;
+	Vertices[2] = INDEX_NONE;
+	AddAABBQuadToDynamicMesh(OutMesh, FVector3d(IslandBounds.Min.X, IslandBounds.Max.Y, 0), FVector3d(IslandBounds.Max.X, OceanBounds.Max.Y, 0), Vertices);
+
+	// Top right quad
+	Vertices[0] = Vertices[1];
+	Vertices[3] = Vertices[2];
+	Vertices[1] = INDEX_NONE;
+	Vertices[2] = INDEX_NONE;
+	AddAABBQuadToDynamicMesh(OutMesh, FVector3d(IslandBounds.Max, 0), FVector3d(OceanBounds.Max, 0), Vertices);
+
+	// Middle right quad
+	Vertices[3] = Vertices[0];
+	Vertices[2] = Vertices[1];
+	Vertices[1] = INDEX_NONE;
+	Vertices[0] = IslandBottomRight;
+	AddAABBQuadToDynamicMesh(OutMesh, FVector3d(IslandBounds.Max.X, IslandBounds.Min.Y, 0), FVector3d(OceanBounds.Max.X, IslandBounds.Max.Y, 0), Vertices);
+
+	// Bottom right quad
+	Vertices[3] = Vertices[0];
+	Vertices[2] = Vertices[1];
+	Vertices[1] = INDEX_NONE;
+	Vertices[0] = INDEX_NONE;
+	AddAABBQuadToDynamicMesh(OutMesh, FVector3d(IslandBounds.Max.X, OceanBounds.Min.Y, 0), FVector3d(OceanBounds.Max.X, IslandBounds.Min.Y, 0), Vertices);
+
+	// Middle bottom quad
+	Vertices[1] = Vertices[0];
+	Vertices[2] = Vertices[3];
+	Vertices[0] = FirstQuadB;
+	Vertices[3] = FirstQuadC;
+	AddAABBQuadToDynamicMesh(OutMesh, FVector3d(IslandBounds.Min.X, OceanBounds.Min.Y, 0), FVector3d(IslandBounds.Max.X, IslandBounds.Min.Y, 0), Vertices);
+	
+	if (ShapeDilation > 0.f && OutDilatedMesh)
 	{
 		TRACE_CPUPROFILER_EVENT_SCOPE(DilateOceanMesh);
 
 		// Inset the mesh by -ShapeDilation to effectively expand the mesh
-		FInsetMeshRegion Inset(&OceanMesh);
+		OutDilatedMesh->Copy(OutMesh);
+		FInsetMeshRegion Inset(OutDilatedMesh);
 		Inset.InsetDistance = -1 * ShapeDilation / 2.f;
 
-		Inset.Triangles.Reserve(OceanMesh.TriangleCount());
-		for (int32 Idx : OceanMesh.TriangleIndicesItr())
+		Inset.Triangles.Reserve(OutDilatedMesh->TriangleCount());
+		for (int32 Idx : OutDilatedMesh->TriangleIndicesItr())
 		{
 			Inset.Triangles.Add(Idx);
 		}
 		
-		if (Inset.Apply())
-		{
-			for (const FVector3d& Vertex : OceanMesh.GetVerticesBuffer())
-			{
-				// push the set of dilated vertices to the persistent mesh
-				FDynamicMeshVertex MeshVertex(FVector3f(Vertex.X, Vertex.Y, 0.f));
-				MeshVertex.Color = FColor::Black;
-				DilatedWaterBodyMeshVertices.Add(MeshVertex);
-			}
-
-			for (const FIndex3i& Triangle : OceanMesh.GetTrianglesBuffer())
-			{
-				DilatedWaterBodyMeshIndices.Append({ (uint32)Triangle.A, (uint32)Triangle.B, (uint32)Triangle.C });
-			}
-		}
-		else
+		if (!Inset.Apply())
 		{
-			UE_LOG(LogWater, Warning, TEXT("Failed to apply mesh inset for shape dilation (%s"), *GetOwner()->GetActorNameOrLabel());
+			UE_LOG(LogWater, Warning, TEXT("Failed to apply mesh inset for shape dilation (%s)"), *GetOwner()->GetActorNameOrLabel());
+			return false;
 		}
 	}	
+	return true;
 }
 
 void UWaterBodyOceanComponent::PostLoad()
 {
 	Super::PostLoad();
 
+#ifdef WITH_EDITORONLY_DATA
 	if (GetLinkerCustomVersion(FFortniteMainBranchObjectVersion::GUID) < FFortniteMainBranchObjectVersion::WaterZonesRefactor)
 	{
 		if (AWaterZone* WaterZone = GetWaterZone())
@@ -238,6 +387,13 @@ void UWaterBodyOceanComponent::PostLoad()
 			SetVisualExtents(WaterZone->GetZoneExtent());
 		}
 	}
+	if (GetLinkerCustomVersion(FFortniteMainBranchObjectVersion::GUID) < FFortniteMainBranchObjectVersion::WaterNontessellatedLODSupportAdded)
+	{
+		// Prior to the refactors when introducing the non-tessellated LOD, the Ocean's visual extent was treated as 2x the actual value.
+		// To maintain visual consistency with previous versions, the value needs to now be updated to be 2x.
+		VisualExtents = VisualExtents * 2.;
+	}
+#endif // WITH_EDITORONLY_DATA
 }
 
 FBoxSphereBounds UWaterBodyOceanComponent::CalcBounds(const FTransform& LocalToWorld) const
diff --git a/Engine/Plugins/Experimental/Water/Source/Runtime/Private/WaterBodyRiverComponent.cpp b/Engine/Plugins/Experimental/Water/Source/Runtime/Private/WaterBodyRiverComponent.cpp
--- a/Engine/Plugins/Experimental/Water/Source/Runtime/Private/WaterBodyRiverComponent.cpp
+++ b/Engine/Plugins/Experimental/Water/Source/Runtime/Private/WaterBodyRiverComponent.cpp
@@ -11,6 +11,7 @@
 #include "PhysicsEngine/BodySetup.h"
 #include "Engine/StaticMesh.h"
 #include "Algo/Transform.h"
+#include "DynamicMesh/DynamicMesh3.h"
 
 // for working around Chaos issue
 #include "Chaos/Particles.h"
@@ -90,39 +91,15 @@ void UWaterBodyRiverComponent::OnUpdateBody(bool bWithExclusionVolumes)
 
 // ----------------------------------------------------------------------------------
 
-static FColor PackWaterFlow(float VelocityMagnitude, float DirectionAngle)
-{
-	check((DirectionAngle >= 0.f) && (DirectionAngle <= TWO_PI));
-
-	const float MaxVelocity = FWaterUtils::GetWaterMaxFlowVelocity(false);
-
-	float NormalizedMagnitude = FMath::Clamp(VelocityMagnitude, 0.f, MaxVelocity) / MaxVelocity;
-	float NormalizedAngle = DirectionAngle / TWO_PI;
-	float MappedMagnitude = NormalizedMagnitude * TNumericLimits<uint16>::Max();
-	float MappedAngle = NormalizedAngle * TNumericLimits<uint16>::Max();
-
-	uint16 ResultMag = (uint16)MappedMagnitude;
-	uint16 ResultAngle = (uint16)MappedAngle;
-	
-	FColor Result;
-	Result.R = (ResultMag >> 8) & 0xFF;
-	Result.G = (ResultMag >> 0) & 0xFF;
-	Result.B = (ResultAngle >> 8) & 0xFF;
-	Result.A = (ResultAngle >> 0) & 0xFF;
-
-	return Result;
-}
-
 static void AddVerticesForRiverSplineStep(
 	float DistanceAlongSpline, 
 	const UWaterBodyRiverComponent* Component, 
 	const UWaterSplineComponent* SplineComp, 
 	const UWaterSplineMetadata* WaterSplineMetadata,
-	TArray<FDynamicMeshVertex>& Vertices, 
-	TArray<uint32>& Indices, 
-	TArray<FDynamicMeshVertex>& DilatedVertices, 
-	TArray<uint32>& DilatedIndices)
+	UE::Geometry::FDynamicMesh3& OutMesh,
+	UE::Geometry::FDynamicMesh3* OutDilatedMesh)
 {
+	using namespace UE::Geometry;
 	check((Component != nullptr) && (SplineComp != nullptr) && (WaterSplineMetadata != nullptr));
 	
 	const FVector Tangent = SplineComp->GetTangentAtDistanceAlongSpline(DistanceAlongSpline, ESplineCoordinateSpace::Local).GetSafeNormal();
@@ -140,17 +117,6 @@ static void AddVerticesForRiverSplineStep(
 	// Prevent there being a relative height difference between the two vertices even when the spline has a slight roll to it
 	OutwardDistance.Z = 0.f;
 
-	const float DilationAmount = Component->ShapeDilation;
-	const FVector DilationOffset = Normal * DilationAmount;
-
-	FDynamicMeshVertex Left(FVector3f(Pos - OutwardDistance));
-	FDynamicMeshVertex Right(FVector3f(Pos + OutwardDistance));
-
-	FDynamicMeshVertex DilatedFarLeft(FVector3f(Pos - OutwardDistance - DilationOffset));
-	FDynamicMeshVertex DilatedLeft(FVector3f(Pos - OutwardDistance));
-	FDynamicMeshVertex DilatedRight(FVector3f(Pos + OutwardDistance));
-	FDynamicMeshVertex DilatedFarRight(FVector3f(Pos + OutwardDistance + DilationOffset));
-
 	float FlowDirection = Tangent.HeadingAngle() + FMath::DegreesToRadians(Component->GetRelativeRotation().Yaw);
 
 	// Restrict all angles between [0, 2 PI]. UnwindRadians returns a value between [-Pi, Pi] so we must remap again:
@@ -167,41 +133,73 @@ static void AddVerticesForRiverSplineStep(
 		FlowDirection = FMath::Fmod(PI + FlowDirection, TWO_PI);
 	}
 
-	const FColor EmptyFlowData = FColor(0.f);
-
-	const FColor PackedFlowData = PackWaterFlow(Velocity, FlowDirection);
-	Left.Color = PackedFlowData;
-	Right.Color = PackedFlowData;
+	FVertexInfo Left(FVector3d(Pos - OutwardDistance));
+	FVertexInfo Right(FVector3d(Pos + OutwardDistance));
 
-	DilatedFarLeft.Color = EmptyFlowData;
-	DilatedLeft.Color = EmptyFlowData;
-	DilatedRight.Color = EmptyFlowData;
-	DilatedFarRight.Color = EmptyFlowData;
+	const FVector3f FlowData = FVector3f(Velocity, FlowDirection, 0.f);
+	Left.Color = FlowData;
+	Right.Color = FlowData;
+	Left.bHaveC = true;
+	Right.bHaveC = true;
 
+	// Append regular geometry to mesh:
 	{
 		/* Non - dilated river segment geometry:
-		   2 --- 3
-		   |  /  |
-		   0 --- 1
+		    0 --- 1
+		    |  /  |
+		   -2 --- -1
 		*/
-		const uint32 BaseIndex = Vertices.Num();
-		Vertices.Append({ Left, Right });
-		Indices.Append({ BaseIndex, BaseIndex + 3, BaseIndex + 1, BaseIndex, BaseIndex + 2, BaseIndex + 3 });
+		const int32 BaseIndex = OutMesh.GetVerticesBuffer().Num();
+
+		OutMesh.AppendVertex(Left);
+		OutMesh.AppendVertex(Right);
+
+		if (BaseIndex != 0)
+		{
+			OutMesh.AppendTriangle(BaseIndex - 2, BaseIndex + 1, BaseIndex - 1);
+			OutMesh.AppendTriangle(BaseIndex - 2, BaseIndex, BaseIndex + 1);
+		}
 	}
 	
-	if (DilationAmount > 0.f)
+	const float DilationAmount = Component->ShapeDilation;
+	// If dilation is required, append dilated geometry:
+	if (DilationAmount > 0.f && OutDilatedMesh)
 	{
+		const FVector DilationOffset = Normal * DilationAmount;
+
+		const FVertexInfo DilatedFarLeft(FVector3d(Pos - OutwardDistance - DilationOffset));
+		const FVertexInfo DilatedLeft(FVector3d(Pos - OutwardDistance));
+		const FVertexInfo DilatedRight(FVector3d(Pos + OutwardDistance));
+		const FVertexInfo DilatedFarRight(FVector3d(Pos + OutwardDistance + DilationOffset));
+
 		/* Dilated River segment geometry:
-			4---5   6---7
-			| / |   | / |
-			0---1   2---3
+			 1 --- 2     3 --- 4
+			 |  /  |      |  /  |
+			-4 --- -3    -2 --- -1
 		*/
-		const uint32 BaseIndex = DilatedVertices.Num();
-		DilatedVertices.Append({ DilatedFarLeft, DilatedLeft, DilatedRight, DilatedFarRight });
-		// Append left dilation quad
-		DilatedIndices.Append({ BaseIndex, BaseIndex + 5, BaseIndex + 1, BaseIndex, BaseIndex + 4, BaseIndex + 5 });
-		// Append right dilation quad
-		DilatedIndices.Append({ BaseIndex + 2, BaseIndex + 6, BaseIndex + 7, BaseIndex + 2, BaseIndex + 7, BaseIndex + 3 });
+		const int32 BaseIndex = OutDilatedMesh->GetVerticesBuffer().Num();
+
+		OutDilatedMesh->AppendVertex(DilatedFarLeft);
+		OutDilatedMesh->AppendVertex(DilatedLeft);
+		OutDilatedMesh->AppendVertex(DilatedRight);
+		OutDilatedMesh->AppendVertex(DilatedFarRight);
+
+		// If this is the first point we need to add the triangles for the starting dilated center quad since that references vertices which were only just appended
+		if (DistanceAlongSpline == 0.f)
+		{
+			OutDilatedMesh->AppendTriangle(BaseIndex - 3, BaseIndex + 2, BaseIndex - 2);
+			OutDilatedMesh->AppendTriangle(BaseIndex - 3, BaseIndex + 1, BaseIndex + 2);
+		}
+		
+		if (BaseIndex != 0)
+		{
+			// Append left dilation quad
+			OutDilatedMesh->AppendTriangle(BaseIndex - 4, BaseIndex + 1, BaseIndex - 3);
+			OutDilatedMesh->AppendTriangle(BaseIndex - 4, BaseIndex, BaseIndex + 1);
+			// Append right dilation quad
+			OutDilatedMesh->AppendTriangle(BaseIndex - 2, BaseIndex + 2, BaseIndex + 3);
+			OutDilatedMesh->AppendTriangle(BaseIndex - 2, BaseIndex + 3, BaseIndex - 1);
+		}
 	}
 }
 
@@ -210,8 +208,15 @@ enum class ERiverBoundaryEdge {
 	End,
 };
 
-static void AddTerminalVerticesForRiverSpline(ERiverBoundaryEdge Edge, const UWaterBodyRiverComponent* Component, const UWaterSplineComponent* SplineComp, const UWaterSplineMetadata* WaterSplineMetadata, TArray<FDynamicMeshVertex>& Vertices, TArray<uint32>& Indices)
+static void AddTerminalVerticesForRiverSpline(
+	ERiverBoundaryEdge Edge, 
+	const UWaterBodyRiverComponent* Component,
+	const UWaterSplineComponent* SplineComp, 
+	const UWaterSplineMetadata* WaterSplineMetadata, 
+	UE::Geometry::FDynamicMesh3& OutDilatedMesh)
 {
+	using namespace UE::Geometry;
+
 	check((Component != nullptr) && (SplineComp != nullptr) && (WaterSplineMetadata != nullptr));
 
 	const float DistanceAlongSpline = (Edge == ERiverBoundaryEdge::Start) ? 0.f : SplineComp->GetSplineLength();
@@ -239,90 +244,73 @@ static void AddTerminalVerticesForRiverSpline(ERiverBoundaryEdge Edge, const UWa
 		TangentialOffset *= -1;
 	}
 
-	FDynamicMeshVertex BackLeft(FVector3f(Pos - OutwardDistance + TangentialOffset - DilationOffset));
-	FDynamicMeshVertex Left(FVector3f(Pos - OutwardDistance + TangentialOffset));
-	FDynamicMeshVertex Right(FVector3f(Pos + OutwardDistance + TangentialOffset));
-	FDynamicMeshVertex BackRight(FVector3f(Pos + OutwardDistance + TangentialOffset + DilationOffset));
+	const FVertexInfo BackLeft(FVector3d(Pos - OutwardDistance + TangentialOffset - DilationOffset));
+	const FVertexInfo Left(FVector3d(Pos - OutwardDistance + TangentialOffset));
+	const FVertexInfo Right(FVector3d(Pos + OutwardDistance + TangentialOffset));
+	const FVertexInfo BackRight(FVector3d(Pos + OutwardDistance + TangentialOffset + DilationOffset));
 
-	// Initialize the vertex data to correct represent what we expect the dilated region to look like (no flow data)
-	BackLeft.Color = FColor(0.f);
-	Left.Color = FColor(0.f);
-	Right.Color = FColor(0.f);
-	BackRight.Color = FColor(0.f);
+	const int32 BaseIndex = OutDilatedMesh.GetVerticesBuffer().Num();
 
-	/* Dilated edge segment geometry:
-		4---5-----6---7
-		| / |  /  | / |
-		0---1-----2---3
-	*/
+	OutDilatedMesh.AppendVertex(BackLeft);
+	OutDilatedMesh.AppendVertex(Left);
+	OutDilatedMesh.AppendVertex(Right);
+	OutDilatedMesh.AppendVertex(BackRight);
 
-	Vertices.Append({ BackLeft, Left, Right, BackRight });
 
-	const uint32 BaseIndex = Edge == ERiverBoundaryEdge::Start ? 0 : Vertices.Num() - 8;
+	if (Edge == ERiverBoundaryEdge::End)
+	{
+		/* Dilated edge segment geometry:
+			 0 --- 1 ------- 2 --- 3
+			 |  /  |    /    |  /  |
+			-4 --- -3 ----- -2 --- -1
+		*/
 
-	// Since iterating the spline returns the final distance, all the side quads are already added for the final point so we only need to do this for the first point.
-	if (Edge == ERiverBoundaryEdge::Start)
-	{ 
+		OutDilatedMesh.AppendTriangle(BaseIndex - 4, BaseIndex + 0, BaseIndex + 1);
+		OutDilatedMesh.AppendTriangle(BaseIndex - 4, BaseIndex + 1, BaseIndex - 3);
+
+		OutDilatedMesh.AppendTriangle(BaseIndex - 3, BaseIndex + 1, BaseIndex + 2);
+		OutDilatedMesh.AppendTriangle(BaseIndex - 3, BaseIndex + 2, BaseIndex - 2);
 
-		Indices.Append({ BaseIndex + 0, BaseIndex + 5, BaseIndex + 1, BaseIndex + 0, BaseIndex + 4, BaseIndex + 5 });
-		Indices.Append({ BaseIndex + 2, BaseIndex + 7, BaseIndex + 3, BaseIndex + 2, BaseIndex + 6, BaseIndex + 7 });
+		OutDilatedMesh.AppendTriangle(BaseIndex - 2, BaseIndex + 2, BaseIndex + 3);
+		OutDilatedMesh.AppendTriangle(BaseIndex - 2, BaseIndex + 3, BaseIndex - 1);
 	}
-	Indices.Append({ BaseIndex + 1, BaseIndex + 6, BaseIndex + 2, BaseIndex + 1, BaseIndex + 5, BaseIndex + 6 });
 }
 
 // ----------------------------------------------------------------------------------
 
-void UWaterBodyRiverComponent::GenerateWaterBodyMesh()
+bool UWaterBodyRiverComponent::GenerateWaterBodyMesh(UE::Geometry::FDynamicMesh3& OutMesh, UE::Geometry::FDynamicMesh3* OutDilatedMesh) const
 {
 	TRACE_CPUPROFILER_EVENT_SCOPE(GenerateRiverMesh);
 
-	WaterBodyMeshVertices.Empty();
-	WaterBodyMeshIndices.Empty();
-
 	const UWaterSplineComponent* SplineComp = GetWaterSpline();
 	if ((SplineComp == nullptr) || (WaterSplineMetadata == nullptr) || (SplineComp->GetNumberOfSplinePoints() < 2))
 	{
-		return;
+		return false;
 	}
 
 	TArray<double> Distances;
 	TArray<FVector> Points;
-	SplineComp->DivideSplineIntoPolylineRecursiveWithDistances(0.f, SplineComp->GetSplineLength(), ESplineCoordinateSpace::Local, FMath::Square(10.f), Points, Distances);
+	SplineComp->DivideSplineIntoPolylineRecursiveWithDistances(0.f, SplineComp->GetSplineLength(), ESplineCoordinateSpace::Local, FMath::Square(5.f), Points, Distances);
 	if (Distances.Num() == 0)
 	{
-		return;
+		return false;
 	}
 	
-	TArray<FDynamicMeshVertex> Vertices;
-	TArray<uint32> Indices;
-
-	TArray<FDynamicMeshVertex> DilatedVertices;
-	TArray<uint32> DilatedIndices;
-
-	// Add an extra point at the start to dilate starting edge
-	if (ShapeDilation > 0.f)
+	if (ShapeDilation > 0.f && OutDilatedMesh)
 	{
-		AddTerminalVerticesForRiverSpline(ERiverBoundaryEdge::Start, this, SplineComp, WaterSplineMetadata, DilatedVertices, DilatedIndices);
+		AddTerminalVerticesForRiverSpline(ERiverBoundaryEdge::Start, this, SplineComp, WaterSplineMetadata, *OutDilatedMesh);
 	}
 
 	for (double DistanceAlongSpline : Distances)
 	{
-		AddVerticesForRiverSplineStep(DistanceAlongSpline, this, SplineComp, WaterSplineMetadata, Vertices, Indices, DilatedVertices, DilatedIndices);
+		AddVerticesForRiverSplineStep(DistanceAlongSpline, this, SplineComp, WaterSplineMetadata, OutMesh, OutDilatedMesh);
 	}
-	// Remove the last two triangles since AddverticesForRiverSplineStep doesn't know which distance is the final distance and tries to create a triangle
-	// which links to vertices that don't exist.
-	Indices.SetNum(Indices.Num() - 6);
 	
-	// Add an extra point at the end to dilate ending edge
-	if (ShapeDilation > 0.f)
+	if (ShapeDilation > 0.f && OutDilatedMesh)
 	{
-		AddTerminalVerticesForRiverSpline(ERiverBoundaryEdge::End, this, SplineComp, WaterSplineMetadata, DilatedVertices, DilatedIndices);
+		AddTerminalVerticesForRiverSpline(ERiverBoundaryEdge::End, this, SplineComp, WaterSplineMetadata, *OutDilatedMesh);
 	}
-
-	WaterBodyMeshVertices = MoveTemp(Vertices);
-	WaterBodyMeshIndices = MoveTemp(Indices);
-	DilatedWaterBodyMeshVertices = MoveTemp(DilatedVertices);
-	DilatedWaterBodyMeshIndices = MoveTemp(DilatedIndices);
+	return true;
 }
 
 FBoxSphereBounds UWaterBodyRiverComponent::CalcBounds(const FTransform& LocalToWorld) const
diff --git a/Engine/Plugins/Experimental/Water/Source/Runtime/Private/WaterBodySceneProxy.cpp b/Engine/Plugins/Experimental/Water/Source/Runtime/Private/WaterBodySceneProxy.cpp
--- a/Engine/Plugins/Experimental/Water/Source/Runtime/Private/WaterBodySceneProxy.cpp
+++ b/Engine/Plugins/Experimental/Water/Source/Runtime/Private/WaterBodySceneProxy.cpp
@@ -7,90 +7,204 @@
 #include "WaterBodyComponent.h"
 #include "Materials/MaterialInstanceDynamic.h"
 #include "SceneInterface.h"
+#include "PrimitiveSceneInfo.h"
+#include "Algo/Transform.h"
 
 
-static TAutoConsoleVariable<int32> CVarWaterShowProxies(
-	TEXT("r.Water.ShowWaterSceneProxies"),
-	0,
-	TEXT("Allows editor visualization of water scene proxies. If the mode is set to 1 we will show only the selected water body in wireframe, if it is set to 2 we will show all in wireframe, and if it is 3 we will show all as opaque meshes"),
-	ECVF_Default);
+FWaterBodySceneProxy::FWaterBodySceneProxy(UWaterBodyComponent* Component, const FBox2D& TessellatedMeshBounds)
+	: FPrimitiveSceneProxy(Component)
+	, WaterBodySectionedLODMesh(GetScene().GetFeatureLevel())
+	, WaterBodyInfoMesh(GetScene().GetFeatureLevel())
+	, WaterBodyInfoDilatedMesh(GetScene().GetFeatureLevel())
+{
+	TRACE_CPUPROFILER_EVENT_SCOPE("FWaterBodySceneProxy::FWaterBodySceneProxy");
+	WaterBodySectionedLODMesh.InitFromSections(Component->WaterBodyMeshSections);
+	WaterBodyInfoMesh.Init(Component->WaterBodyMeshVertices, Component->WaterBodyMeshIndices);
+	WaterBodyInfoDilatedMesh.Init(Component->DilatedWaterBodyMeshVertices, Component->DilatedWaterBodyMeshIndices);
+	WaterBodySectionedLODMesh.RebuildIndexBuffer(TessellatedMeshBounds);
 
-// ----------------------------------------------------------------------------------
+	if (UMaterialInstance* WaterInfoMaterialInstance = Component->GetWaterInfoMaterialInstance())
+	{
+		WaterInfoMaterial = WaterInfoMaterialInstance->GetRenderProxy();
+		WaterInfoMaterialRelevance |= WaterInfoMaterialInstance->GetRelevance_Concurrent(GetScene().GetFeatureLevel());
+	}
+	if (UMaterialInstance* WaterMaterialInstance = Component->GetWaterLODMaterialInstance())
+	{
+		WaterLODMaterial = WaterMaterialInstance->GetRenderProxy();
+		WaterLODMaterialRelevance |= WaterMaterialInstance->GetRelevance_Concurrent(GetScene().GetFeatureLevel());
+	}
+}
 
-class FWaterBodyMeshSection
+void FWaterBodySceneProxy::FWaterBodySectionedLODMesh::RebuildIndexBuffer(const FBox2D& TessellatedWaterMeshBounds)
 {
-public:
-	FStaticMeshVertexBuffers VertexBuffers;
-	FDynamicMeshIndexBuffer32 IndexBuffer;
-	FLocalVertexFactory VertexFactory;
+	TRACE_CPUPROFILER_EVENT_SCOPE("FWaterBodySceneProxy::RebuildIndexBuffer");
 
-	FWaterBodyMeshSection(ERHIFeatureLevel::Type InFeatureLevel)
-		: VertexFactory(InFeatureLevel, "FWaterBodyMeshSection")
-	{}
-};
+	TArray<uint32> Indices;
+	for (const FWaterBodySectionedMeshProxy& MeshSectionProxy : Sections)
+	{
+		if (!(TessellatedWaterMeshBounds.IsInsideOrOn(MeshSectionProxy.Bounds.Min) && TessellatedWaterMeshBounds.IsInsideOrOn(MeshSectionProxy.Bounds.Max)))
+		{
+			Indices.Append(MeshSectionProxy.Indices);
+		}
+	}
 
-void FWaterBodySceneProxy::InitResources(FWaterBodyMeshSection& Section)
-{
-	BeginInitResource(&Section.VertexBuffers.PositionVertexBuffer);
-	BeginInitResource(&Section.VertexBuffers.StaticMeshVertexBuffer);
-	BeginInitResource(&Section.VertexBuffers.ColorVertexBuffer);
-	BeginInitResource(&Section.IndexBuffer);
-	BeginInitResource(&Section.VertexFactory);
+	IndexBuffer.Indices = MoveTemp(Indices);
+
+	if (IndexBuffer.Indices.Num() > 0)
+	{
+		if (IndexBuffer.IsInitialized())
+		{
+			BeginUpdateResourceRHI(&IndexBuffer);
+		}
+		else
+		{
+			BeginInitResource(&IndexBuffer);
+
+			bInitialized = true;
+		}
+	}
 }
 
-FWaterBodySceneProxy::FWaterBodySceneProxy(UWaterBodyComponent* Component)
-	: FPrimitiveSceneProxy(Component)
+void FWaterBodySceneProxy::FWaterBodySectionedLODMesh::InitFromSections(const TArray<FWaterBodyMeshSection>& MeshSections)
 {
-	if (Component->WaterBodyMeshVertices.Num() > 0 && Component->WaterBodyMeshIndices.Num() > 0)
+	TArray<FDynamicMeshVertex> Vertices;
+	for (const FWaterBodyMeshSection& MeshSection : MeshSections)
 	{
-		FWaterBodyMeshSection& NewSection = Sections.Emplace_GetRef(GetScene().GetFeatureLevel());
-		NewSection.IndexBuffer.Indices = Component->WaterBodyMeshIndices;
+		if (MeshSection.Vertices.Num() > 0 && MeshSection.Indices.Num() > 0)
+		{
+			FWaterBodySectionedMeshProxy& NewSection = Sections.Emplace_GetRef(MeshSection.SectionBounds);
+			
+			const uint32 BaseIndex = Vertices.Num();
+			TArray<uint32, TInlineAllocator<6>> Indices;
+			Algo::Transform(MeshSection.Indices, Indices, [BaseIndex](const uint32 Index) { return Index + BaseIndex; });
 
-		NewSection.VertexBuffers.InitFromDynamicVertex(&NewSection.VertexFactory, Component->WaterBodyMeshVertices);
+			NewSection.Indices = MoveTemp(Indices);
 
-		InitResources(NewSection);
+			Vertices.Append(MeshSection.Vertices);
+		}
 	}
 
-	if (Component->DilatedWaterBodyMeshVertices.Num() > 0 && Component->DilatedWaterBodyMeshIndices.Num() > 0)
+	if (Vertices.Num() > 0)
 	{
-		FWaterBodyMeshSection& NewSection = DilatedSections.Emplace_GetRef(GetScene().GetFeatureLevel());
-		NewSection.IndexBuffer.Indices = Component->DilatedWaterBodyMeshIndices;
-
-		NewSection.VertexBuffers.InitFromDynamicVertex(&NewSection.VertexFactory, Component->DilatedWaterBodyMeshVertices);
+		VertexBuffers.InitFromDynamicVertex(&VertexFactory, Vertices);
+		BeginInitResource(&VertexBuffers.PositionVertexBuffer);
+		BeginInitResource(&VertexBuffers.StaticMeshVertexBuffer);
+		BeginInitResource(&VertexBuffers.ColorVertexBuffer);
+		BeginInitResource(&VertexFactory);
+	}
+}
 
-		InitResources(NewSection);
+void FWaterBodySceneProxy::FWaterBodyMesh::Init(TArray<FDynamicMeshVertex>& Vertices, TArray<uint32>& Indices)
+{
+	if (Indices.Num() > 0)
+	{
+		IndexBuffer.Indices = Indices;
+		VertexBuffers.InitFromDynamicVertex(&VertexFactory, Vertices);
+		BeginInitResource(&VertexBuffers.PositionVertexBuffer);
+		BeginInitResource(&VertexBuffers.StaticMeshVertexBuffer);
+		BeginInitResource(&VertexBuffers.ColorVertexBuffer);
+		BeginInitResource(&VertexFactory);
+		BeginInitResource(&IndexBuffer);
 	}
+}
+
+void FWaterBodySceneProxy::FWaterBodySectionedLODMesh::ReleaseResources()
+{
+	Sections.Empty();
+	VertexBuffers.PositionVertexBuffer.ReleaseResource();
+	VertexBuffers.StaticMeshVertexBuffer.ReleaseResource();
+	VertexBuffers.ColorVertexBuffer.ReleaseResource();
+	VertexFactory.ReleaseResource();
+	IndexBuffer.ReleaseResource();
+}
 
-	if (UMaterialInstance* WaterInfoMaterial = Component->GetWaterInfoMaterialInstance())
+void FWaterBodySceneProxy::FWaterBodyMesh::ReleaseResources()
+{
+	VertexBuffers.PositionVertexBuffer.ReleaseResource();
+	VertexBuffers.StaticMeshVertexBuffer.ReleaseResource();
+	VertexBuffers.ColorVertexBuffer.ReleaseResource();
+	VertexFactory.ReleaseResource();
+	IndexBuffer.ReleaseResource();
+}
+
+bool FWaterBodySceneProxy::FWaterBodyMesh::GetMeshElements(FMeshBatch& OutMeshBatch, uint8 DepthPriorityGroup, bool bUseReverseCulling) const
+{
+	uint32 FirstIndex = 0;
+	uint32 NumPrimitives = IndexBuffer.Indices.Num() / 3;
+	if (NumPrimitives == 0)
 	{
-		Material = WaterInfoMaterial->GetRenderProxy();
+		return false;
 	}
+
+	OutMeshBatch.VertexFactory = &VertexFactory;
+	OutMeshBatch.ReverseCulling = bUseReverseCulling;
+	OutMeshBatch.Type = PT_TriangleList;
+	OutMeshBatch.DepthPriorityGroup = DepthPriorityGroup;
+	OutMeshBatch.bCanApplyViewModeOverrides = false;
+
+	FMeshBatchElement& BatchElement = OutMeshBatch.Elements[0];
+	BatchElement.IndexBuffer = &IndexBuffer;
+	BatchElement.FirstIndex = FirstIndex;
+	BatchElement.NumPrimitives = NumPrimitives;
+	check(BatchElement.NumPrimitives != 0);
+	BatchElement.MinVertexIndex = 0;
+	BatchElement.MaxVertexIndex = VertexBuffers.PositionVertexBuffer.GetNumVertices() - 1;
+
+	return true;
 }
 
-FWaterBodySceneProxy::~FWaterBodySceneProxy()
+uint32 FWaterBodySceneProxy::FWaterBodySectionedLODMesh::GetAllocatedSize() const
+{
+	return Sections.GetAllocatedSize();
+}
+
+bool FWaterBodySceneProxy::FWaterBodySectionedLODMesh::GetMeshElements(FMeshBatch& OutMeshBatch, uint8 DepthPriorityGroup, bool bUseReverseCulling) const
 {
-	for (FWaterBodyMeshSection& Section : Sections)
+	if (!bInitialized)
 	{
-		Section.VertexBuffers.PositionVertexBuffer.ReleaseResource();
-		Section.VertexBuffers.StaticMeshVertexBuffer.ReleaseResource();
-		Section.VertexBuffers.ColorVertexBuffer.ReleaseResource();
-		Section.IndexBuffer.ReleaseResource();
-		Section.VertexFactory.ReleaseResource();
+		return false;
 	}
+	const int32 NumPrimitives = IndexBuffer.Indices.Num() / 3;
 
-	for (FWaterBodyMeshSection& Section : DilatedSections)
+	if (NumPrimitives == 0)
 	{
-		Section.VertexBuffers.PositionVertexBuffer.ReleaseResource();
-		Section.VertexBuffers.StaticMeshVertexBuffer.ReleaseResource();
-		Section.VertexBuffers.ColorVertexBuffer.ReleaseResource();
-		Section.IndexBuffer.ReleaseResource();
-		Section.VertexFactory.ReleaseResource();
+		return false;
 	}
+
+	OutMeshBatch.VertexFactory = &VertexFactory;
+	OutMeshBatch.ReverseCulling = bUseReverseCulling;
+	OutMeshBatch.Type = PT_TriangleList;
+	OutMeshBatch.DepthPriorityGroup = DepthPriorityGroup;
+	OutMeshBatch.bCanApplyViewModeOverrides = false;
+	OutMeshBatch.LODIndex = 0;
+	OutMeshBatch.bUseForDepthPass = true;
+
+	FMeshBatchElement& BatchElement = OutMeshBatch.Elements[0];
+	BatchElement.IndexBuffer = &IndexBuffer;
+	BatchElement.FirstIndex = 0;
+	BatchElement.NumPrimitives = NumPrimitives;
+	check(BatchElement.NumPrimitives != 0);
+	BatchElement.MinVertexIndex = 0;
+	BatchElement.MaxVertexIndex = VertexBuffers.PositionVertexBuffer.GetNumVertices() - 1;
+
+	return true;
+}
+
+FWaterBodySceneProxy::~FWaterBodySceneProxy()
+{
+	WaterBodyInfoMesh.ReleaseResources();
+	WaterBodyInfoDilatedMesh.ReleaseResources();
+	WaterBodySectionedLODMesh.ReleaseResources();
 }
 
 void FWaterBodySceneProxy::GetDynamicMeshElements(const TArray<const FSceneView*>& Views, const FSceneViewFamily& ViewFamily, uint32 VisibilityMap, FMeshElementCollector& Collector) const
 {
-	FMaterialRenderProxy* MaterialToUse = Material;
+	TRACE_CPUPROFILER_EVENT_SCOPE("FWaterBodySceneProxy::GetDynamicMeshElements");
+
+	const bool bWithinWaterInfoPasses = CurrentWaterInfoPass != EWaterInfoPass::None;
+
+	FMaterialRenderProxy* MaterialToUse = bWithinWaterInfoPasses ? WaterInfoMaterial : WaterLODMaterial;
+
 	if (!MaterialToUse)
 	{
 		FColoredMaterialRenderProxy* FallbackMaterial = new FColoredMaterialRenderProxy(
@@ -102,13 +216,8 @@ void FWaterBodySceneProxy::GetDynamicMeshElements(const TArray<const FSceneView*
 		MaterialToUse = FallbackMaterial;
 	}
 
-	const bool bWithinWaterInfoPasses = CurrentWaterInfoPass != EWaterInfoPass::None;
-
-	const int32 ShowProxiesCVar = CVarWaterShowProxies.GetValueOnRenderThread();
-	const bool bDebugShowWaterSceneProxiesEnabled = !bWithinWaterInfoPasses && ((ShowProxiesCVar == 1 && IsSelected()) || ShowProxiesCVar >= 2);
-
 	// If we are not in the waterinfo pass and the cvar is not set to show opaque bodies, we should be in wireframe
-	const bool bWireframe = AllowDebugViewmodes() && !(bWithinWaterInfoPasses) && (CVarWaterShowProxies.GetValueOnRenderThread() != 3);
+	const bool bWireframe = AllowDebugViewmodes() && ViewFamily.EngineShowFlags.Wireframe;
 
 	if (bWireframe)
 	{
@@ -121,53 +230,41 @@ void FWaterBodySceneProxy::GetDynamicMeshElements(const TArray<const FSceneView*
 		MaterialToUse = WireframeMaterialInstance;
 	}
 
-	auto AddWaterBodyMeshSection = [&](const FWaterBodyMeshSection& Section, int32 ViewIndex)
-	{
-		FMeshBatch& Mesh = Collector.AllocateMesh();
-		FMeshBatchElement& BatchElement = Mesh.Elements[0];
-		BatchElement.IndexBuffer = &Section.IndexBuffer;
-		Mesh.bWireframe = bWireframe;
-		Mesh.VertexFactory = &Section.VertexFactory;
-		Mesh.MaterialRenderProxy = MaterialToUse;
-
-		bool bHasPrecomputedVolumetricLightmap;
-		FMatrix PreviousLocalToWorld;
-		int32 SingleCaptureIndex;
-		bool bOutputVelocity;
-		GetScene().GetPrimitiveUniformShaderParameters_RenderThread(GetPrimitiveSceneInfo(), bHasPrecomputedVolumetricLightmap, PreviousLocalToWorld, SingleCaptureIndex, bOutputVelocity);
-		bOutputVelocity |= AlwaysHasVelocity();
-
-		FDynamicPrimitiveUniformBuffer& DynamicPrimitiveUniformBuffer = Collector.AllocateOneFrameResource<FDynamicPrimitiveUniformBuffer>();
-		DynamicPrimitiveUniformBuffer.Set(GetLocalToWorld(), PreviousLocalToWorld, GetBounds(), GetLocalBounds(), true, bHasPrecomputedVolumetricLightmap, bOutputVelocity);
-		BatchElement.PrimitiveUniformBufferResource = &DynamicPrimitiveUniformBuffer.UniformBuffer;
-
-		BatchElement.FirstIndex = 0;
-		BatchElement.NumPrimitives = Section.IndexBuffer.Indices.Num() / 3;
-		check(BatchElement.NumPrimitives != 0);
-		BatchElement.MinVertexIndex = 0;
-		BatchElement.MaxVertexIndex = Section.VertexBuffers.PositionVertexBuffer.GetNumVertices() - 1;
-		Mesh.ReverseCulling = IsLocalToWorldDeterminantNegative();
-		Mesh.Type = PT_TriangleList;
-		Mesh.DepthPriorityGroup = bWireframe ? SDPG_Foreground : SDPG_World;
-		Mesh.bCanApplyViewModeOverrides = false;
-		Collector.AddMesh(ViewIndex, Mesh);
-	};
-
+	const bool bHasDilatedSections = WaterBodyInfoDilatedMesh.IndexBuffer.Indices.Num() > 0;
 	for (int32 ViewIndex = 0; ViewIndex < Views.Num(); ++ViewIndex)
 	{
-		if ((CurrentWaterInfoPass == EWaterInfoPass::Color) || (bDebugShowWaterSceneProxiesEnabled))
+		TRACE_CPUPROFILER_EVENT_SCOPE("FWaterBodySceneProxy::AddSectionsForView");
+
+		if (CurrentWaterInfoPass == EWaterInfoPass::Color)
 		{
-			for (const FWaterBodyMeshSection& Section : Sections)
+			FMeshBatch& MeshBatch = Collector.AllocateMesh();
+			MeshBatch.MaterialRenderProxy = MaterialToUse;
+			MeshBatch.bWireframe = bWireframe;
+			MeshBatch.bUseForDepthPass = true;
+			if (WaterBodyInfoMesh.GetMeshElements(MeshBatch, SDPG_World, IsLocalToWorldDeterminantNegative()))
 			{
-				AddWaterBodyMeshSection(Section, ViewIndex);
+				Collector.AddMesh(ViewIndex, MeshBatch);
 			}
 		}
-
-		if ((CurrentWaterInfoPass == EWaterInfoPass::Dilation) || (bDebugShowWaterSceneProxiesEnabled))
+		else if ((CurrentWaterInfoPass == EWaterInfoPass::Dilation) && bHasDilatedSections)
+		{
+			FMeshBatch& MeshBatch = Collector.AllocateMesh();
+			MeshBatch.MaterialRenderProxy = MaterialToUse;
+			MeshBatch.bWireframe = bWireframe;
+			MeshBatch.bUseForDepthPass = true;
+			if (WaterBodyInfoDilatedMesh.GetMeshElements(MeshBatch, SDPG_World, IsLocalToWorldDeterminantNegative()))
+			{
+				Collector.AddMesh(ViewIndex, MeshBatch);
+			}
+		}
+		else if (CurrentWaterInfoPass == EWaterInfoPass::None)
 		{
-			for (const FWaterBodyMeshSection& Section : DilatedSections)
+			FMeshBatch& MeshBatch = Collector.AllocateMesh();
+			MeshBatch.MaterialRenderProxy = MaterialToUse;
+			MeshBatch.bWireframe = bWireframe;
+			if (WaterBodySectionedLODMesh.GetMeshElements(MeshBatch, SDPG_World, IsLocalToWorldDeterminantNegative()))
 			{
-				AddWaterBodyMeshSection(Section, ViewIndex);
+				Collector.AddMesh(ViewIndex, MeshBatch);
 			}
 		}
 	}
@@ -179,11 +276,24 @@ FPrimitiveViewRelevance FWaterBodySceneProxy::GetViewRelevance(const FSceneView*
 	Result.bDrawRelevance = IsShown(View);
 	Result.bShadowRelevance = IsShadowCast(View);
 	Result.bDynamicRelevance = true;
+	Result.bStaticRelevance = false;
+	Result.bRenderInDepthPass = true;
 	Result.bRenderInMainPass = ShouldRenderInMainPass();
 	Result.bUsesLightingChannels = GetLightingChannelMask() != GetDefaultLightingChannelMask();
 	Result.bRenderCustomDepth = ShouldRenderCustomDepth();
 	Result.bTranslucentSelfShadow = bCastVolumetricTranslucentShadow;
-	Result.bVelocityRelevance = IsMovable();
+
+	const bool bWithinWaterInfoPasses = CurrentWaterInfoPass != EWaterInfoPass::None;
+	if (bWithinWaterInfoPasses)
+	{
+		WaterInfoMaterialRelevance.SetPrimitiveViewRelevance(Result);
+	}
+	else
+	{
+		WaterLODMaterialRelevance.SetPrimitiveViewRelevance(Result);
+	}
+
+	Result.bVelocityRelevance = DrawsVelocity() && Result.bOpaque && Result.bRenderInMainPass;
 	return Result;
 }
 
@@ -200,22 +310,30 @@ uint32 FWaterBodySceneProxy::GetMemoryFootprint() const
 
 uint32 FWaterBodySceneProxy::GetAllocatedSize() const
 {
-
-	return FPrimitiveSceneProxy::GetAllocatedSize() + Sections.GetAllocatedSize() + DilatedSections.GetAllocatedSize();
+	return FPrimitiveSceneProxy::GetAllocatedSize()
+		+ WaterBodySectionedLODMesh.GetAllocatedSize();
 }
 
 bool FWaterBodySceneProxy::IsShown(const FSceneView* View) const
 {
-	if (CurrentWaterInfoPass == EWaterInfoPass::None)
-	{
-		const int32 ShowProxiesCVar = CVarWaterShowProxies.GetValueOnRenderThread();
-		if ((ShowProxiesCVar == 1 && IsSelected()) || ShowProxiesCVar >= 2)
+	return FPrimitiveSceneProxy::IsShown(View);
+}
+
+void FWaterBodySceneProxy::OnTessellatedWaterMeshBoundsChanged_GameThread(const FBox2D& TessellatedWaterMeshBounds)
+{
+	check(IsInParallelGameThread() || IsInGameThread());
+
+	FWaterBodySceneProxy* SceneProxy = this;
+	ENQUEUE_RENDER_COMMAND(SetNewActiveWaterInfoGridArea)(
+		[SceneProxy, TessellatedWaterMeshBounds](FRHICommandListImmediate& RHICmdList)
 		{
-			return true;
-		}
+			SceneProxy->OnTessellatedWaterMeshBoundsChanged_RenderThread(TessellatedWaterMeshBounds);
+		});
+}
 
-		return false;
-	}
+void FWaterBodySceneProxy::OnTessellatedWaterMeshBoundsChanged_RenderThread(const FBox2D& TessellatedWaterMeshBounds)
+{
+	check(IsInRenderingThread());
 
-	return FPrimitiveSceneProxy::IsShown(View);
+	WaterBodySectionedLODMesh.RebuildIndexBuffer(TessellatedWaterMeshBounds);
 }
diff --git a/Engine/Plugins/Experimental/Water/Source/Runtime/Private/WaterInfoRendering.cpp b/Engine/Plugins/Experimental/Water/Source/Runtime/Private/WaterInfoRendering.cpp
--- a/Engine/Plugins/Experimental/Water/Source/Runtime/Private/WaterInfoRendering.cpp
+++ b/Engine/Plugins/Experimental/Water/Source/Runtime/Private/WaterInfoRendering.cpp
@@ -37,7 +37,7 @@ struct FUpdateWaterInfoParams
 	FRenderTarget* RenderTarget = nullptr;
 	FTexture* OutputTexture = nullptr;
 
-	FVector2D WaterZoneExtents;
+	FVector WaterZoneExtents;
 	FVector2f WaterHeightExtents;
 	float GroundZMin;
 	float CaptureZ;
@@ -413,7 +413,7 @@ static void UpdateWaterInfoRendering_RenderThread(
 
 		SCOPED_DRAW_EVENT(RHICmdList, DepthRendering_RT);
 
-		FRDGBuilder GraphBuilder(RHICmdList, RDG_EVENT_NAME("WaterInfoColorRendering"), ERDGBuilderFlags::AllowParallelExecute);
+		FRDGBuilder GraphBuilder(RHICmdList, RDG_EVENT_NAME("WaterInfoDepthRendering"), ERDGBuilderFlags::AllowParallelExecute);
 
 		// We need to execute the pre-render view extensions before we do any view dependent work.
 		FSceneRenderer::ViewExtensionPreRender_RenderThread(GraphBuilder, DepthRenderer);
@@ -524,7 +524,7 @@ static void UpdateWaterInfoRendering_RenderThread(
 
 		SCOPED_DRAW_EVENT(RHICmdList, DilationRendering_RT);
 
-		FRDGBuilder GraphBuilder(RHICmdList, RDG_EVENT_NAME("WaterInfoColorRendering"), ERDGBuilderFlags::AllowParallelExecute);
+		FRDGBuilder GraphBuilder(RHICmdList, RDG_EVENT_NAME("WaterInfoDilationRendering"), ERDGBuilderFlags::AllowParallelExecute);
 
 
 		// We need to execute the pre-render view extensions before we do any view dependent work.
@@ -690,9 +690,9 @@ void UpdateWaterInfoRendering(
 	{
 		return;
 	}
+	const FVector ZoneExtent = Context.ZoneToRender->GetTessellatedWaterMeshExtent();
 
-	const FVector2D ZoneExtent = Context.ZoneToRender->GetZoneExtent();
-	FVector ViewLocation = Context.ZoneToRender->GetActorLocation();
+	FVector ViewLocation = Context.ZoneToRender->GetTessellatedWaterMeshCenter();
 	ViewLocation.Z = Context.CaptureZ;
 
 	// Zone rendering always happens facing towards negative z.
@@ -762,7 +762,7 @@ void UpdateWaterInfoRendering(
 	Params.WaterHeightExtents = Context.ZoneToRender->GetWaterHeightExtents();
 	Params.GroundZMin = Context.ZoneToRender->GetGroundZMin();
 	Params.VelocityBlurRadius = Context.ZoneToRender->GetVelocityBlurRadius();
-	Params.WaterZoneExtents = Context.ZoneToRender->GetZoneExtent();
+	Params.WaterZoneExtents = ZoneExtent;
 
 	ENQUEUE_RENDER_COMMAND(WaterInfoCommand)(
 	[Params, ZoneName = Context.ZoneToRender->GetActorNameOrLabel()](FRHICommandListImmediate& RHICmdList)
diff --git a/Engine/Plugins/Experimental/Water/Source/Runtime/Private/WaterMeshComponent.cpp b/Engine/Plugins/Experimental/Water/Source/Runtime/Private/WaterMeshComponent.cpp
--- a/Engine/Plugins/Experimental/Water/Source/Runtime/Private/WaterMeshComponent.cpp
+++ b/Engine/Plugins/Experimental/Water/Source/Runtime/Private/WaterMeshComponent.cpp
@@ -150,6 +150,14 @@ FMaterialRelevance UWaterMeshComponent::GetWaterMaterialRelevance(ERHIFeatureLev
 	return Result;
 }
 
+void UWaterMeshComponent::PushTessellatedWaterMeshBoundsToPoxy(const FBox2D& TessellatedWaterMeshBounds)
+{
+	if (SceneProxy)
+	{
+		static_cast<FWaterMeshSceneProxy*>(SceneProxy)->OnTessellatedWaterMeshBoundsChanged_GameThread(TessellatedWaterMeshBounds);
+	}
+}
+
 void UWaterMeshComponent::SetExtentInTiles(FIntPoint NewExtentInTiles)
 {
 	ExtentInTiles = NewExtentInTiles;
@@ -573,6 +581,18 @@ void UWaterMeshComponent::PostEditChangeProperty(FPropertyChangedEvent& Property
 	{
 		const FName PropertyName = PropertyThatChanged->GetFName();
 
+		// Properties that require water body geometry to be rebuilt since they depend on the mesh grid layout
+		if (PropertyName == GET_MEMBER_NAME_CHECKED(UWaterMeshComponent, TileSize)
+			|| PropertyName == GET_MEMBER_NAME_CHECKED(UWaterMeshComponent, ExtentInTiles)
+			|| PropertyName == TEXT("RelativeScale3D")
+			|| PropertyName == TEXT("RelativeLocation")
+			|| PropertyName == TEXT("RelativeRotation"))
+		{
+			AWaterZone* WaterZone = GetOwner<AWaterZone>();
+			check(WaterZone);
+			WaterZone->MarkForRebuild(EWaterZoneRebuildFlags::UpdateWaterBodyLODSections);
+		}
+
 		// Properties that needs the scene proxy to be rebuilt
 		if (PropertyName == GET_MEMBER_NAME_CHECKED(UWaterMeshComponent, LODScale)
 			|| PropertyName == GET_MEMBER_NAME_CHECKED(UWaterMeshComponent, TessellationFactor)
@@ -587,4 +607,14 @@ void UWaterMeshComponent::PostEditChangeProperty(FPropertyChangedEvent& Property
 		}
 	}
 }
+
+void UWaterMeshComponent::PostEditComponentMove(bool bFinished)
+{
+	if (bFinished)
+	{
+		AWaterZone* WaterZone = GetOwner<AWaterZone>();
+		check(WaterZone);
+		WaterZone->MarkForRebuild(EWaterZoneRebuildFlags::UpdateWaterBodyLODSections);
+	}
+}
 #endif
diff --git a/Engine/Plugins/Experimental/Water/Source/Runtime/Private/WaterMeshSceneProxy.cpp b/Engine/Plugins/Experimental/Water/Source/Runtime/Private/WaterMeshSceneProxy.cpp
--- a/Engine/Plugins/Experimental/Water/Source/Runtime/Private/WaterMeshSceneProxy.cpp
+++ b/Engine/Plugins/Experimental/Water/Source/Runtime/Private/WaterMeshSceneProxy.cpp
@@ -266,6 +266,7 @@ void FWaterMeshSceneProxy::GetDynamicMeshElements(const TArray<const FSceneView*
 			TraversalDesc.PreViewTranslation = View->ViewMatrices.GetPreViewTranslation();
 			TraversalDesc.LODScale = LODScale;
 			TraversalDesc.bLODMorphingEnabled = !!CVarWaterMeshLODMorphEnabled.GetValueOnRenderThread();
+			TraversalDesc.TessellatedWaterMeshBounds = TessellatedWaterMeshBounds;
 
 #if !(UE_BUILD_SHIPPING || UE_BUILD_TEST)
 			//Debug
@@ -505,6 +506,7 @@ void FWaterMeshSceneProxy::GetDynamicRayTracingInstances(FRayTracingMaterialGath
 	TraversalDesc.Frustum = FConvexVolume(); // Default volume to disable frustum culling
 	TraversalDesc.LODScale = LODScale;
 	TraversalDesc.bLODMorphingEnabled = !!CVarWaterMeshLODMorphEnabled.GetValueOnRenderThread();
+	TraversalDesc.TessellatedWaterMeshBounds = TessellatedWaterMeshBounds;
 
 	WaterQuadTree.BuildWaterTileInstanceData(TraversalDesc, WaterInstanceData);
 
@@ -658,4 +660,24 @@ HHitProxy* FWaterMeshSceneProxy::CreateHitProxies(UPrimitiveComponent* Component
 	// No default hit proxy.
 	return nullptr;
 }
-#endif // WITH_WATER_SELECTION_SUPPORT
\ No newline at end of file
+#endif // WITH_WATER_SELECTION_SUPPORT
+
+
+void FWaterMeshSceneProxy::OnTessellatedWaterMeshBoundsChanged_GameThread(const FBox2D& InTessellatedWaterMeshBounds)
+{
+	check(IsInParallelGameThread() || IsInGameThread());
+
+	FWaterMeshSceneProxy* SceneProxy = this;
+	ENQUEUE_RENDER_COMMAND(OnTessellatedWaterMeshBoundsChanged)(
+		[SceneProxy, InTessellatedWaterMeshBounds](FRHICommandListImmediate& RHICmdList)
+		{
+			SceneProxy->OnTessellatedWaterMeshBoundsChanged_RenderThread(InTessellatedWaterMeshBounds);
+		});
+}
+
+void FWaterMeshSceneProxy::OnTessellatedWaterMeshBoundsChanged_RenderThread(const FBox2D& InTessellatedWaterMeshBounds)
+{
+	check(IsInRenderingThread());
+
+	TessellatedWaterMeshBounds = InTessellatedWaterMeshBounds;
+}
\ No newline at end of file
diff --git a/Engine/Plugins/Experimental/Water/Source/Runtime/Private/WaterQuadTree.cpp b/Engine/Plugins/Experimental/Water/Source/Runtime/Private/WaterQuadTree.cpp
--- a/Engine/Plugins/Experimental/Water/Source/Runtime/Private/WaterQuadTree.cpp
+++ b/Engine/Plugins/Experimental/Water/Source/Runtime/Private/WaterQuadTree.cpp
@@ -263,6 +263,68 @@ void FWaterQuadTree::FNode::SelectLOD(const FNodeData& InNodeData, int32 InLODLe
 	}
 }
 
+void FWaterQuadTree::FNode::SelectLODWithinBounds(const FNodeData& InNodeData, int32 InLODLevel, const FTraversalDesc& InTraversalDesc, FTraversalOutput& Output) const
+{
+	// #todo_water [roey]: this function currently forces all nodes to render at their lowest lod size. This isn't _that_ bad considering most of the nodes are close
+	// enough to the camera to render at lowest lod level anyways but ideally we would leverage the same lod selection system as the non-bounds implementation.
+
+	const FWaterBodyRenderData& WaterBodyRenderData = InNodeData.WaterBodyRenderData[WaterBodyIndex];
+	const FVector CenterPosition = Bounds.GetCenter();
+	const FVector Extent = Bounds.GetExtent();
+
+	// Early out on frustum culling 
+	if (!InTraversalDesc.Frustum.IntersectBox(CenterPosition, Extent))
+	{
+		// Handled
+		return;
+	}
+
+	check(InTraversalDesc.TessellatedWaterMeshBounds.bIsValid);
+	if (InLODLevel == 0)
+	{
+		if ((InTraversalDesc.TessellatedWaterMeshBounds.IsInsideOrOn(FVector2D(Bounds.Min)) && InTraversalDesc.TessellatedWaterMeshBounds.IsInsideOrOn(FVector2D(Bounds.Max))) &&
+			CanRender(0, InTraversalDesc.ForceCollapseDensityLevel, WaterBodyRenderData))
+		{
+			AddNodeForRender(InNodeData, WaterBodyRenderData, 0, InLODLevel, InTraversalDesc, Output);
+		}
+	}
+	else
+	{
+		// If this node has a complete subtree it will not contain any actual children, they are implicit to save memory so we generate them here
+		if (HasCompleteSubtree && IsSubtreeSameWaterBody)
+		{
+			FNode ChildNode;
+			const FVector HalfBoundSize(Extent.X, Extent.Y, Extent.Z*2.0f);
+			const FVector HalfOffsets[] = { {0.0f, 0.0f, 0.0f}, {1.0f, 0.0f, 0.0f} , {0.0f, 1.0f, 0.0f} , {1.0f, 1.0f, 0.0f} };
+			for (int i = 0; i < 4; i++)
+			{
+				const FVector ChildMin = Bounds.Min + HalfBoundSize * HalfOffsets[i];
+				const FVector ChildMax = ChildMin + HalfBoundSize;
+				const FBox ChildBounds(ChildMin, ChildMax);
+
+				// Create a temporary node to traverse
+				ChildNode.HasCompleteSubtree = 1;
+				ChildNode.IsSubtreeSameWaterBody = 1;
+				ChildNode.TransitionWaterBodyIndex = TransitionWaterBodyIndex;
+				ChildNode.WaterBodyIndex = WaterBodyIndex;
+				ChildNode.Bounds = ChildBounds;
+
+				ChildNode.SelectLODWithinBounds(InNodeData, InLODLevel - 1, InTraversalDesc, Output);
+			}
+		}
+		else
+		{
+			for (int32 ChildIndex : Children)
+			{
+				if (ChildIndex > 0)
+				{
+					InNodeData.Nodes[ChildIndex].SelectLODWithinBounds(InNodeData, InLODLevel - 1, InTraversalDesc, Output);
+				}
+			}
+		}
+	}
+}
+
 void FWaterQuadTree::FNode::AddNodes(FNodeData& InNodeData, const FBox& InMeshBounds, const FBox& InWaterBodyBounds, uint32 InWaterBodyIndex, int32 InLODLevel, uint32 InParentIndex)
 {
 	const FWaterBodyRenderData& InWaterBody = InNodeData.WaterBodyRenderData[InWaterBodyIndex];
@@ -645,7 +707,14 @@ void FWaterQuadTree::BuildWaterTileInstanceData(const FTraversalDesc& InTraversa
 {
 	TRACE_CPUPROFILER_EVENT_SCOPE(BuildWaterTileInstanceData);
 	check(bIsReadOnly);
-	NodeData.Nodes[0].SelectLOD(NodeData, TreeDepth, InTraversalDesc, Output);
+	if (InTraversalDesc.TessellatedWaterMeshBounds.bIsValid)
+	{
+		NodeData.Nodes[0].SelectLODWithinBounds(NodeData, TreeDepth, InTraversalDesc, Output);
+	}
+	else
+	{
+		NodeData.Nodes[0].SelectLOD(NodeData, TreeDepth, InTraversalDesc, Output);
+	}
 
 	// Append Far Mesh tiles
 	if (FarMeshData.InstanceData.Num() > 0 && FarMeshData.MaterialIndex != INDEX_NONE)
diff --git a/Engine/Plugins/Experimental/Water/Source/Runtime/Private/WaterSubsystem.cpp b/Engine/Plugins/Experimental/Water/Source/Runtime/Private/WaterSubsystem.cpp
--- a/Engine/Plugins/Experimental/Water/Source/Runtime/Private/WaterSubsystem.cpp
+++ b/Engine/Plugins/Experimental/Water/Source/Runtime/Private/WaterSubsystem.cpp
@@ -180,6 +180,15 @@ FWaterBodyManager* UWaterSubsystem::GetWaterBodyManager(const UWorld* InWorld)
 	return nullptr;
 }
 
+TWeakPtr<FWaterViewExtension, ESPMode::ThreadSafe> UWaterSubsystem::GetWaterViewExtension(const UWorld* InWorld)
+{
+	if (UWaterSubsystem* Subsystem = GetWaterSubsystem(InWorld))
+	{
+		return Subsystem->WaterViewExtension;
+	}
+	return {};
+}
+
 void UWaterSubsystem::Tick(float DeltaTime)
 {
 	Super::Tick(DeltaTime);
@@ -507,11 +516,6 @@ void UWaterSubsystem::MarkAllWaterZonesForRebuild(EWaterZoneRebuildFlags Rebuild
 	}
 }
 
-void UWaterSubsystem::MarkWaterInfoTextureForRebuild(const UE::WaterInfo::FRenderingContext& WaterInfoContext)
-{
-	WaterViewExtension->MarkWaterInfoTextureForRebuild(WaterInfoContext);
-}
-
 void UWaterSubsystem::NotifyWaterScalabilityChangedInternal(IConsoleVariable* CVar)
 {
 	OnWaterScalabilityChanged.Broadcast();
diff --git a/Engine/Plugins/Experimental/Water/Source/Runtime/Private/WaterViewExtension.cpp b/Engine/Plugins/Experimental/Water/Source/Runtime/Private/WaterViewExtension.cpp
--- a/Engine/Plugins/Experimental/Water/Source/Runtime/Private/WaterViewExtension.cpp
+++ b/Engine/Plugins/Experimental/Water/Source/Runtime/Private/WaterViewExtension.cpp
@@ -6,6 +6,16 @@
 #include "WaterModule.h"
 #include "WaterZoneActor.h"
 #include "WaterInfoRendering.h"
+#include "EngineUtils.h"
+#include "WaterMeshComponent.h"
+
+static TAutoConsoleVariable<float> CVarFreezeWaterLODUpdates(
+	TEXT("r.Water.FreezeLODUpdates"),
+	0,
+	TEXT(""),
+	ECVF_Default);
+
+// ----------------------------------------------------------------------------------
 
 FWaterViewExtension::FWaterViewExtension(const FAutoRegister& AutoReg, UWorld* InWorld)
 	: FWorldSceneViewExtension(AutoReg, InWorld)
@@ -18,20 +28,76 @@ FWaterViewExtension::~FWaterViewExtension()
 
 void FWaterViewExtension::SetupViewFamily(FSceneViewFamily& InViewFamily)
 {
-	if (WaterInfoContextsToRender.Num() > 0)
+}
+
+
+void FWaterViewExtension::SetupView(FSceneViewFamily& InViewFamily, FSceneView& InView)
+{
+	if (CVarFreezeWaterLODUpdates.GetValueOnGameThread() != 0)
 	{
-		const TWeakObjectPtr<UWorld> WorldPtr = GetWorld();
-		check(WorldPtr.IsValid())
+		return;
+	}
+
+	const FVector ViewLocation = InView.ViewLocation;
+
+	const TWeakObjectPtr<UWorld> WorldPtr = GetWorld();
+	check(WorldPtr.IsValid())
+
+	// #todo_water: hacky way to prevent updating from within the water info passes. Alternative solution might be to set a flag from within the info render
+	// similarly to the SetWithinWaterInfoPass functions for water body components 
+	//This won't be necessary when waterinfo passes don't utilize scene renderers.
+	if (!InView.bIsSceneCapture)
+	{
+		for (const TPair<AWaterZone*, UE::WaterInfo::FRenderingContext>& Pair : WaterInfoContextsToRender)
+		{
+			AWaterZone* WaterZone = Pair.Key;
+
+			if (WaterZone->IsNonTessellatedLODMeshEnabled())
+			{
+				const UWaterMeshComponent* WaterMesh= WaterZone->GetWaterMeshComponent();
+				check(WaterMesh);
+
+				const float SectionSize = WaterZone->GetNonTessellatedLODSectionSize();
+
+				FVector TessellatedWaterMeshCenter = ViewLocation.GridSnap(SectionSize);
+
+				WaterZone->SetTessellatedWaterMeshCenter(TessellatedWaterMeshCenter);
+
+				const FVector TessellatedWaterMeshHalfExtent = WaterZone->GetTessellatedWaterMeshExtent() / 2.0;
+				const FBox2D TessellatedWaterMeshBounds(FVector2D(TessellatedWaterMeshCenter - TessellatedWaterMeshHalfExtent), FVector2D(TessellatedWaterMeshCenter + TessellatedWaterMeshHalfExtent));
+
+				WaterZone->GetWaterMeshComponent()->PushTessellatedWaterMeshBoundsToPoxy(TessellatedWaterMeshBounds);
+
+				WaterZone->ForEachWaterBodyComponent([TessellatedWaterMeshBounds](UWaterBodyComponent* WaterBodyComponent)
+				{
+					WaterBodyComponent->PushTessellatedWaterMeshBoundsToProxy(TessellatedWaterMeshBounds);
+
+					return true;
+				});
+
+				const FBox2D CurrentCellBounds(FVector2D(ViewLocation) - FVector2D(SectionSize / 2.0), FVector2D(ViewLocation) + FVector2D(SectionSize / 2.0));
+				UpdateBounds.Add(WaterZone, CurrentCellBounds);
+			}
+
+			const UE::WaterInfo::FRenderingContext& Context(Pair.Value);
+			UE::WaterInfo::UpdateWaterInfoRendering(WorldPtr.Get()->Scene, Context);
+		}
 
-		// Move out of the cache list to prevent re-entering this function when we set up the view extensions for the water info update itself
-		TMap<AWaterZone*, UE::WaterInfo::FRenderingContext> ContextsToRender = MoveTemp(WaterInfoContextsToRender);
 		WaterInfoContextsToRender.Empty();
 
-		for (const TPair<AWaterZone*, UE::WaterInfo::FRenderingContext>& Pair : ContextsToRender)
+		for (AWaterZone* WaterZone : TActorRange<AWaterZone>(WorldPtr.Get()))
 		{
-			UE::WaterInfo::UpdateWaterInfoRendering(WorldPtr.Get()->Scene, Pair.Value);
+			if (WaterZone->IsNonTessellatedLODMeshEnabled())
+			{
+				const FBox2D* ZoneUpdateBounds = UpdateBounds.Find(WaterZone);
+				if (ZoneUpdateBounds == nullptr || !ZoneUpdateBounds->IsInside(FVector2D(ViewLocation)))
+				{
+					WaterZone->MarkForRebuild(EWaterZoneRebuildFlags::UpdateWaterInfoTexture);
+				}
+			}
 		}
 	}
+
 }
 
 void FWaterViewExtension::PreRenderViewFamily_RenderThread(FRDGBuilder& GraphBuilder, FSceneViewFamily& InViewFamily)
diff --git a/Engine/Plugins/Experimental/Water/Source/Runtime/Private/WaterZoneActor.cpp b/Engine/Plugins/Experimental/Water/Source/Runtime/Private/WaterZoneActor.cpp
--- a/Engine/Plugins/Experimental/Water/Source/Runtime/Private/WaterZoneActor.cpp
+++ b/Engine/Plugins/Experimental/Water/Source/Runtime/Private/WaterZoneActor.cpp
@@ -13,6 +13,7 @@
 #include "UObject/FortniteMainBranchObjectVersion.h"
 #include "WaterBodyOceanComponent.h"
 #include "RenderCaptureInterface.h"
+#include "WaterViewExtension.h"
 
 #if	WITH_EDITOR
 #include "Algo/Transform.h"
@@ -35,7 +36,8 @@ AWaterZone::AWaterZone(const FObjectInitializer& Initializer)
 {
 	WaterMesh = CreateDefaultSubobject<UWaterMeshComponent>(TEXT("WaterMesh"));
 	SetRootComponent(WaterMesh);
-	ZoneExtent = FVector2D(51200.f, 51200.f);
+	ZoneExtent = FVector2D(51200., 51200.);
+	TessellatedWaterMeshExtent = FVector(35000., 35000., 10000.);
 	
 #if	WITH_EDITOR
 	// Setup bounds component
@@ -47,7 +49,7 @@ AWaterZone::AWaterZone(const FObjectInitializer& Initializer)
 		BoundsComponent->SetGenerateOverlapEvents(false);
 		BoundsComponent->SetupAttachment(WaterMesh);
 		// Bounds component extent is half-extent, ZoneExtent is full extent.
-		BoundsComponent->SetBoxExtent(FVector(ZoneExtent / 2.f, 8192.f));
+		BoundsComponent->SetBoxExtent(FVector(ZoneExtent / 2., 8192.));
 	}
 
 	if (GIsEditor && !IsTemplate())
@@ -96,7 +98,8 @@ void AWaterZone::PostLoad()
 	Super::PostLoad();
 
 #if WITH_EDITORONLY_DATA
-	const FVector2D ExtentInTiles = WaterMesh->GetExtentInTiles();
+	// WaterMesh ExtentInTiles returns the half extent.
+	const FVector2D ExtentInTiles = 2.0 * FVector2D(WaterMesh->GetExtentInTiles());
 	ZoneExtent = FVector2D(ExtentInTiles * WaterMesh->GetTileSize());
 	OnExtentChanged();
 #endif // WITH_EDITORONLY_DATA
@@ -112,18 +115,35 @@ void AWaterZone::MarkForRebuild(EWaterZoneRebuildFlags Flags)
 	{
 		bNeedsWaterInfoRebuild = true;
 	}
+	if (EnumHasAnyFlags(Flags, EWaterZoneRebuildFlags::UpdateWaterBodyLODSections))
+	{
+		ForEachWaterBodyComponent([](UWaterBodyComponent* WaterBodyComponent)
+		{
+			WaterBodyComponent->UpdateWaterBodyRenderData();
+			return true;
+		});
+	}
 }
 
 void AWaterZone::ForEachWaterBodyComponent(TFunctionRef<bool(UWaterBodyComponent*)> Predicate)
 {
-	FWaterBodyManager::ForEachWaterBodyComponent(GetWorld(), [this, Predicate](UWaterBodyComponent* Component)
+	for (const TWeakObjectPtr<UWaterBodyComponent>& WaterBodyComponent : OwnedWaterBodies)
+	{
+		if (WaterBodyComponent.IsValid() && !Predicate(WaterBodyComponent.Get()))
 		{
-			if (Component->GetWaterZone() == this)
-			{
-				return Predicate(Component);
-			}
-			return true;
-		});
+			break;
+		}
+	}
+}
+
+void AWaterZone::AddWaterBodyComponent(UWaterBodyComponent* WaterBodyComponent)
+{
+	OwnedWaterBodies.AddUnique(WaterBodyComponent);
+}
+
+void AWaterZone::RemoveWaterBodyComponent(UWaterBodyComponent* WaterBodyComponent)
+{
+	OwnedWaterBodies.RemoveSwap(WaterBodyComponent);
 }
 
 void AWaterZone::Update()
@@ -154,7 +174,12 @@ void AWaterZone::PostEditMove(bool bFinished)
 	Super::PostEditMove(bFinished);
 
 	// Ensure that the water mesh is rebuilt if it moves
-	MarkForRebuild(EWaterZoneRebuildFlags::All);
+	EWaterZoneRebuildFlags RebuildFlags = EWaterZoneRebuildFlags::UpdateWaterInfoTexture | EWaterZoneRebuildFlags::UpdateWaterMesh;
+	if (bFinished)
+	{
+		RebuildFlags |= EWaterZoneRebuildFlags::UpdateWaterBodyLODSections;
+	}
+	MarkForRebuild(RebuildFlags);
 }
 
 void AWaterZone::PostEditChangeProperty(FPropertyChangedEvent& PropertyChangedEvent)
@@ -182,6 +207,18 @@ void AWaterZone::PostEditChangeProperty(FPropertyChangedEvent& PropertyChangedEv
 	{
 		MarkForRebuild(EWaterZoneRebuildFlags::UpdateWaterInfoTexture);
 	}
+	else if (PropertyChangedEvent.GetPropertyName() == GET_MEMBER_NAME_CHECKED(AWaterZone, TessellatedWaterMeshExtent))
+	{
+		MarkForRebuild(EWaterZoneRebuildFlags::All);
+	}
+	else if (PropertyChangedEvent.GetPropertyName() == GET_MEMBER_NAME_CHECKED(AWaterZone, bEnableNonTessellatedLODMesh))
+	{
+		MarkForRebuild(EWaterZoneRebuildFlags::All);
+	}
+	else if (PropertyChangedEvent.GetPropertyName() == GET_MEMBER_NAME_CHECKED(AWaterZone, NonTessellatedLODSectionScale))
+	{
+		MarkForRebuild(EWaterZoneRebuildFlags::All);
+	}
 }
 
 void AWaterZone::OnActorSelectionChanged(const TArray<UObject*>& NewSelection, bool bForceRefresh)
@@ -197,7 +234,7 @@ void AWaterZone::OnActorSelectionChanged(const TArray<UObject*>& NewSelection, b
 	if (SelectedWaterBodies != NewWeakWaterBodiesSelection)
 	{
 		SelectedWaterBodies = NewWeakWaterBodiesSelection;
-		MarkForRebuild(EWaterZoneRebuildFlags::All);
+		MarkForRebuild(EWaterZoneRebuildFlags::UpdateWaterMesh | EWaterZoneRebuildFlags::UpdateWaterInfoTexture);
 	}
 }
 #endif // WITH_EDITOR
@@ -206,9 +243,10 @@ void AWaterZone::OnExtentChanged()
 {
 	// Compute the new tile extent based on the new bounds
 	const float MeshTileSize = WaterMesh->GetTileSize();
+	const FVector2D ZoneHalfExtent = ZoneExtent / 2.0;
 
-	int32 NewExtentInTilesX = FMath::FloorToInt(ZoneExtent.X / MeshTileSize);
-	int32 NewExtentInTilesY = FMath::FloorToInt(ZoneExtent.Y / MeshTileSize);
+	int32 NewExtentInTilesX = FMath::FloorToInt(ZoneHalfExtent.X / MeshTileSize);
+	int32 NewExtentInTilesY = FMath::FloorToInt(ZoneHalfExtent.Y / MeshTileSize);
 	
 	// We must ensure that the zone is always at least 1x1
 	NewExtentInTilesX = FMath::Max(1, NewExtentInTilesX);
@@ -217,8 +255,7 @@ void AWaterZone::OnExtentChanged()
 	WaterMesh->SetExtentInTiles(FIntPoint(NewExtentInTilesX, NewExtentInTilesY));
 
 #if WITH_EDITOR
-	// Bounds component extent is half-extent, ZoneExtent is full extent.
-	BoundsComponent->SetBoxExtent(FVector(ZoneExtent / 2.f, 8192.f));
+	BoundsComponent->SetBoxExtent(FVector(ZoneHalfExtent, 8192.f));
 #endif // WITH_EDITOR
 
 	MarkForRebuild(EWaterZoneRebuildFlags::All);
@@ -228,7 +265,8 @@ void AWaterZone::OnExtentChanged()
 #if WITH_EDITOR
 void AWaterZone::OnBoundsComponentModified()
 {
-	const FVector2D NewBounds = FVector2D(BoundsComponent->GetUnscaledBoxExtent());
+	// BoxExtent is the radius of the box (half-extent).
+	const FVector2D NewBounds = 2.0 * FVector2D(BoundsComponent->GetUnscaledBoxExtent());
 
 	SetZoneExtent(NewBounds);
 }
@@ -246,31 +284,29 @@ bool AWaterZone::UpdateWaterInfoTexture()
 		// If they do not, we must submit compile jobs for them and wait until they are finished before re-rendering.
 		TArray<UMaterialInterface*> UsedMaterials;
 
-		// #todo_water [roey]: we should try caching this list to avoid potentially iterating over a lot of water bodies which may not belong to this zone specifically.
-		// For now whenever we update the water info texture we will collect all water bodies within the zone and pass those to the renderer each time this function is called.
-		TArray<UWaterBodyComponent*> WaterBodies;
-		ForEachWaterBodyComponent([World, &WaterBodies, &WaterZMax, &WaterZMin, &UsedMaterials](UWaterBodyComponent* Component)
+		TArray<UWaterBodyComponent*> WaterBodiesToRender;
+		ForEachWaterBodyComponent([World, &WaterBodiesToRender, &WaterZMax, &WaterZMin, &UsedMaterials](UWaterBodyComponent* WaterBodyComponent)
 		{
 			// skip components which don't affect the water info texture
-			if (!Component->AffectsWaterInfo())
+			if (!WaterBodyComponent->AffectsWaterInfo())
 			{
 				return true;
 			}
 
-			if (UMaterialInterface* WaterInfoMaterial = Component->GetWaterInfoMaterialInstance())
+			if (UMaterialInterface* WaterInfoMaterial = WaterBodyComponent->GetWaterInfoMaterialInstance())
 			{
 				UsedMaterials.Add(WaterInfoMaterial);
 			}
 
-			WaterBodies.Add(Component);
-			const FBox WaterBodyBounds = Component->CalcBounds(Component->GetComponentToWorld()).GetBox();
+			WaterBodiesToRender.Add(WaterBodyComponent);
+			const FBox WaterBodyBounds = WaterBodyComponent->CalcBounds(WaterBodyComponent->GetComponentToWorld()).GetBox();
 			WaterZMax = FMath::Max(WaterZMax, WaterBodyBounds.Max.Z);
 			WaterZMin = FMath::Min(WaterZMin, WaterBodyBounds.Min.Z);
 			return true;
 		});
 
 		// If we don't have any water bodies we don't need to do anything.
-		if (WaterBodies.Num() == 0)
+		if (WaterBodiesToRender.Num() == 0)
 		{
 			return true;
 		}
@@ -334,25 +370,29 @@ bool AWaterZone::UpdateWaterInfoTexture()
 			return false;
 		}
 
-
 		const ETextureRenderTargetFormat Format = bHalfPrecisionTexture ? ETextureRenderTargetFormat::RTF_RGBA16f : RTF_RGBA32f;
-		WaterInfoTexture = FWaterUtils::GetOrCreateTransientRenderTarget2D(WaterInfoTexture, TEXT("WaterInfoTexture"), RenderTargetResolution, Format);
+		UTextureRenderTarget2D* NewWaterInfoTexture = FWaterUtils::GetOrCreateTransientRenderTarget2D(WaterInfoTexture, TEXT("WaterInfoTexture"), RenderTargetResolution, Format);
+		const bool bNeedsMIDUpdate = WaterInfoTexture != NewWaterInfoTexture;
+		WaterInfoTexture = NewWaterInfoTexture;
 
 		UE::WaterInfo::FRenderingContext Context;
 		Context.ZoneToRender = this;
-		Context.WaterBodies = WaterBodies;
+		Context.WaterBodies = WaterBodiesToRender;
 		Context.GroundActors = MoveTemp(GroundActors);
 		Context.CaptureZ = FMath::Max(WaterZMax, GroundZMax) + CaptureZOffset;
 		Context.TextureRenderTarget = WaterInfoTexture;
 
-		if (UWaterSubsystem* WaterSubsystem = UWaterSubsystem::GetWaterSubsystem(World))
+		if (TWeakPtr<FWaterViewExtension> WaterViewExtension = UWaterSubsystem::GetWaterViewExtension(World); WaterViewExtension.IsValid())
 		{
-			WaterSubsystem->MarkWaterInfoTextureForRebuild(Context);
+			WaterViewExtension.Pin()->MarkWaterInfoTextureForRebuild(Context);
 		}
 
-		for (UWaterBodyComponent* Component : WaterBodies)
+		if (bNeedsMIDUpdate)
 		{
-			Component->UpdateMaterialInstances();
+			for (UWaterBodyComponent* Component : WaterBodiesToRender)
+			{
+				Component->UpdateMaterialInstances();
+			}
 		}
 
 		UE_LOG(LogWater, Verbose, TEXT("Queued Water Info texture update"));
@@ -360,3 +400,34 @@ bool AWaterZone::UpdateWaterInfoTexture()
 
 	return true;
 }
+
+FVector AWaterZone::GetTessellatedWaterMeshCenter() const
+{
+	if (IsNonTessellatedLODMeshEnabled())
+	{
+		return TessellatedWaterMeshCenter;
+	}
+	else
+	{
+		return GetActorLocation();
+	}
+}
+
+FVector AWaterZone::GetTessellatedWaterMeshExtent() const
+{
+	if (IsNonTessellatedLODMeshEnabled())
+	{
+		return TessellatedWaterMeshExtent.GridSnap(GetNonTessellatedLODSectionSize());
+	}
+	else
+	{
+		// #todo_water [roey]: better implementation for 3D extent
+		return FVector(GetZoneExtent(), 0.0);
+	}
+}
+
+float AWaterZone::GetNonTessellatedLODSectionSize() const
+{
+	check(WaterMesh)
+	return WaterMesh->GetTileSize() * NonTessellatedLODSectionScale;
+}
diff --git a/Engine/Plugins/Experimental/Water/Source/Runtime/Public/WaterBodyComponent.h b/Engine/Plugins/Experimental/Water/Source/Runtime/Public/WaterBodyComponent.h
--- a/Engine/Plugins/Experimental/Water/Source/Runtime/Public/WaterBodyComponent.h
+++ b/Engine/Plugins/Experimental/Water/Source/Runtime/Public/WaterBodyComponent.h
@@ -31,6 +31,7 @@ class AWaterZone;
 class ALandscapeProxy;
 class UMaterialInstanceDynamic;
 class FTokenizedMessage;
+namespace UE::Geometry { class FDynamicMesh3; }
 
 // ----------------------------------------------------------------------------------
 
@@ -80,6 +81,29 @@ enum class EWaterBodyStatus : uint8
 	InvalidWaveData,
 };
 
+
+// ----------------------------------------------------------------------------------
+
+struct FWaterBodyMeshSection
+{
+public:
+	FWaterBodyMeshSection()
+		: Vertices()
+		, Indices()
+		, SectionBounds()
+	{
+	}
+
+	bool IsValid() const { return (bool)SectionBounds.bIsValid; }
+
+	uint32 GetAllocatedSize() const { return Vertices.GetAllocatedSize() + Indices.GetAllocatedSize(); }
+
+	TArray<FDynamicMeshVertex> Vertices;
+	TArray<uint32> Indices;
+
+	FBox2D SectionBounds;
+};
+
 // ----------------------------------------------------------------------------------
 
 UCLASS(Abstract, HideCategories = (Tags, Activation, Cooking, Replication, Input, AssetUserData, Mesh))
@@ -178,6 +202,9 @@ public:
 	UFUNCTION(BlueprintCallable, Category = Rendering)
 	UMaterialInterface* GetWaterMaterial() const { return WaterMaterial; }
 
+	/** Returns material used to render the water LOD mesh sections */
+	UMaterialInterface* GetWaterLODMaterial() const { return WaterLODMaterial; }
+
 	/** Sets water material */
 	void SetWaterMaterial(UMaterialInterface* InMaterial);
 
@@ -185,6 +212,10 @@ public:
 	UFUNCTION(BlueprintCallable, Category = Rendering)
 	UMaterialInstanceDynamic* GetWaterMaterialInstance();
 
+	/** Returns water LOD MID */
+	UFUNCTION(BlueprintCallable, Category = Rendering)
+	UMaterialInstanceDynamic* GetWaterLODMaterialInstance();
+
 	/** Returns under water post process MID */
 	UFUNCTION(BlueprintCallable, Category = Rendering)
 	UMaterialInstanceDynamic* GetUnderwaterPostProcessMaterialInstance();
@@ -330,12 +361,15 @@ public:
 	virtual void Reset() {}
 
 	/** Gets the water zone to which this component belongs */
-	AWaterZone* GetWaterZone() const;
+	AWaterZone* GetWaterZone() const { return OwningWaterZone; }
 
 	void UpdateWaterBodyRenderData();
+	void UpdateNonTessellatedMeshSections();
 
 	virtual FPrimitiveSceneProxy* CreateSceneProxy() override;
 	virtual void GetUsedMaterials(TArray<UMaterialInterface*>& OutMaterials, bool bGetDebugMaterials) const override;
+
+	void PushTessellatedWaterMeshBoundsToProxy(const FBox2D& TessellatedWaterMeshBounds);
 	
 	/** Set the navigation area class */
 	void SetNavAreaClass(TSubclassOf<UNavAreaBase> NewWaterNavAreaClass) { WaterNavAreaClass = NewWaterNavAreaClass; }
@@ -373,8 +407,8 @@ protected:
 
 	virtual void OnUpdateBody(bool bWithExclusionVolumes) {}
 
-	/* Generates the mesh used to render the Water Info texture */
-	virtual void GenerateWaterBodyMesh() {}
+	/* Generates the meshes used to render the Water Info texture */
+	virtual bool GenerateWaterBodyMesh(UE::Geometry::FDynamicMesh3& OutMesh, UE::Geometry::FDynamicMesh3* OutDilatedMesh = nullptr) const { checkf(!AffectsWaterInfo(), TEXT("WaterBodyComponent affects water info but does not implement GenerateWaterBodyMesh!")); return false; }
 
 	/** Returns navigation area class */
 	TSubclassOf<UNavAreaBase> GetNavAreaClass() const { return WaterNavAreaClass; }
@@ -411,8 +445,11 @@ protected:
 	EWaterBodyQueryFlags CheckAndAjustQueryFlags(EWaterBodyQueryFlags InQueryFlags) const;
 	void UpdateSplineComponent();
 	void UpdateExclusionVolumes();
+	void UpdateWaterZones();
+	AWaterZone* FindWaterZone() const;
 	bool UpdateWaterHeight();
 	virtual void CreateOrUpdateWaterMID();
+	void CreateOrUpdateWaterLODMID();
 	void CreateOrUpdateUnderwaterPostProcessMID();
 	void CreateOrUpdateWaterInfoMID();
 	void PrepareCurrentPostProcessSettings();
@@ -421,6 +458,9 @@ protected:
 	void RequestGPUWaveDataUpdate();
 	EObjectFlags GetTransientMIDFlags() const; 
 	void DeprecateData();
+	void RebuildWaterBodyInfoMesh();
+	void RebuildWaterBodyLODSections();
+	void OnTessellatedWaterMeshBoundsChanged();
 
 	virtual void Serialize(FArchive& Ar) override;
 	virtual void PostLoad() override;
@@ -480,7 +520,7 @@ public:
 
 	/** Post process settings to apply when the camera goes underwater (only available when bGenerateCollisions is true because collisions are needed to detect if it's under water).
 	Note: Underwater post process material is setup using UnderwaterPostProcessMaterial. */
-	UPROPERTY(Category = Rendering, EditAnywhere, BlueprintReadWrite, meta = (EditCondition = "bGenerateCollisions && UnderwaterPostProcessMaterial != nullptr", DisplayAfter = "UnderwaterPostProcessMaterial"))
+	UPROPERTY(Category = Rendering, EditAnywhere, BlueprintReadWrite, AdvancedDisplay, meta = (EditCondition = "bGenerateCollisions && UnderwaterPostProcessMaterial != nullptr", DisplayAfter = "UnderwaterPostProcessMaterial"))
 	FUnderwaterPostProcessSettings UnderwaterPostProcessSettings;
 
 	UPROPERTY(Category = Terrain, EditAnywhere, BlueprintReadWrite)
@@ -492,6 +532,9 @@ public:
 	UPROPERTY(Category = HLOD, EditAnywhere, BlueprintReadOnly, meta = (DisplayName = "Water HLOD Material"))
 	TObjectPtr<UMaterialInterface> WaterHLODMaterial;
 
+	UPROPERTY(Category = Rendering, EditAnywhere, BlueprintReadOnly, meta = (DisplayName = "Water LOD Material"))
+	UMaterialInterface* WaterLODMaterial;
+
 	/** Post process material to apply when the camera goes underwater (only available when bGenerateCollisions is true because collisions are needed to detect if it's under water). */
 	UPROPERTY(Category = Rendering, EditAnywhere, BlueprintReadOnly, meta = (EditCondition = "bGenerateCollisions", DisplayAfter = "WaterMaterial"))
 	TObjectPtr<UMaterialInterface> UnderwaterPostProcessMaterial;
@@ -544,6 +587,9 @@ protected:
 	UPROPERTY(Category = Debug, VisibleInstanceOnly, Transient, NonPIEDuplicateTransient, TextExportTransient, meta = (DisplayAfter = "WaterMaterial"))
 	TObjectPtr<UMaterialInstanceDynamic> WaterMID;
 
+	UPROPERTY(Category = Debug, VisibleInstanceOnly, Transient, NonPIEDuplicateTransient, TextExportTransient, meta = (DisplayAfter = "WaterLODMaterial"))
+	UMaterialInstanceDynamic* WaterLODMID;
+
 	UPROPERTY(Category = Debug, VisibleInstanceOnly, Transient, NonPIEDuplicateTransient, TextExportTransient, meta = (DisplayAfter = "UnderwaterPostProcessMaterial"))
 	TObjectPtr<UMaterialInstanceDynamic> UnderwaterPostProcessMID;
 	
@@ -560,6 +606,9 @@ protected:
 	UPROPERTY(Transient)
 	mutable TWeakObjectPtr<ALandscapeProxy> Landscape;
 
+	UPROPERTY()
+	AWaterZone* OwningWaterZone;
+
 	UPROPERTY(Transient)
 	FPostProcessSettings CurrentPostProcessSettings;
 
@@ -567,7 +616,9 @@ protected:
 	UPROPERTY(EditAnywhere, BlueprintReadWrite, Category = Navigation, meta = (EditCondition = "bCanAffectNavigation && bGenerateCollisions"))
 	TSubclassOf<UNavAreaBase> WaterNavAreaClass;
 
-	// #todo_water [roey]: these need to be refactored into individual section structs and serialized
+	// #todo_water [roey]: serialize and don't rebuild on load.
+	TArray<FWaterBodyMeshSection> WaterBodyMeshSections;
+
 	TArray<FDynamicMeshVertex> WaterBodyMeshVertices;
 	TArray<uint32> WaterBodyMeshIndices;
 
diff --git a/Engine/Plugins/Experimental/Water/Source/Runtime/Public/WaterBodyLakeComponent.h b/Engine/Plugins/Experimental/Water/Source/Runtime/Public/WaterBodyLakeComponent.h
--- a/Engine/Plugins/Experimental/Water/Source/Runtime/Public/WaterBodyLakeComponent.h
+++ b/Engine/Plugins/Experimental/Water/Source/Runtime/Public/WaterBodyLakeComponent.h
@@ -23,7 +23,7 @@ public:
 	virtual EWaterBodyType GetWaterBodyType() const override { return EWaterBodyType::Lake; }
 	virtual TArray<UPrimitiveComponent*> GetCollisionComponents(bool bInOnlyEnabledComponents = true) const override;
 	virtual TArray<UPrimitiveComponent*> GetStandardRenderableComponents() const override;
-	virtual void GenerateWaterBodyMesh() override;
+	virtual bool GenerateWaterBodyMesh(UE::Geometry::FDynamicMesh3& OutMesh, UE::Geometry::FDynamicMesh3* OutDilatedMesh = nullptr) const override;
 
 	virtual FBoxSphereBounds CalcBounds(const FTransform& LocalToWorld) const override;
 
diff --git a/Engine/Plugins/Experimental/Water/Source/Runtime/Public/WaterBodyOceanComponent.h b/Engine/Plugins/Experimental/Water/Source/Runtime/Public/WaterBodyOceanComponent.h
--- a/Engine/Plugins/Experimental/Water/Source/Runtime/Public/WaterBodyOceanComponent.h
+++ b/Engine/Plugins/Experimental/Water/Source/Runtime/Public/WaterBodyOceanComponent.h
@@ -33,7 +33,7 @@ protected:
 	virtual void BeginUpdateWaterBody() override;
 	virtual void OnUpdateBody(bool bWithExclusionVolumes) override;
 	virtual void Reset() override;
-	virtual void GenerateWaterBodyMesh() override;
+	virtual bool GenerateWaterBodyMesh(UE::Geometry::FDynamicMesh3& OutMesh, UE::Geometry::FDynamicMesh3* OutDilatedMesh = nullptr) const override;
 
 	virtual void PostLoad() override;
 	virtual FBoxSphereBounds CalcBounds(const FTransform& LocalToWorld) const;
diff --git a/Engine/Plugins/Experimental/Water/Source/Runtime/Public/WaterBodyRiverComponent.h b/Engine/Plugins/Experimental/Water/Source/Runtime/Public/WaterBodyRiverComponent.h
--- a/Engine/Plugins/Experimental/Water/Source/Runtime/Public/WaterBodyRiverComponent.h
+++ b/Engine/Plugins/Experimental/Water/Source/Runtime/Public/WaterBodyRiverComponent.h
@@ -36,7 +36,7 @@ protected:
 	virtual void Reset() override;
 	virtual void UpdateMaterialInstances() override;
 	virtual void OnUpdateBody(bool bWithExclusionVolumes) override;
-	virtual void GenerateWaterBodyMesh() override;
+	virtual bool GenerateWaterBodyMesh(UE::Geometry::FDynamicMesh3& OutMesh, UE::Geometry::FDynamicMesh3* OutDilatedMesh = nullptr) const override;
 
 	virtual FBoxSphereBounds CalcBounds(const FTransform& LocalToWorld) const override;
 #if WITH_EDITOR
diff --git a/Engine/Plugins/Experimental/Water/Source/Runtime/Public/WaterBodySceneProxy.h b/Engine/Plugins/Experimental/Water/Source/Runtime/Public/WaterBodySceneProxy.h
--- a/Engine/Plugins/Experimental/Water/Source/Runtime/Public/WaterBodySceneProxy.h
+++ b/Engine/Plugins/Experimental/Water/Source/Runtime/Public/WaterBodySceneProxy.h
@@ -12,8 +12,9 @@ class UWaterBodyComponent;
 class UWaterBodyLakeComponent;
 class UWaterBodyOceanComponent;
 class UWaterBodyRiverComponent;
-class FWaterBodyMeshSection;
 class FMaterialRenderProxy;
+struct FWaterBodyMeshSection;
+struct FWaterBodySectionedMeshProxy;
 
 enum class EWaterInfoPass
 {
@@ -26,7 +27,7 @@ enum class EWaterInfoPass
 class FWaterBodySceneProxy final : public FPrimitiveSceneProxy
 {
 public:
-	FWaterBodySceneProxy(UWaterBodyComponent* Component);
+	FWaterBodySceneProxy(UWaterBodyComponent* Component, const FBox2D& TessellatedMeshBounds);
 	virtual ~FWaterBodySceneProxy();
 	virtual void GetDynamicMeshElements(const TArray<const FSceneView*>& Views, const FSceneViewFamily& ViewFamily, uint32 VisibilityMap, FMeshElementCollector& Collector) const override;
 	virtual FPrimitiveViewRelevance GetViewRelevance(const FSceneView* View) const override;
@@ -39,12 +40,72 @@ public:
 	bool IsWithinWaterInfoPass(EWaterInfoPass InPass) const { return CurrentWaterInfoPass == InPass; }
 	void SetWithinWaterInfoPass(EWaterInfoPass InPass) { CurrentWaterInfoPass = InPass; }
 
+	void OnTessellatedWaterMeshBoundsChanged_GameThread(const FBox2D& TessellatedWaterMeshBounds);
 private:
-	void InitResources(FWaterBodyMeshSection& Section);
+	void OnTessellatedWaterMeshBoundsChanged_RenderThread(const FBox2D& TessellatedWaterMeshBounds);
 
-	TArray<FWaterBodyMeshSection> Sections;
-	TArray<FWaterBodyMeshSection> DilatedSections;
-	FMaterialRenderProxy* Material = nullptr;
+	struct FWaterBodySectionedLODMesh
+	{
+	public:
+		FWaterBodySectionedLODMesh(ERHIFeatureLevel::Type InFeatureLevel) : VertexFactory(InFeatureLevel, "WaterBodySectionedMesh"), Sections() {}
+
+		void InitFromSections(const TArray<FWaterBodyMeshSection>& MeshSections);
+		void ReleaseResources();
+		void RebuildIndexBuffer(const FBox2D& TessellatedWaterMeshBounds);
+		uint32 GetAllocatedSize() const;
+
+		bool GetMeshElements(FMeshBatch& OutMeshBatch, uint8 DepthPriorityGroup, bool bUseReverseCulling) const;
+
+		FStaticMeshVertexBuffers VertexBuffers;
+		FDynamicMeshIndexBuffer32 IndexBuffer;
+		FLocalVertexFactory VertexFactory;
+
+		struct FWaterBodySectionedMeshProxy
+		{
+		public:
+			FWaterBodySectionedMeshProxy(const FBox2D& InBounds)
+				: Bounds(InBounds)
+				, Indices()
+			{
+				Indices.Reserve(6); // Most sections are quads with 6 indices.
+			}
+
+			FBox2D Bounds;
+			TArray<uint32> Indices;
+		};
+
+		TArray<FWaterBodySectionedMeshProxy> Sections;
+
+		bool bInitialized = false;
+	};
+
+	struct FWaterBodyMesh
+	{
+	public:
+		FWaterBodyMesh(ERHIFeatureLevel::Type InFeatureLevel) : VertexFactory(InFeatureLevel, "WaterBodyMesh") {}
+
+		void Init(TArray<FDynamicMeshVertex>& Vertices, TArray<uint32>& Indices);
+		void ReleaseResources();
+
+		bool GetMeshElements(FMeshBatch& OutMeshBatch, uint8 DepthPriorityGroup, bool bUseReverseCulling) const;
+
+		FStaticMeshVertexBuffers VertexBuffers;
+		FDynamicMeshIndexBuffer32 IndexBuffer;
+		FLocalVertexFactory VertexFactory;
+	};
+
+	FWaterBodySectionedLODMesh WaterBodySectionedLODMesh;
+
+	FWaterBodyMesh WaterBodyInfoMesh;
+	FWaterBodyMesh WaterBodyInfoDilatedMesh;
+
+	/** Material to use when rendering this scene proxy into the water info texture */
+	FMaterialRenderProxy* WaterInfoMaterial = nullptr;
+	/** Material to use when rendering this scene proxy in the main scene */
+	FMaterialRenderProxy* WaterLODMaterial = nullptr;
+
+	FMaterialRelevance WaterInfoMaterialRelevance;
+	FMaterialRelevance WaterLODMaterialRelevance;
 
 	EWaterInfoPass CurrentWaterInfoPass = EWaterInfoPass::None;
 };
diff --git a/Engine/Plugins/Experimental/Water/Source/Runtime/Public/WaterMeshComponent.h b/Engine/Plugins/Experimental/Water/Source/Runtime/Public/WaterMeshComponent.h
--- a/Engine/Plugins/Experimental/Water/Source/Runtime/Public/WaterMeshComponent.h
+++ b/Engine/Plugins/Experimental/Water/Source/Runtime/Public/WaterMeshComponent.h
@@ -46,6 +46,8 @@ public:
 	/** Use this instead of GetMaterialRelevance, since this one will go over all materials from all tiles */
 	FMaterialRelevance GetWaterMaterialRelevance(ERHIFeatureLevel::Type InFeatureLevel) const;
 
+	void PushTessellatedWaterMeshBoundsToPoxy(const FBox2D& TessellatedWaterMeshBounds);
+
 	const FWaterQuadTree& GetWaterQuadTree() const { return WaterQuadTree; }
 
 	const UE_TRANSITIONAL_OBJECT_PTR_TEMPLATE(TSet, UMaterialInterface)& GetUsedMaterialsSet() const { return UsedMaterials; }
@@ -68,16 +70,16 @@ public:
 	 *	Setting this to 0 will allow every level to collapse
 	 *	Setting this to something higher than the LODCount will have no effect
 	 */
-	UPROPERTY(EditAnywhere, Category = Mesh, meta = (ClampMin = "-1"))
+	UPROPERTY(EditAnywhere, Category = Rendering, meta = (ClampMin = "-1"))
 	int32 ForceCollapseDensityLevel = -1;
 
-	UPROPERTY(EditAnywhere, Category = "Mesh|FarDistance")
+	UPROPERTY(EditAnywhere, Category = "Rendering|FarDistance")
 	TObjectPtr<UMaterialInterface> FarDistanceMaterial = nullptr;
 
-	UPROPERTY(EditAnywhere, Category = "Mesh|FarDistance", meta = (ClampMin = "0"))
+	UPROPERTY(EditAnywhere, Category = "Rendering|FarDistance", meta = (ClampMin = "0"))
 	float FarDistanceMeshExtent = 0.0f;
 
-	UFUNCTION(BlueprintPure, Category = Mesh)
+	UFUNCTION(BlueprintPure, Category = Rendering)
 	bool IsEnabled() const { return bIsEnabled; }
 private:
 	//~ Begin USceneComponent Interface
@@ -88,11 +90,11 @@ private:
 	void RebuildWaterMesh(float InTileSize, const FIntPoint& InExtentInTiles);
 
 	/** World size of the water tiles at LOD0. Multiply this with the ExtentInTiles to get the world extents of the system */
-	UPROPERTY(EditAnywhere, Category = Mesh, meta = (ClampMin = "100", AllowPrivateAcces = "true"))
+	UPROPERTY(EditAnywhere, Category = Rendering, meta = (ClampMin = "100", AllowPrivateAcces = "true"))
 	float TileSize = 2400.0f;
 
 	/** The extent of the system in number of tiles. Maximum number of tiles for this system will be ExtentInTiles.X*2*ExtentInTiles.Y*2 */
-	UPROPERTY(EditAnywhere, Category = Mesh, meta = (ClampMin = "1", AllowPrivateAcces = "true"))
+	UPROPERTY(EditAnywhere, Category = Rendering, meta = (ClampMin = "1", AllowPrivateAcces = "true"))
 	FIntPoint ExtentInTiles = FIntPoint(64, 64);
 
 	/** Tiles containing water, stored in a quad tree */
@@ -118,16 +120,17 @@ private:
 	float LODScaleBiasScalability = 0.0f;
 
 	/** Highest tessellation factor of a water tile. Max number of verts on the side of a tile will be (2^TessellationFactor)+1)  */
-	UPROPERTY(EditAnywhere, Category = Mesh, meta = (ClampMin = "1", ClampMax = "12"))
+	UPROPERTY(EditAnywhere, Category = Rendering, meta = (ClampMin = "1", ClampMax = "12"))
 	int32 TessellationFactor = 6;
 
 	/** World scale of the concentric LODs */
-	UPROPERTY(EditAnywhere, Category = Mesh, meta = (ClampMin = "0.5"))
+	UPROPERTY(EditAnywhere, Category = Rendering, meta = (ClampMin = "0.5"))
 	float LODScale = 1.0f;
 
 #if WITH_EDITOR
 	//~ Begin USceneComponent Interface
 	virtual void PostEditChangeProperty(FPropertyChangedEvent& PropertyChangedEvent) override;
+	virtual void PostEditComponentMove(bool bFinished) override;
 	//~ Begin USceneComponent Interface
 #endif
 };
\ No newline at end of file
diff --git a/Engine/Plugins/Experimental/Water/Source/Runtime/Public/WaterMeshSceneProxy.h b/Engine/Plugins/Experimental/Water/Source/Runtime/Public/WaterMeshSceneProxy.h
--- a/Engine/Plugins/Experimental/Water/Source/Runtime/Public/WaterMeshSceneProxy.h
+++ b/Engine/Plugins/Experimental/Water/Source/Runtime/Public/WaterMeshSceneProxy.h
@@ -61,6 +61,8 @@ public:
 	virtual bool IsRayTracingRelevant() const override { return true; }
 #endif
 
+	void OnTessellatedWaterMeshBoundsChanged_GameThread(const FBox2D& InTessellatedWaterMeshBounds);
+
 private:
 	struct FWaterLODParams
 	{
@@ -88,6 +90,8 @@ private:
 	void SetupRayTracingInstances(int32 NumInstances, uint32 DensityIndex);
 #endif
 
+	void OnTessellatedWaterMeshBoundsChanged_RenderThread(const FBox2D& InTessellatedWaterMeshBounds);
+
 	FMaterialRelevance MaterialRelevance;
 
 	// One vertex factory per LOD
@@ -102,6 +106,9 @@ private:
 	/** Per-"water render group" user data (the number of groups might vary depending on whether we're in the editor or not) */
 	WaterMeshUserDataBuffersType* WaterMeshUserDataBuffers;
 
+	/** The world-space bounds of the current water info texture coverage. The Water mesh should only render tiles within this bounding box. */
+	FBox2D TessellatedWaterMeshBounds = FBox2D(ForceInit);
+
 	/** Scale of the concentric LOD squares  */
 	float LODScale = -1.0f;
 
diff --git a/Engine/Plugins/Experimental/Water/Source/Runtime/Public/WaterQuadTree.h b/Engine/Plugins/Experimental/Water/Source/Runtime/Public/WaterQuadTree.h
--- a/Engine/Plugins/Experimental/Water/Source/Runtime/Public/WaterQuadTree.h
+++ b/Engine/Plugins/Experimental/Water/Source/Runtime/Public/WaterQuadTree.h
@@ -107,6 +107,7 @@ struct FWaterQuadTree
 		FVector PreViewTranslation = FVector::ZeroVector;
 		FConvexVolume Frustum;
 		bool bLODMorphingEnabled = true;
+		FBox2D TessellatedWaterMeshBounds = FBox2D(ForceInit);
 
 #if !(UE_BUILD_SHIPPING || UE_BUILD_TEST)
 		// Debug
@@ -211,6 +212,9 @@ private:
 		/** Recursive function to select nodes visible from the current point of view */
 		void SelectLOD(const FNodeData& InNodeData, int32 InLODLevel, const FTraversalDesc& InTraversalDesc, FTraversalOutput& Output) const;
 
+		/** Recursive function to select nodes visible from the current point of view within an active bounding box */
+		void SelectLODWithinBounds(const FNodeData& InNodeData, int32 InLODLevel, const FTraversalDesc& InTraversalDesc, FTraversalOutput& Output) const;
+
 		/** Recursive function to query the height(prior to any displacement) at a given location, return false if no height could be found */
 		bool QueryBaseHeightAtLocation(const FNodeData& InNodeData, const FVector2D& InWorldLocationXY, float& OutHeight) const;
 
diff --git a/Engine/Plugins/Experimental/Water/Source/Runtime/Public/WaterSubsystem.h b/Engine/Plugins/Experimental/Water/Source/Runtime/Public/WaterSubsystem.h
--- a/Engine/Plugins/Experimental/Water/Source/Runtime/Public/WaterSubsystem.h
+++ b/Engine/Plugins/Experimental/Water/Source/Runtime/Public/WaterSubsystem.h
@@ -99,6 +99,9 @@ public:
 	/** Static helper function to get a waterbody manager from a world, returns nullptr if world or manager don't exist */
 	static FWaterBodyManager* GetWaterBodyManager(const UWorld* InWorld);
 
+	/** Static helper function to get a weak ptr to the water scene view extension for a given world. */
+	static TWeakPtr<FWaterViewExtension, ESPMode::ThreadSafe> GetWaterViewExtension(const UWorld* InWorld);
+
 	UE_DEPRECATED(5.1, "Please use FWaterBodyManager::ForEachWaterBodyComponent instead.")
 	static void ForEachWaterBodyComponent(const UWorld* World, TFunctionRef<bool(UWaterBodyComponent*)> Predicate) {}
 
@@ -172,8 +175,6 @@ public:
 	
 	void MarkAllWaterZonesForRebuild(EWaterZoneRebuildFlags RebuildFlags = EWaterZoneRebuildFlags::All);
 
-	void MarkWaterInfoTextureForRebuild(const UE::WaterInfo::FRenderingContext& WaterInfoContext);
-
 #if WITH_EDITOR
 	/** Little scope object to temporarily change the value of bAllowWaterSubsystemOnPreviewWorld */
 	struct WATER_API FScopedAllowWaterSubsystemOnPreviewWorld
@@ -251,7 +252,7 @@ private:
 
 	FUnderwaterPostProcessVolume UnderwaterPostProcessVolume;
 
-	TSharedPtr<FWaterViewExtension, ESPMode::ThreadSafe> WaterViewExtension;
+	TSharedPtr<FWaterViewExtension> WaterViewExtension;
 
 #if WITH_EDITOR
 	/** By default, there is no water subsystem allowed on preview worlds except when explicitly requested : */
diff --git a/Engine/Plugins/Experimental/Water/Source/Runtime/Public/WaterViewExtension.h b/Engine/Plugins/Experimental/Water/Source/Runtime/Public/WaterViewExtension.h
--- a/Engine/Plugins/Experimental/Water/Source/Runtime/Public/WaterViewExtension.h
+++ b/Engine/Plugins/Experimental/Water/Source/Runtime/Public/WaterViewExtension.h
@@ -16,9 +16,9 @@ public:
 	~FWaterViewExtension();
 
 	// FSceneViewExtensionBase implementation : 
-	virtual void SetupView(FSceneViewFamily& InViewFamily, FSceneView& InView) override {}
 	virtual void BeginRenderViewFamily(FSceneViewFamily& InViewFamily) override {}
 
+	virtual void SetupView(FSceneViewFamily& InViewFamily, FSceneView& InView) override;
 	virtual void SetupViewFamily(FSceneViewFamily& InViewFamily) override;
 	virtual void PreRenderViewFamily_RenderThread(FRDGBuilder& GraphBuilder, FSceneViewFamily& InViewFamily) override;
 	// End FSceneViewExtensionBase implementation
@@ -26,6 +26,12 @@ public:
 	void MarkWaterInfoTextureForRebuild(const UE::WaterInfo::FRenderingContext& RenderContext);
 
 private:
-	// Store contexts in a map to prevent multiple update entries for the same zone. We only need to keep the most recent.
+	/** Queued Water Info rendering contexts to submit for rendering on the next SetupView call */
 	TMap<AWaterZone*, UE::WaterInfo::FRenderingContext> WaterInfoContextsToRender;
+
+	/** 
+	 * For each water zone, store the bounds of the tile from which the water zone was last rendered.
+	 * When the view location crosses the bounds, submit a new WaterInfo update to reflect the new active area
+	 */
+	TMap<AWaterZone*, FBox2D> UpdateBounds;
 };
diff --git a/Engine/Plugins/Experimental/Water/Source/Runtime/Public/WaterZoneActor.h b/Engine/Plugins/Experimental/Water/Source/Runtime/Public/WaterZoneActor.h
--- a/Engine/Plugins/Experimental/Water/Source/Runtime/Public/WaterZoneActor.h
+++ b/Engine/Plugins/Experimental/Water/Source/Runtime/Public/WaterZoneActor.h
@@ -18,6 +18,7 @@ enum class EWaterZoneRebuildFlags
 	None = 0,
 	UpdateWaterInfoTexture = (1 << 1),
 	UpdateWaterMesh = (1 << 2),
+	UpdateWaterBodyLODSections = (1 << 3),
 	All = (~0),
 };
 ENUM_CLASS_FLAGS(EWaterZoneRebuildFlags);
@@ -36,6 +37,9 @@ public:
 	/** Execute a predicate function on each valid water body within the water zone. Predicate should return false for early exit. */
 	void ForEachWaterBodyComponent(TFunctionRef<bool(UWaterBodyComponent*)> Predicate);
 
+	void AddWaterBodyComponent(UWaterBodyComponent* WaterBodyComponent);
+	void RemoveWaterBodyComponent(UWaterBodyComponent* WaterBodyComponent);
+
 	FVector2D GetZoneExtent() const { return ZoneExtent; }
 	void SetZoneExtent(FVector2D NewExtents);
 
@@ -48,21 +52,32 @@ public:
 	virtual void PostLoadSubobjects(FObjectInstancingGraph* OuterInstanceGraph) override;
 	virtual void PostLoad() override;
 
-#if WITH_EDITOR
-	virtual bool CanChangeIsSpatiallyLoadedFlag() const override { return false; }
-#endif
-
 	FVector2f GetWaterHeightExtents() const { return WaterHeightExtents; }
 	float GetGroundZMin() const { return GroundZMin; }
 
 	UPROPERTY(Transient, DuplicateTransient, VisibleAnywhere, BlueprintReadOnly, Category = Texture)
 	TObjectPtr<UTextureRenderTarget2D> WaterInfoTexture;
+
+	FVector GetTessellatedWaterMeshCenter() const;
+	void SetTessellatedWaterMeshCenter(FVector NewCenter) { TessellatedWaterMeshCenter = NewCenter;}
+
+	FVector GetTessellatedWaterMeshExtent() const;
+
+	uint32 GetNonTessellatedLODScale() const { return NonTessellatedLODSectionScale; }
+	float GetNonTessellatedLODSectionSize() const;
+
+	bool IsNonTessellatedLODMeshEnabled() const { return bEnableNonTessellatedLODMesh; }
+
+#if WITH_EDITOR
+	virtual bool CanChangeIsSpatiallyLoadedFlag() const override { return false; }
+#endif
+
 private:
 
 	bool UpdateWaterInfoTexture();
 
 	void OnExtentChanged();
-	
+
 #if WITH_EDITOR
 	void OnActorSelectionChanged(const TArray<UObject*>& NewSelection, bool bForceRefresh);
 	
@@ -76,7 +91,10 @@ private:
 	/** Called when the Bounds component is modified. Updates the value of ZoneExtent to match the new bounds */
 	void OnBoundsComponentModified();
 #endif // WITH_EDITOR
-	
+
+	UPROPERTY(Transient, Category = Water, VisibleAnywhere)
+	TArray<TWeakObjectPtr<UWaterBodyComponent>> OwnedWaterBodies;
+
 	UPROPERTY(EditAnywhere, BlueprintReadOnly, Category = Texture, meta = (AllowPrivateAccess = "true"))
 	FIntPoint RenderTargetResolution;
 
@@ -84,12 +102,12 @@ private:
 	UPROPERTY(VisibleAnywhere, Category = Water, BlueprintReadOnly, meta = (AllowPrivateAccess = "true"))
 	TObjectPtr<UWaterMeshComponent> WaterMesh;
 
-	/** Radius of the zone bounding box */
+	/** Width of the zone bounding box */
 	UPROPERTY(Category = Shape, EditAnywhere, BlueprintReadWrite, meta = (AllowPrivateAccess = "true"))
 	FVector2D ZoneExtent;
 
 	/** Offsets the height above the water zone at which the WaterInfoTexture is rendered. This is applied after computing the maximum Z of all the water bodies within the zone. */
-	UPROPERTY(Category = Shape, EditAnywhere, meta = (AllowPrivateAccess = "true"))
+	UPROPERTY(Category = Shape, EditAnywhere, AdvancedDisplay, meta = (AllowPrivateAccess = "true"))
 	float CaptureZOffset = 64.f;
 
 	/** Determines if the WaterInfoTexture should be 16 or 32 bits per channel */
@@ -100,11 +118,23 @@ private:
 	UPROPERTY(Category = Texture, EditAnywhere, AdvancedDisplay)
 	int32 VelocityBlurRadius = 1;
 
+	/** Area around the camera covered by the tessellated water mesh when LOD is enabled. */
+	UPROPERTY(Category = NonTessellatedLOD, EditAnywhere, meta = (EditCondition = "bEnableNonTessellatedLODMesh"))
+	FVector TessellatedWaterMeshExtent;
+
+	UPROPERTY(Category = NonTessellatedLOD, EditAnywhere, meta = (ClampMin = "1", DisplayName = "Non Tessellated LOD Section Scale", EditCondition = "bEnableNonTessellatedLODMesh"))
+	uint32 NonTessellatedLODSectionScale = 4;
+
+	UPROPERTY(Category = NonTessellatedLOD, EditAnywhere, meta = (DisplayName = "Enable Non-Tessellated LOD Mesh"))
+	bool bEnableNonTessellatedLODMesh = false;
+
 	bool bNeedsWaterInfoRebuild = true;
 
 	FVector2f WaterHeightExtents;
 	float GroundZMin;
 
+	FVector TessellatedWaterMeshCenter;
+
 #if WITH_EDITORONLY_DATA
 	/** A manipulatable box for visualizing/editing the water zone bounds */
 	UPROPERTY(Transient)
diff --git a/Engine/Plugins/Runtime/GeometryProcessing/Source/DynamicMesh/Private/Operations/MeshPlaneCut.cpp b/Engine/Plugins/Runtime/GeometryProcessing/Source/DynamicMesh/Private/Operations/MeshPlaneCut.cpp
--- a/Engine/Plugins/Runtime/GeometryProcessing/Source/DynamicMesh/Private/Operations/MeshPlaneCut.cpp
+++ b/Engine/Plugins/Runtime/GeometryProcessing/Source/DynamicMesh/Private/Operations/MeshPlaneCut.cpp
@@ -180,19 +180,19 @@ bool FMeshPlaneCut::CutWithoutDelete(bool bSplitVerticesAtPlane, float OffsetVer
 		ensure(OffsetVertices == 0.0); // it would be weird to not split vertices and still request any offset of the 'other side' vertices; please don't do that
 	}
 
-
-	if (!TriLabels)
-	{
-		ensure(!bSplitVerticesAtPlane); // need labels to split verts currently
-		return false;
-	}
-
 	// collapse degenerate edges if we got em
 	if (bCollapseDegenerateEdgesOnCut)
 	{
 		CollapseDegenerateEdges(OnCutEdges, ZeroEdges);
 	}
 
+	if (!bSplitVerticesAtPlane)
+	{
+		return true;
+	}
+
+	ensure(TriLabels); // need labels to split verts currently
+
 	TMap<int, int> OldLabelToNew;
 	int AvailableID = NewLabelStartID;
 	FVector3d VertexOffsetVec = (double)OffsetVertices * PlaneNormal;
diff --git a/Engine/Source/Runtime/Core/Public/UObject/FortniteMainBranchObjectVersion.h b/Engine/Source/Runtime/Core/Public/UObject/FortniteMainBranchObjectVersion.h
--- a/Engine/Source/Runtime/Core/Public/UObject/FortniteMainBranchObjectVersion.h
+++ b/Engine/Source/Runtime/Core/Public/UObject/FortniteMainBranchObjectVersion.h
@@ -234,6 +234,9 @@ struct CORE_API FFortniteMainBranchObjectVersion
 
 		// Reset default value for Water exclusion volumes to make them more intuitive and support the "it just works" philosophy.
 		WaterExclusionVolumeExcludeAllDefault,
+		
+		// Added water non-tessellated LOD
+		WaterNontessellatedLODSupportAdded,
 
 		// -----<new versions can be added above this line>-------------------------------------------------
 		VersionPlusOne,
