Water: removed use of native pointer in uproperty

diff --git a/Engine/Plugins/Experimental/Water/Source/Runtime/Public/WaterBodyComponent.h b/Engine/Plugins/Experimental/Water/Source/Runtime/Public/WaterBodyComponent.h
--- a/Engine/Plugins/Experimental/Water/Source/Runtime/Public/WaterBodyComponent.h
+++ b/Engine/Plugins/Experimental/Water/Source/Runtime/Public/WaterBodyComponent.h
@@ -533,7 +533,7 @@ public:
 	TObjectPtr<UMaterialInterface> WaterHLODMaterial;
 
 	UPROPERTY(Category = Rendering, EditAnywhere, BlueprintReadOnly, meta = (DisplayName = "Water LOD Material"))
-	UMaterialInterface* WaterLODMaterial;
+	TObjectPtr<UMaterialInterface> WaterLODMaterial;
 
 	/** Post process material to apply when the camera goes underwater (only available when bGenerateCollisions is true because collisions are needed to detect if it's under water). */
 	UPROPERTY(Category = Rendering, EditAnywhere, BlueprintReadOnly, meta = (EditCondition = "bGenerateCollisions", DisplayAfter = "WaterMaterial"))
@@ -588,7 +588,7 @@ protected:
 	TObjectPtr<UMaterialInstanceDynamic> WaterMID;
 
 	UPROPERTY(Category = Debug, VisibleInstanceOnly, Transient, NonPIEDuplicateTransient, TextExportTransient, meta = (DisplayAfter = "WaterLODMaterial"))
-	UMaterialInstanceDynamic* WaterLODMID;
+	TObjectPtr<UMaterialInstanceDynamic> WaterLODMID;
 
 	UPROPERTY(Category = Debug, VisibleInstanceOnly, Transient, NonPIEDuplicateTransient, TextExportTransient, meta = (DisplayAfter = "UnderwaterPostProcessMaterial"))
 	TObjectPtr<UMaterialInstanceDynamic> UnderwaterPostProcessMID;
@@ -607,7 +607,7 @@ protected:
 	mutable TWeakObjectPtr<ALandscapeProxy> Landscape;
 
 	UPROPERTY()
-	AWaterZone* OwningWaterZone;
+	TObjectPtr<AWaterZone> OwningWaterZone;
 
 	UPROPERTY(Transient)
 	FPostProcessSettings CurrentPostProcessSettings;
