Tentative fix for a water renderng issues on older Android devices with Adreno GPUs

diff --git a/Engine/Plugins/Experimental/Water/Shaders/Private/WaterInfoMerge.usf b/Engine/Plugins/Experimental/Water/Shaders/Private/WaterInfoMerge.usf
--- a/Engine/Plugins/Experimental/Water/Shaders/Private/WaterInfoMerge.usf
+++ b/Engine/Plugins/Experimental/Water/Shaders/Private/WaterInfoMerge.usf
@@ -61,6 +61,9 @@ void Main(
 
 	float GroundZ = CaptureZ - GroundSceneDepth;
 
+	// Always initialize output to avoid bugs in old shader compilers 
+	OutColor = 0.f;
+
 	// If the GroundDepth is further than the far plane of the projection, set the GroundZ to some really far away value
 	if (GroundDeviceDepth == 1.0)
 	{
