Water: modified landscape search so that if the position of a water body is exactly on the boundary of a landscape component it will still return as a valid overlapping landscape.

This is especially important in WP maps where we may have a waterbody at the corner of multiple valid landscape proxies.

diff --git a/Engine/Plugins/Experimental/Water/Source/Runtime/Private/WaterBodyComponent.cpp b/Engine/Plugins/Experimental/Water/Source/Runtime/Private/WaterBodyComponent.cpp
--- a/Engine/Plugins/Experimental/Water/Source/Runtime/Private/WaterBodyComponent.cpp
+++ b/Engine/Plugins/Experimental/Water/Source/Runtime/Private/WaterBodyComponent.cpp
@@ -909,7 +909,7 @@ ALandscapeProxy* UWaterBodyComponent::FindLandscape() const
 			if (It->GetWorld() == GetWorld())
 			{
 				FBox Box = It->GetComponentsBoundingBox();
-				if (Box.IsInsideXY(Location))
+				if (Box.IsInsideOrOnXY(Location))
 				{
 					Landscape = *It;
 					return Landscape.Get();
