Water: removed optimization which prevented updating water material instances when the water zone was re-rendered. Material instances must always be updated when the zone extent/position/capturez change which is most of the time in the editor.
Needs better condition to check for MID update if we decide to re-enable this.

diff --git a/Engine/Plugins/Experimental/Water/Source/Runtime/Private/WaterZoneActor.cpp b/Engine/Plugins/Experimental/Water/Source/Runtime/Private/WaterZoneActor.cpp
--- a/Engine/Plugins/Experimental/Water/Source/Runtime/Private/WaterZoneActor.cpp
+++ b/Engine/Plugins/Experimental/Water/Source/Runtime/Private/WaterZoneActor.cpp
@@ -371,9 +371,7 @@ bool AWaterZone::UpdateWaterInfoTexture()
 		}
 
 		const ETextureRenderTargetFormat Format = bHalfPrecisionTexture ? ETextureRenderTargetFormat::RTF_RGBA16f : RTF_RGBA32f;
-		UTextureRenderTarget2D* NewWaterInfoTexture = FWaterUtils::GetOrCreateTransientRenderTarget2D(WaterInfoTexture, TEXT("WaterInfoTexture"), RenderTargetResolution, Format);
-		const bool bNeedsMIDUpdate = WaterInfoTexture != NewWaterInfoTexture;
-		WaterInfoTexture = NewWaterInfoTexture;
+		WaterInfoTexture = FWaterUtils::GetOrCreateTransientRenderTarget2D(WaterInfoTexture, TEXT("WaterInfoTexture"), RenderTargetResolution, Format);
 
 		UE::WaterInfo::FRenderingContext Context;
 		Context.ZoneToRender = this;
@@ -387,12 +385,9 @@ bool AWaterZone::UpdateWaterInfoTexture()
 			WaterViewExtension.Pin()->MarkWaterInfoTextureForRebuild(Context);
 		}
 
-		if (bNeedsMIDUpdate)
+		for (UWaterBodyComponent* Component : WaterBodiesToRender)
 		{
-			for (UWaterBodyComponent* Component : WaterBodiesToRender)
-			{
-				Component->UpdateMaterialInstances();
-			}
+			Component->UpdateMaterialInstances();
 		}
 
 		UE_LOG(LogWater, Verbose, TEXT("Queued Water Info texture update"));
