Water: implemented an actionable message on load to create the water collision profile when it is not defined rather than checking out the DefaultEngine.ini on load.

diff --git a/Engine/Plugins/Experimental/Water/Source/Editor/Private/WaterEditorModule.cpp b/Engine/Plugins/Experimental/Water/Source/Editor/Private/WaterEditorModule.cpp
--- a/Engine/Plugins/Experimental/Water/Source/Editor/Private/WaterEditorModule.cpp
+++ b/Engine/Plugins/Experimental/Water/Source/Editor/Private/WaterEditorModule.cpp
@@ -36,6 +36,8 @@
 #include "WaterBodyBrushCacheContainerThumbnailRenderer.h"
 #include "ThumbnailRendering/ThumbnailManager.h"
 #include "WaterWavesEditorToolkit.h"
+#include "Engine/AssetManager.h"
+#include "WaterRuntimeSettings.h"
 
 #define LOCTEXT_NAMESPACE "WaterEditor"
 
@@ -79,12 +81,22 @@ void FWaterEditorModule::StartupModule()
 	}
 
 	UThumbnailManager::Get().RegisterCustomRenderer(UWaterBodyBrushCacheContainer::StaticClass(), UWaterBodyBrushCacheContainerThumbnailRenderer::StaticClass());
+
+	OnLoadCollisionProfileConfigHandle = UCollisionProfile::Get()->OnLoadProfileConfig.AddLambda([this](UCollisionProfile* CollisionProfile)
+		{
+			check(UCollisionProfile::Get() == CollisionProfile);
+			CheckForWaterCollisionProfile();
+		});
+
+	CheckForWaterCollisionProfile();
 }
 
 void FWaterEditorModule::ShutdownModule()
 {
 	if (UObjectInitialized())
 	{
+		UCollisionProfile::Get()->OnLoadProfileConfig.Remove(OnLoadCollisionProfileConfigHandle);
+
 		UThumbnailManager::Get().UnregisterCustomRenderer(UWaterBodyBrushCacheContainer::StaticClass());
 	}
 
@@ -150,51 +162,101 @@ void FWaterEditorModule::OnLevelActorAddedToWorld(AActor* Actor)
 		UWorld* ActorWorld = Actor->GetWorld();
 		if (ActorWorld && ActorWorld->IsEditorWorld())
 		{
-			bool bHasWaterManager = !!TActorIterator<AWaterLandscapeBrush>(ActorWorld);
-			if (!bHasWaterManager)
+			TArray<ALandscape*> FoundLandscapes;
+
+			// Search for all overlapping landscapes and add a water brush to them.
+			// If we cannot find a suitable landscape via this method, default to using the first landscape in the world.
+			
+			const bool bNonColliding = true;
+			const bool bIncludeChildActors = false;
+			const FBox ActorBounds = Actor->GetComponentsBoundingBox(bNonColliding, bIncludeChildActors);
+
+			for (ALandscape* Landscape : TActorRange<ALandscape>(ActorWorld))
 			{
-				const UWaterEditorSettings* WaterEditorSettings = GetDefault<UWaterEditorSettings>();
-				check(WaterEditorSettings != nullptr);
-				TSubclassOf<AWaterLandscapeBrush> WaterBrushClass = WaterEditorSettings->GetWaterManagerClass();
-				if (UClass* WaterBrushClassPtr = WaterBrushClass.Get())
+				const FBox LandscapeBounds = Landscape->GetComponentsBoundingBox(bNonColliding, bIncludeChildActors);
+				if (LandscapeBounds.Intersect(ActorBounds))
 				{
-					UActorFactory* WaterBrushActorFactory = GEditor->FindActorFactoryForActorClass(WaterBrushClassPtr);
+					FoundLandscapes.Add(Landscape);
+				}
+			}
+
+			if (FoundLandscapes.Num() == 0)
+			{
+				UE_LOG(LogWaterEditor, Warning, TEXT("Could not find a suitable landscape to which to assign the water brush! Defaulting to the first landscape."));
+				const TActorIterator<ALandscape> It(ActorWorld);
+				if (ALandscape* Landscape = It ? *It : nullptr)
+				{
+					FoundLandscapes.Add(Landscape);
+				}
+			}
+
+			// Spawn a Water brush for every landscape this actor overlaps with.
+			FBox Bounds;
+			for (ALandscape* FoundLandscape : FoundLandscapes)
+			{
+				const FBox LandscapeBounds = FoundLandscape->GetComponentsBoundingBox(bNonColliding, bIncludeChildActors);
+				Bounds += LandscapeBounds;
+
+				bool bHasWaterManager = false;
+				FoundLandscape->ForEachLayer([&bHasWaterManager](FLandscapeLayer& Layer)
+				{
+					for (const FLandscapeLayerBrush& Brush : Layer.Brushes)
+					{
+						bHasWaterManager |= Cast<AWaterLandscapeBrush>(Brush.GetBrush()) != nullptr;
+					}
+				});
 
-					// If the water manager doesn't exist, spawn it now in the same level as the landscape
-					for (ALandscape* Landscape : TActorRange<ALandscape>(ActorWorld))
+				if (!bHasWaterManager)
+				{
+					const UWaterEditorSettings* WaterEditorSettings = GetDefault<UWaterEditorSettings>();
+					check(WaterEditorSettings != nullptr);
+					TSubclassOf<AWaterLandscapeBrush> WaterBrushClass = WaterEditorSettings->GetWaterManagerClass();
+					if (UClass* WaterBrushClassPtr = WaterBrushClass.Get())
 					{
-						FString BrushActorString = FString::Format(TEXT("{0}_{1}"), { Landscape->GetActorLabel(), WaterBrushClassPtr->GetName() } );
-						FName BrushActorName = MakeUniqueObjectName(Landscape->GetOuter(), WaterBrushClassPtr, FName(BrushActorString));
-						FActorSpawnParameters SpawnParams;
-						SpawnParams.Name = BrushActorName;
-						SpawnParams.bAllowDuringConstructionScript = true; // This can be called by construction script if the actor being added to the world is part of a blueprint, for example : 
-						AWaterLandscapeBrush* NewBrush = (WaterBrushActorFactory != nullptr) 
-							? Cast<AWaterLandscapeBrush>(WaterBrushActorFactory->CreateActor(ActorWorld, Landscape->GetLevel(), FTransform::Identity, SpawnParams))
-							: ActorWorld->SpawnActor<AWaterLandscapeBrush>(WaterBrushClassPtr, SpawnParams);
-						if (NewBrush)
+						UActorFactory* WaterBrushActorFactory = GEditor->FindActorFactoryForActorClass(WaterBrushClassPtr);
+
+						if (FoundLandscape)
 						{
-							bHasWaterManager = true;
-							NewBrush->SetActorLabel(BrushActorString);
-							NewBrush->SetTargetLandscape(Landscape);
+							FString BrushActorString = FString::Format(TEXT("{0}_{1}"), { FoundLandscape->GetActorLabel(), WaterBrushClassPtr->GetName() } );
+							FName BrushActorName = MakeUniqueObjectName(FoundLandscape->GetOuter(), WaterBrushClassPtr, FName(BrushActorString));
+							FActorSpawnParameters SpawnParams;
+							SpawnParams.Name = BrushActorName;
+							SpawnParams.bAllowDuringConstructionScript = true; // This can be called by construction script if the actor being added to the world is part of a blueprint, for example : 
+							AWaterLandscapeBrush* NewBrush = (WaterBrushActorFactory != nullptr) 
+								? Cast<AWaterLandscapeBrush>(WaterBrushActorFactory->CreateActor(ActorWorld, FoundLandscape->GetLevel(), FTransform(LandscapeBounds.GetCenter()), SpawnParams))
+								: ActorWorld->SpawnActor<AWaterLandscapeBrush>(WaterBrushClassPtr, SpawnParams);
+							if (NewBrush)
+							{
+								bHasWaterManager = true;
+								NewBrush->SetActorLabel(BrushActorString);
+								NewBrush->SetTargetLandscape(FoundLandscape);
+							}
 						}
-						break;
 					}
-				}
-				else
-				{
-					UE_LOG(LogWaterEditor, Log, TEXT("Could not find Water Manager class %s to spawn"), *WaterEditorSettings->GetWaterManagerClassPath().GetAssetPathString());
+					else
+					{
+						UE_LOG(LogWaterEditor, Warning, TEXT("Could not find Water Manager class %s to spawn"), *WaterEditorSettings->GetWaterManagerClassPath().GetAssetPathString());
+					}
 				}
 			}
 
 			const bool bHasMeshActor = !!TActorIterator<AWaterZone>(ActorWorld);
 
-			if (!bHasMeshActor && bHasWaterManager)
+			if (!bHasMeshActor)
 			{
 				FActorSpawnParameters SpawnParams;
 				SpawnParams.OverrideLevel = ActorWorld->PersistentLevel;
 				SpawnParams.bAllowDuringConstructionScript = true; // This can be called by construction script if the actor being added to the world is part of a blueprint, for example : 
 				AWaterZone* WaterZoneActor = ActorWorld->SpawnActor<AWaterZone>(AWaterZone::StaticClass(), SpawnParams);
 
+				// Set a more sensible default location and extent so that the zone fully encapsulates the landscape if one exists.
+				if (FoundLandscapes.Num() > 0)
+				{
+					WaterZoneActor->SetActorLocation(Bounds.GetCenter());
+					// FBox::GetExtent returns the radius, SetZoneExtent expects diameter.
+					WaterZoneActor->SetZoneExtent(2 * FVector2D(Bounds.GetExtent()));
+				}
+
 				// Set the defaults here because the actor factory isn't triggered on manual SpawnActor.
 				const FWaterZoneActorDefaults& WaterMeshActorDefaults = GetDefault<UWaterEditorSettings>()->WaterZoneActorDefaults;
 				WaterZoneActor->GetWaterMeshComponent()->FarDistanceMaterial = WaterMeshActorDefaults.GetFarDistanceMaterial();
@@ -204,6 +266,46 @@ void FWaterEditorModule::OnLevelActorAddedToWorld(AActor* Actor)
 	}
 }
 
+void FWaterEditorModule::CheckForWaterCollisionProfile()
+{
+	// Make sure WaterCollisionProfileName is added to Engine's collision profiles
+	const FName WaterCollisionProfileName = GetDefault<UWaterRuntimeSettings>()->GetDefaultWaterCollisionProfileName();
+	FCollisionResponseTemplate WaterBodyCollisionProfile;
+	if (!UCollisionProfile::Get()->GetProfileTemplate(WaterCollisionProfileName, WaterBodyCollisionProfile))
+	{
+		FMessageLog("LoadErrors").Error()
+			->AddToken(FTextToken::Create(LOCTEXT("MissingWaterCollisionProfile", "Collision Profile settings do not include an entry for the Water Body Collision profile, which is required for water collision to function.")))
+			->AddToken(FActionToken::Create(LOCTEXT("AddWaterCollisionProfile", "Add entry to DefaultEngine.ini?"), FText(),
+				FOnActionTokenExecuted::CreateRaw(this, &FWaterEditorModule::AddWaterCollisionProfile), true));
+	}
+}
+
+void FWaterEditorModule::AddWaterCollisionProfile()
+{
+	// Make sure WaterCollisionProfileName is added to Engine's collision profiles
+	const FName WaterCollisionProfileName = GetDefault<UWaterRuntimeSettings>()->GetDefaultWaterCollisionProfileName();
+	FCollisionResponseTemplate WaterBodyCollisionProfile;
+	if (!UCollisionProfile::Get()->GetProfileTemplate(WaterCollisionProfileName, WaterBodyCollisionProfile))
+	{
+		WaterBodyCollisionProfile.Name = WaterCollisionProfileName;
+		WaterBodyCollisionProfile.CollisionEnabled = ECollisionEnabled::QueryOnly;
+		WaterBodyCollisionProfile.ObjectType = ECollisionChannel::ECC_WorldStatic;
+		WaterBodyCollisionProfile.bCanModify = false;
+		WaterBodyCollisionProfile.ResponseToChannels = FCollisionResponseContainer::GetDefaultResponseContainer();
+		WaterBodyCollisionProfile.ResponseToChannels.Camera = ECR_Ignore;
+		WaterBodyCollisionProfile.ResponseToChannels.Visibility = ECR_Ignore;
+		WaterBodyCollisionProfile.ResponseToChannels.WorldDynamic = ECR_Overlap;
+		WaterBodyCollisionProfile.ResponseToChannels.Pawn = ECR_Overlap;
+		WaterBodyCollisionProfile.ResponseToChannels.PhysicsBody = ECR_Overlap;
+		WaterBodyCollisionProfile.ResponseToChannels.Destructible = ECR_Overlap;
+		WaterBodyCollisionProfile.ResponseToChannels.Vehicle = ECR_Overlap;
+#if WITH_EDITORONLY_DATA
+		WaterBodyCollisionProfile.HelpMessage = TEXT("Default Water Collision Profile (Created by Water Plugin)");
+#endif
+		FCollisionProfilePrivateAccessor::AddProfileTemplate(WaterBodyCollisionProfile);
+	}
+}
+
 IMPLEMENT_MODULE(FWaterEditorModule, WaterEditor);
 
 #undef LOCTEXT_NAMESPACE
\ No newline at end of file
diff --git a/Engine/Plugins/Experimental/Water/Source/Editor/Public/WaterEditorModule.h b/Engine/Plugins/Experimental/Water/Source/Editor/Public/WaterEditorModule.h
--- a/Engine/Plugins/Experimental/Water/Source/Editor/Public/WaterEditorModule.h
+++ b/Engine/Plugins/Experimental/Water/Source/Editor/Public/WaterEditorModule.h
@@ -36,6 +36,10 @@ private:
 
 	void OnLevelActorAddedToWorld(AActor* Actor);
 
+	void CheckForWaterCollisionProfile();
+
+	void AddWaterCollisionProfile();
+
 private:
 	/** Array of component class names we have registered, so we know what to unregister afterwards */
 	TArray<FName> RegisteredComponentClassNames;
@@ -43,4 +47,6 @@ private:
 	static EAssetTypeCategories::Type WaterAssetCategory;
 
 	TArray<TSharedPtr<IAssetTypeActions>> CreatedAssetTypeActions;
+
+	FDelegateHandle OnLoadCollisionProfileConfigHandle;
 };
\ No newline at end of file
diff --git a/Engine/Plugins/Experimental/Water/Source/Runtime/Private/WaterSubsystem.cpp b/Engine/Plugins/Experimental/Water/Source/Runtime/Private/WaterSubsystem.cpp
--- a/Engine/Plugins/Experimental/Water/Source/Runtime/Private/WaterSubsystem.cpp
+++ b/Engine/Plugins/Experimental/Water/Source/Runtime/Private/WaterSubsystem.cpp
@@ -22,7 +22,6 @@
 #include "Math/NumericLimits.h"
 #include "BuoyancyManager.h"
 #include "WaterRuntimeSettings.h"
-#include "Engine/CollisionProfile.h"
 #include "Engine/StaticMesh.h"
 #include "WaterViewExtension.h"
 
@@ -281,8 +280,6 @@ void UWaterSubsystem::Initialize(FSubsystemCollectionBase& Collection)
 		// Store the buoyancy manager we create for future use.
 		BuoyancyManager = World->SpawnActor<ABuoyancyManager>(SpawnInfo);
 	}
-	UCollisionProfile::Get()->OnLoadProfileConfig.AddUObject(this, &UWaterSubsystem::OnLoadProfileConfig);
-	AddWaterCollisionProfile();
 }
 
 void UWaterSubsystem::Deinitialize()
@@ -290,8 +287,6 @@ void UWaterSubsystem::Deinitialize()
 	UWorld* World = GetWorld();
 	check(World != nullptr);
 
-	UCollisionProfile::Get()->OnLoadProfileConfig.RemoveAll(this);
-
 	FConsoleVariableDelegate NullCallback;
 	CVarShallowWaterSimulationRenderTargetSize->SetOnChangedCallback(NullCallback);
 	CVarShallowWaterSim->SetOnChangedCallback(NullCallback);
@@ -347,38 +342,6 @@ FWaterBodyManager& UWaterSubsystem::GetWaterBodyManagerInternal()
 	PRAGMA_ENABLE_DEPRECATION_WARNINGS
 }
 
-void UWaterSubsystem::AddWaterCollisionProfile()
-{
-	// Make sure WaterCollisionProfileName is added to Engine's collision profiles
-	const FName WaterCollisionProfileName = GetDefault<UWaterRuntimeSettings>()->GetDefaultWaterCollisionProfileName();
-	FCollisionResponseTemplate WaterBodyCollisionProfile;
-	if (!UCollisionProfile::Get()->GetProfileTemplate(WaterCollisionProfileName, WaterBodyCollisionProfile))
-	{
-		WaterBodyCollisionProfile.Name = WaterCollisionProfileName;
-		WaterBodyCollisionProfile.CollisionEnabled = ECollisionEnabled::QueryOnly;
-		WaterBodyCollisionProfile.ObjectType = ECollisionChannel::ECC_WorldStatic;
-		WaterBodyCollisionProfile.bCanModify = false;
-		WaterBodyCollisionProfile.ResponseToChannels = FCollisionResponseContainer::GetDefaultResponseContainer();
-		WaterBodyCollisionProfile.ResponseToChannels.Camera = ECR_Ignore;
-		WaterBodyCollisionProfile.ResponseToChannels.Visibility = ECR_Ignore;
-		WaterBodyCollisionProfile.ResponseToChannels.WorldDynamic = ECR_Overlap;
-		WaterBodyCollisionProfile.ResponseToChannels.Pawn = ECR_Overlap;
-		WaterBodyCollisionProfile.ResponseToChannels.PhysicsBody = ECR_Overlap;
-		WaterBodyCollisionProfile.ResponseToChannels.Destructible = ECR_Overlap;
-		WaterBodyCollisionProfile.ResponseToChannels.Vehicle = ECR_Overlap;
-#if WITH_EDITORONLY_DATA
-		WaterBodyCollisionProfile.HelpMessage = TEXT("Default Water Collision Profile (Created by Water Plugin)");
-#endif
-		FCollisionProfilePrivateAccessor::AddProfileTemplate(WaterBodyCollisionProfile);
-	}
-}
-
-void UWaterSubsystem::OnLoadProfileConfig(UCollisionProfile* CollisionProfile)
-{
-	check(CollisionProfile == UCollisionProfile::Get());
-	AddWaterCollisionProfile();
-}
-
 bool UWaterSubsystem::IsShallowWaterSimulationEnabled() const
 {
 	return ShallowWaterSim != 0;
diff --git a/Engine/Plugins/Experimental/Water/Source/Runtime/Public/WaterSubsystem.h b/Engine/Plugins/Experimental/Water/Source/Runtime/Public/WaterSubsystem.h
--- a/Engine/Plugins/Experimental/Water/Source/Runtime/Public/WaterSubsystem.h
+++ b/Engine/Plugins/Experimental/Water/Source/Runtime/Public/WaterSubsystem.h
@@ -200,8 +200,6 @@ private:
 	void ComputeUnderwaterPostProcess(FVector ViewLocation, FSceneView* SceneView);
 	void SetMPCTime(float Time, float PrevTime);
 	void AdjustUnderwaterWaterInfoQueryFlags(EWaterBodyQueryFlags& InOutFlags);
-	void OnLoadProfileConfig(class UCollisionProfile* CollisionProfile);
-	void AddWaterCollisionProfile();
 	void ApplyRuntimeSettings(const UWaterRuntimeSettings* Settings, EPropertyChangeType::Type ChangeType);
 
 	FWaterBodyManager& GetWaterBodyManagerInternal();
