Water: changed default water brush and water zone spawn behaviour. Now more intelligently chooses the correct landscape, supports multiple landscapes, and sets the zone bounds to properly match the landscape.

diff --git a/Engine/Plugins/Experimental/Water/Source/Editor/Private/WaterEditorModule.cpp b/Engine/Plugins/Experimental/Water/Source/Editor/Private/WaterEditorModule.cpp
--- a/Engine/Plugins/Experimental/Water/Source/Editor/Private/WaterEditorModule.cpp
+++ b/Engine/Plugins/Experimental/Water/Source/Editor/Private/WaterEditorModule.cpp
@@ -150,51 +150,101 @@ void FWaterEditorModule::OnLevelActorAddedToWorld(AActor* Actor)
 		UWorld* ActorWorld = Actor->GetWorld();
 		if (ActorWorld && ActorWorld->IsEditorWorld())
 		{
-			bool bHasWaterManager = !!TActorIterator<AWaterLandscapeBrush>(ActorWorld);
-			if (!bHasWaterManager)
+			TArray<ALandscape*> FoundLandscapes;
+
+			// Search for all overlapping landscapes and add a water brush to them.
+			// If we cannot find a suitable landscape via this method, default to using the first landscape in the world.
+			
+			const bool bNonColliding = true;
+			const bool bIncludeChildActors = false;
+			const FBox ActorBounds = Actor->GetComponentsBoundingBox(bNonColliding, bIncludeChildActors);
+
+			for (ALandscape* Landscape : TActorRange<ALandscape>(ActorWorld))
+			{
+				const FBox LandscapeBounds = Landscape->GetComponentsBoundingBox(bNonColliding, bIncludeChildActors);
+				if (LandscapeBounds.Intersect(ActorBounds))
+				{
+					FoundLandscapes.Add(Landscape);
+				}
+			}
+
+			if (FoundLandscapes.Num() == 0)
+			{
+				UE_LOG(LogWaterEditor, Warning, TEXT("Could not find a suitable landscape to which to assign the water brush! Defaulting to the first landscape."));
+				const TActorIterator<ALandscape> It(ActorWorld);
+				if (ALandscape* Landscape = It ? *It : nullptr)
+				{
+					FoundLandscapes.Add(Landscape);
+				}
+			}
+
+			// Spawn a Water brush for every landscape this actor overlaps with.
+			FBox Bounds;
+			for (ALandscape* FoundLandscape : FoundLandscapes)
 			{
-				const UWaterEditorSettings* WaterEditorSettings = GetDefault<UWaterEditorSettings>();
-				check(WaterEditorSettings != nullptr);
-				TSubclassOf<AWaterLandscapeBrush> WaterBrushClass = WaterEditorSettings->GetWaterManagerClass();
-				if (UClass* WaterBrushClassPtr = WaterBrushClass.Get())
+				const FBox LandscapeBounds = FoundLandscape->GetComponentsBoundingBox(bNonColliding, bIncludeChildActors);
+				Bounds += LandscapeBounds;
+
+				bool bHasWaterManager = false;
+				FoundLandscape->ForEachLayer([&bHasWaterManager](FLandscapeLayer& Layer)
 				{
-					UActorFactory* WaterBrushActorFactory = GEditor->FindActorFactoryForActorClass(WaterBrushClassPtr);
+					for (const FLandscapeLayerBrush& Brush : Layer.Brushes)
+					{
+						bHasWaterManager |= Cast<AWaterLandscapeBrush>(Brush.GetBrush()) != nullptr;
+					}
+				});
 
-					// If the water manager doesn't exist, spawn it now in the same level as the landscape
-					for (ALandscape* Landscape : TActorRange<ALandscape>(ActorWorld))
+				if (!bHasWaterManager)
+				{
+					const UWaterEditorSettings* WaterEditorSettings = GetDefault<UWaterEditorSettings>();
+					check(WaterEditorSettings != nullptr);
+					TSubclassOf<AWaterLandscapeBrush> WaterBrushClass = WaterEditorSettings->GetWaterManagerClass();
+					if (UClass* WaterBrushClassPtr = WaterBrushClass.Get())
 					{
-						FString BrushActorString = FString::Format(TEXT("{0}_{1}"), { Landscape->GetActorLabel(), WaterBrushClassPtr->GetName() } );
-						FName BrushActorName = MakeUniqueObjectName(Landscape->GetOuter(), WaterBrushClassPtr, FName(BrushActorString));
-						FActorSpawnParameters SpawnParams;
-						SpawnParams.Name = BrushActorName;
-						SpawnParams.bAllowDuringConstructionScript = true; // This can be called by construction script if the actor being added to the world is part of a blueprint, for example : 
-						AWaterLandscapeBrush* NewBrush = (WaterBrushActorFactory != nullptr) 
-							? Cast<AWaterLandscapeBrush>(WaterBrushActorFactory->CreateActor(ActorWorld, Landscape->GetLevel(), FTransform::Identity, SpawnParams))
-							: ActorWorld->SpawnActor<AWaterLandscapeBrush>(WaterBrushClassPtr, SpawnParams);
-						if (NewBrush)
+						UActorFactory* WaterBrushActorFactory = GEditor->FindActorFactoryForActorClass(WaterBrushClassPtr);
+
+						if (FoundLandscape)
 						{
-							bHasWaterManager = true;
-							NewBrush->SetActorLabel(BrushActorString);
-							NewBrush->SetTargetLandscape(Landscape);
+							FString BrushActorString = FString::Format(TEXT("{0}_{1}"), { FoundLandscape->GetActorLabel(), WaterBrushClassPtr->GetName() } );
+							FName BrushActorName = MakeUniqueObjectName(FoundLandscape->GetOuter(), WaterBrushClassPtr, FName(BrushActorString));
+							FActorSpawnParameters SpawnParams;
+							SpawnParams.Name = BrushActorName;
+							SpawnParams.bAllowDuringConstructionScript = true; // This can be called by construction script if the actor being added to the world is part of a blueprint, for example : 
+							AWaterLandscapeBrush* NewBrush = (WaterBrushActorFactory != nullptr) 
+								? Cast<AWaterLandscapeBrush>(WaterBrushActorFactory->CreateActor(ActorWorld, FoundLandscape->GetLevel(), FTransform(LandscapeBounds.GetCenter()), SpawnParams))
+								: ActorWorld->SpawnActor<AWaterLandscapeBrush>(WaterBrushClassPtr, SpawnParams);
+							if (NewBrush)
+							{
+								bHasWaterManager = true;
+								NewBrush->SetActorLabel(BrushActorString);
+								NewBrush->SetTargetLandscape(FoundLandscape);
+							}
 						}
-						break;
 					}
-				}
-				else
-				{
-					UE_LOG(LogWaterEditor, Log, TEXT("Could not find Water Manager class %s to spawn"), *WaterEditorSettings->GetWaterManagerClassPath().GetAssetPathString());
+					else
+					{
+						UE_LOG(LogWaterEditor, Warning, TEXT("Could not find Water Manager class %s to spawn"), *WaterEditorSettings->GetWaterManagerClassPath().GetAssetPathString());
+					}
 				}
 			}
 
 			const bool bHasMeshActor = !!TActorIterator<AWaterZone>(ActorWorld);
 
-			if (!bHasMeshActor && bHasWaterManager)
+			if (!bHasMeshActor)
 			{
 				FActorSpawnParameters SpawnParams;
 				SpawnParams.OverrideLevel = ActorWorld->PersistentLevel;
 				SpawnParams.bAllowDuringConstructionScript = true; // This can be called by construction script if the actor being added to the world is part of a blueprint, for example : 
 				AWaterZone* WaterZoneActor = ActorWorld->SpawnActor<AWaterZone>(AWaterZone::StaticClass(), SpawnParams);
 
+				// Set a more sensible default location and extent so that the zone fully encapsulates the landscape if one exists.
+				if (FoundLandscapes.Num() > 0)
+				{
+					WaterZoneActor->SetActorLocation(Bounds.GetCenter());
+					// FBox::GetExtent returns the radius, SetZoneExtent expects diameter.
+					WaterZoneActor->SetZoneExtent(2 * FVector2D(Bounds.GetExtent()));
+				}
+
 				// Set the defaults here because the actor factory isn't triggered on manual SpawnActor.
 				const FWaterZoneActorDefaults& WaterMeshActorDefaults = GetDefault<UWaterEditorSettings>()->WaterZoneActorDefaults;
 				WaterZoneActor->GetWaterMeshComponent()->FarDistanceMaterial = WaterMeshActorDefaults.GetFarDistanceMaterial();
diff --git a/Engine/Source/Runtime/Landscape/Classes/Landscape.h b/Engine/Source/Runtime/Landscape/Classes/Landscape.h
--- a/Engine/Source/Runtime/Landscape/Classes/Landscape.h
+++ b/Engine/Source/Runtime/Landscape/Classes/Landscape.h
@@ -129,7 +129,7 @@ struct FLandscapeLayerBrush
 
 #if WITH_EDITOR
 	UTextureRenderTarget2D* Render(bool InIsHeightmap, const FIntRect& InLandscapeSize, UTextureRenderTarget2D* InLandscapeRenderTarget, const FName& InWeightmapLayerName = NAME_None);
-	ALandscapeBlueprintBrushBase* GetBrush() const;
+	LANDSCAPE_API ALandscapeBlueprintBrushBase* GetBrush() const;
 	bool IsAffectingHeightmap() const;
 	bool IsAffectingWeightmapLayer(const FName& InWeightmapLayerName) const;
 	void SetOwner(ALandscape* InOwner);
