Water: implemented an actionable message on load to create the water collision profile when it is not defined rather than checking out the DefaultEngine.ini on load.

diff --git a/Engine/Plugins/Experimental/Water/Source/Editor/Private/WaterEditorModule.cpp b/Engine/Plugins/Experimental/Water/Source/Editor/Private/WaterEditorModule.cpp
--- a/Engine/Plugins/Experimental/Water/Source/Editor/Private/WaterEditorModule.cpp
+++ b/Engine/Plugins/Experimental/Water/Source/Editor/Private/WaterEditorModule.cpp
@@ -36,6 +36,8 @@
 #include "WaterBodyBrushCacheContainerThumbnailRenderer.h"
 #include "ThumbnailRendering/ThumbnailManager.h"
 #include "WaterWavesEditorToolkit.h"
+#include "Engine/AssetManager.h"
+#include "WaterRuntimeSettings.h"
 
 #define LOCTEXT_NAMESPACE "WaterEditor"
 
@@ -79,12 +81,22 @@ void FWaterEditorModule::StartupModule()
 	}
 
 	UThumbnailManager::Get().RegisterCustomRenderer(UWaterBodyBrushCacheContainer::StaticClass(), UWaterBodyBrushCacheContainerThumbnailRenderer::StaticClass());
+
+	OnLoadCollisionProfileConfigHandle = UCollisionProfile::Get()->OnLoadProfileConfig.AddLambda([this](UCollisionProfile* CollisionProfile)
+		{
+			check(UCollisionProfile::Get() == CollisionProfile);
+			CheckForWaterCollisionProfile();
+		});
+
+	CheckForWaterCollisionProfile();
 }
 
 void FWaterEditorModule::ShutdownModule()
 {
 	if (UObjectInitialized())
 	{
+		UCollisionProfile::Get()->OnLoadProfileConfig.Remove(OnLoadCollisionProfileConfigHandle);
+
 		UThumbnailManager::Get().UnregisterCustomRenderer(UWaterBodyBrushCacheContainer::StaticClass());
 	}
 
@@ -254,6 +266,46 @@ void FWaterEditorModule::OnLevelActorAddedToWorld(AActor* Actor)
 	}
 }
 
+void FWaterEditorModule::CheckForWaterCollisionProfile()
+{
+	// Make sure WaterCollisionProfileName is added to Engine's collision profiles
+	const FName WaterCollisionProfileName = GetDefault<UWaterRuntimeSettings>()->GetDefaultWaterCollisionProfileName();
+	FCollisionResponseTemplate WaterBodyCollisionProfile;
+	if (!UCollisionProfile::Get()->GetProfileTemplate(WaterCollisionProfileName, WaterBodyCollisionProfile))
+	{
+		FMessageLog("LoadErrors").Error()
+			->AddToken(FTextToken::Create(LOCTEXT("MissingWaterCollisionProfile", "Collision Profile settings do not include an entry for the Water Body Collision profile, which is required for water collision to function.")))
+			->AddToken(FActionToken::Create(LOCTEXT("AddWaterCollisionProfile", "Add entry to DefaultEngine.ini?"), FText(),
+				FOnActionTokenExecuted::CreateRaw(this, &FWaterEditorModule::AddWaterCollisionProfile), true));
+	}
+}
+
+void FWaterEditorModule::AddWaterCollisionProfile()
+{
+	// Make sure WaterCollisionProfileName is added to Engine's collision profiles
+	const FName WaterCollisionProfileName = GetDefault<UWaterRuntimeSettings>()->GetDefaultWaterCollisionProfileName();
+	FCollisionResponseTemplate WaterBodyCollisionProfile;
+	if (!UCollisionProfile::Get()->GetProfileTemplate(WaterCollisionProfileName, WaterBodyCollisionProfile))
+	{
+		WaterBodyCollisionProfile.Name = WaterCollisionProfileName;
+		WaterBodyCollisionProfile.CollisionEnabled = ECollisionEnabled::QueryOnly;
+		WaterBodyCollisionProfile.ObjectType = ECollisionChannel::ECC_WorldStatic;
+		WaterBodyCollisionProfile.bCanModify = false;
+		WaterBodyCollisionProfile.ResponseToChannels = FCollisionResponseContainer::GetDefaultResponseContainer();
+		WaterBodyCollisionProfile.ResponseToChannels.Camera = ECR_Ignore;
+		WaterBodyCollisionProfile.ResponseToChannels.Visibility = ECR_Ignore;
+		WaterBodyCollisionProfile.ResponseToChannels.WorldDynamic = ECR_Overlap;
+		WaterBodyCollisionProfile.ResponseToChannels.Pawn = ECR_Overlap;
+		WaterBodyCollisionProfile.ResponseToChannels.PhysicsBody = ECR_Overlap;
+		WaterBodyCollisionProfile.ResponseToChannels.Destructible = ECR_Overlap;
+		WaterBodyCollisionProfile.ResponseToChannels.Vehicle = ECR_Overlap;
+#if WITH_EDITORONLY_DATA
+		WaterBodyCollisionProfile.HelpMessage = TEXT("Default Water Collision Profile (Created by Water Plugin)");
+#endif
+		FCollisionProfilePrivateAccessor::AddProfileTemplate(WaterBodyCollisionProfile);
+	}
+}
+
 IMPLEMENT_MODULE(FWaterEditorModule, WaterEditor);
 
 #undef LOCTEXT_NAMESPACE
\ No newline at end of file
diff --git a/Engine/Plugins/Experimental/Water/Source/Editor/Public/WaterEditorModule.h b/Engine/Plugins/Experimental/Water/Source/Editor/Public/WaterEditorModule.h
--- a/Engine/Plugins/Experimental/Water/Source/Editor/Public/WaterEditorModule.h
+++ b/Engine/Plugins/Experimental/Water/Source/Editor/Public/WaterEditorModule.h
@@ -36,6 +36,10 @@ private:
 
 	void OnLevelActorAddedToWorld(AActor* Actor);
 
+	void CheckForWaterCollisionProfile();
+
+	void AddWaterCollisionProfile();
+
 private:
 	/** Array of component class names we have registered, so we know what to unregister afterwards */
 	TArray<FName> RegisteredComponentClassNames;
@@ -43,4 +47,6 @@ private:
 	static EAssetTypeCategories::Type WaterAssetCategory;
 
 	TArray<TSharedPtr<IAssetTypeActions>> CreatedAssetTypeActions;
+
+	FDelegateHandle OnLoadCollisionProfileConfigHandle;
 };
\ No newline at end of file
diff --git a/Engine/Plugins/Experimental/Water/Source/Runtime/Private/WaterSubsystem.cpp b/Engine/Plugins/Experimental/Water/Source/Runtime/Private/WaterSubsystem.cpp
--- a/Engine/Plugins/Experimental/Water/Source/Runtime/Private/WaterSubsystem.cpp
+++ b/Engine/Plugins/Experimental/Water/Source/Runtime/Private/WaterSubsystem.cpp
@@ -22,7 +22,6 @@
 #include "Math/NumericLimits.h"
 #include "BuoyancyManager.h"
 #include "WaterRuntimeSettings.h"
-#include "Engine/CollisionProfile.h"
 #include "Engine/StaticMesh.h"
 #include "WaterViewExtension.h"
 
@@ -281,8 +280,6 @@ void UWaterSubsystem::Initialize(FSubsystemCollectionBase& Collection)
 		// Store the buoyancy manager we create for future use.
 		BuoyancyManager = World->SpawnActor<ABuoyancyManager>(SpawnInfo);
 	}
-	UCollisionProfile::Get()->OnLoadProfileConfig.AddUObject(this, &UWaterSubsystem::OnLoadProfileConfig);
-	AddWaterCollisionProfile();
 }
 
 void UWaterSubsystem::Deinitialize()
@@ -290,8 +287,6 @@ void UWaterSubsystem::Deinitialize()
 	UWorld* World = GetWorld();
 	check(World != nullptr);
 
-	UCollisionProfile::Get()->OnLoadProfileConfig.RemoveAll(this);
-
 	FConsoleVariableDelegate NullCallback;
 	CVarShallowWaterSimulationRenderTargetSize->SetOnChangedCallback(NullCallback);
 	CVarShallowWaterSim->SetOnChangedCallback(NullCallback);
@@ -347,38 +342,6 @@ FWaterBodyManager& UWaterSubsystem::GetWaterBodyManagerInternal()
 	PRAGMA_ENABLE_DEPRECATION_WARNINGS
 }
 
-void UWaterSubsystem::AddWaterCollisionProfile()
-{
-	// Make sure WaterCollisionProfileName is added to Engine's collision profiles
-	const FName WaterCollisionProfileName = GetDefault<UWaterRuntimeSettings>()->GetDefaultWaterCollisionProfileName();
-	FCollisionResponseTemplate WaterBodyCollisionProfile;
-	if (!UCollisionProfile::Get()->GetProfileTemplate(WaterCollisionProfileName, WaterBodyCollisionProfile))
-	{
-		WaterBodyCollisionProfile.Name = WaterCollisionProfileName;
-		WaterBodyCollisionProfile.CollisionEnabled = ECollisionEnabled::QueryOnly;
-		WaterBodyCollisionProfile.ObjectType = ECollisionChannel::ECC_WorldStatic;
-		WaterBodyCollisionProfile.bCanModify = false;
-		WaterBodyCollisionProfile.ResponseToChannels = FCollisionResponseContainer::GetDefaultResponseContainer();
-		WaterBodyCollisionProfile.ResponseToChannels.Camera = ECR_Ignore;
-		WaterBodyCollisionProfile.ResponseToChannels.Visibility = ECR_Ignore;
-		WaterBodyCollisionProfile.ResponseToChannels.WorldDynamic = ECR_Overlap;
-		WaterBodyCollisionProfile.ResponseToChannels.Pawn = ECR_Overlap;
-		WaterBodyCollisionProfile.ResponseToChannels.PhysicsBody = ECR_Overlap;
-		WaterBodyCollisionProfile.ResponseToChannels.Destructible = ECR_Overlap;
-		WaterBodyCollisionProfile.ResponseToChannels.Vehicle = ECR_Overlap;
-#if WITH_EDITORONLY_DATA
-		WaterBodyCollisionProfile.HelpMessage = TEXT("Default Water Collision Profile (Created by Water Plugin)");
-#endif
-		FCollisionProfilePrivateAccessor::AddProfileTemplate(WaterBodyCollisionProfile);
-	}
-}
-
-void UWaterSubsystem::OnLoadProfileConfig(UCollisionProfile* CollisionProfile)
-{
-	check(CollisionProfile == UCollisionProfile::Get());
-	AddWaterCollisionProfile();
-}
-
 bool UWaterSubsystem::IsShallowWaterSimulationEnabled() const
 {
 	return ShallowWaterSim != 0;
diff --git a/Engine/Plugins/Experimental/Water/Source/Runtime/Public/WaterSubsystem.h b/Engine/Plugins/Experimental/Water/Source/Runtime/Public/WaterSubsystem.h
--- a/Engine/Plugins/Experimental/Water/Source/Runtime/Public/WaterSubsystem.h
+++ b/Engine/Plugins/Experimental/Water/Source/Runtime/Public/WaterSubsystem.h
@@ -200,8 +200,6 @@ private:
 	void ComputeUnderwaterPostProcess(FVector ViewLocation, FSceneView* SceneView);
 	void SetMPCTime(float Time, float PrevTime);
 	void AdjustUnderwaterWaterInfoQueryFlags(EWaterBodyQueryFlags& InOutFlags);
-	void OnLoadProfileConfig(class UCollisionProfile* CollisionProfile);
-	void AddWaterCollisionProfile();
 	void ApplyRuntimeSettings(const UWaterRuntimeSettings* Settings, EPropertyChangeType::Type ChangeType);
 
 	FWaterBodyManager& GetWaterBodyManagerInternal();
