Water: Added a warning when lakes and oceans fail to triangulate to give users some indication when their geometry is invalid.

diff --git a/Engine/Plugins/Experimental/Water/Source/Runtime/Private/WaterBodyLakeComponent.cpp b/Engine/Plugins/Experimental/Water/Source/Runtime/Private/WaterBodyLakeComponent.cpp
--- a/Engine/Plugins/Experimental/Water/Source/Runtime/Private/WaterBodyLakeComponent.cpp
+++ b/Engine/Plugins/Experimental/Water/Source/Runtime/Private/WaterBodyLakeComponent.cpp
@@ -75,7 +75,10 @@ bool UWaterBodyLakeComponent::GenerateWaterBodyMesh(UE::Geometry::FDynamicMesh3&
 	FConstrainedDelaunay2d Triangulation;
 	Triangulation.FillRule = FConstrainedDelaunay2d::EFillRule::Positive;
 	Triangulation.Add(LakePoly);
-	Triangulation.Triangulate();
+	if (!Triangulation.Triangulate())
+	{
+		UE_LOG(LogWater, Error, TEXT("Failed to triangulate Lake mesh for %s"), *GetOwner()->GetActorNameOrLabel());
+	}
 
 	if (Triangulation.Triangles.Num() == 0)
 	{
diff --git a/Engine/Plugins/Experimental/Water/Source/Runtime/Private/WaterBodyOceanComponent.cpp b/Engine/Plugins/Experimental/Water/Source/Runtime/Private/WaterBodyOceanComponent.cpp
--- a/Engine/Plugins/Experimental/Water/Source/Runtime/Private/WaterBodyOceanComponent.cpp
+++ b/Engine/Plugins/Experimental/Water/Source/Runtime/Private/WaterBodyOceanComponent.cpp
@@ -224,8 +224,7 @@ bool UWaterBodyOceanComponent::GenerateWaterBodyMesh(UE::Geometry::FDynamicMesh3
 	Triangulation.Add(Island);
 	if (!Triangulation.Triangulate())
 	{
-		UE_LOG(LogWater, Warning, TEXT("Failed to triangulate ocean (%s). Spline might not be fully contained by VisualExtents."), *GetOwner()->GetActorNameOrLabel());
-		return false;
+		UE_LOG(LogWater, Error, TEXT("Failed to triangulate Ocean mesh for %s. Ensure that the spline points are fully contained within the VisualExtents of the ocean!"), *GetOwner()->GetActorNameOrLabel());
 	}
 
 	if (Triangulation.Triangles.Num() == 0)
