Water: disable editing of falloff angle/width depending which mode is selected to make it more obvious to users which setting is currently active.

diff --git a/Engine/Plugins/Experimental/Water/Source/Runtime/Public/WaterFalloffSettings.h b/Engine/Plugins/Experimental/Water/Source/Runtime/Public/WaterFalloffSettings.h
--- a/Engine/Plugins/Experimental/Water/Source/Runtime/Public/WaterFalloffSettings.h
+++ b/Engine/Plugins/Experimental/Water/Source/Runtime/Public/WaterFalloffSettings.h
@@ -22,10 +22,10 @@ struct FWaterFalloffSettings
 	UPROPERTY(EditAnywhere, BlueprintReadWrite, Category = FalloffSettings)
 	EWaterBrushFalloffMode FalloffMode;
 
-	UPROPERTY(EditAnywhere, BlueprintReadWrite, Category = FalloffSettings)
+	UPROPERTY(EditAnywhere, BlueprintReadWrite, Category = FalloffSettings, meta = (EditCondition = "FalloffMode == EWaterBrushFalloffMode::Angle"))
 	float FalloffAngle;
 
-	UPROPERTY(EditAnywhere, BlueprintReadWrite, Category = FalloffSettings, meta = (ClampMin = "0.1"))
+	UPROPERTY(EditAnywhere, BlueprintReadWrite, Category = FalloffSettings, meta = (ClampMin = "0.1", EditCondition = "FalloffMode == EWaterBrushFalloffMode::Width"))
 	float FalloffWidth;
 
 	UPROPERTY(EditAnywhere, BlueprintReadWrite, Category = FalloffSettings)
