Water: fixed issue where Gerstner waves were calling BP scripts during PostDuplicate which is not allowed. Instead, just mark the gpu buffers dirty as the other properties which were being set by the bp functions should have already been duplicated properly anyways.

diff --git a/Engine/Plugins/Experimental/Water/Source/Runtime/Private/GerstnerWaterWaves.cpp b/Engine/Plugins/Experimental/Water/Source/Runtime/Private/GerstnerWaterWaves.cpp
--- a/Engine/Plugins/Experimental/Water/Source/Runtime/Private/GerstnerWaterWaves.cpp
+++ b/Engine/Plugins/Experimental/Water/Source/Runtime/Private/GerstnerWaterWaves.cpp
@@ -35,7 +35,10 @@ void UGerstnerWaterWaves::PostDuplicate(EDuplicateMode::Type DuplicateMode)
 {
 	Super::PostDuplicate(DuplicateMode);
 
-	RecomputeWaves(/* bAllowBPScript = */true);
+	if (UGerstnerWaterWaveSubsystem* GerstnerWaterWaveSubsystem = GEngine ? GEngine->GetEngineSubsystem<UGerstnerWaterWaveSubsystem>() : nullptr)
+	{
+		GerstnerWaterWaveSubsystem->RebuildGPUData();
+	}
 }
 
 #if WITH_EDITOR
