Fixed PVS warning in WaterEditorModule

diff --git a/Engine/Plugins/Experimental/Water/Source/Editor/Private/WaterEditorModule.cpp b/Engine/Plugins/Experimental/Water/Source/Editor/Private/WaterEditorModule.cpp
--- a/Engine/Plugins/Experimental/Water/Source/Editor/Private/WaterEditorModule.cpp
+++ b/Engine/Plugins/Experimental/Water/Source/Editor/Private/WaterEditorModule.cpp
@@ -194,6 +194,11 @@ void FWaterEditorModule::OnLevelActorAddedToWorld(AActor* Actor)
 			FBox Bounds;
 			for (ALandscape* FoundLandscape : FoundLandscapes)
 			{
+				if (!IsValid(FoundLandscape))
+				{
+					continue;
+				}
+
 				const FBox LandscapeBounds = FoundLandscape->GetComponentsBoundingBox(bNonColliding, bIncludeChildActors);
 				Bounds += LandscapeBounds;
 
