[Water]
* BuoyancyTypes - Removed include not needed

diff --git a/Engine/Plugins/Experimental/Water/Source/Runtime/Public/BuoyancyTypes.h b/Engine/Plugins/Experimental/Water/Source/Runtime/Public/BuoyancyTypes.h
--- a/Engine/Plugins/Experimental/Water/Source/Runtime/Public/BuoyancyTypes.h
+++ b/Engine/Plugins/Experimental/Water/Source/Runtime/Public/BuoyancyTypes.h
@@ -4,7 +4,6 @@
 
 #include "WaterBodyTypes.h"
 #include "Chaos/GeometryParticlesfwd.h"
-#include "PhysicsProxy/SingleParticlePhysicsProxy.h"
 #include "PhysicsProxy/SingleParticlePhysicsProxyFwd.h"
 #include "BuoyancyTypes.generated.h"
 
