Water: Create Water landscape edit layer in Water Editor Module when we spawn the water brush instead of relying on it to do so when the target landscape changes.

diff --git a/Engine/Plugins/Experimental/Water/Source/Editor/Private/WaterEditorModule.cpp b/Engine/Plugins/Experimental/Water/Source/Editor/Private/WaterEditorModule.cpp
--- a/Engine/Plugins/Experimental/Water/Source/Editor/Private/WaterEditorModule.cpp
+++ b/Engine/Plugins/Experimental/Water/Source/Editor/Private/WaterEditorModule.cpp
@@ -173,7 +173,7 @@ void FWaterEditorModule::OnLevelActorAddedToWorld(AActor* Actor)
 
 			for (ALandscape* Landscape : TActorRange<ALandscape>(ActorWorld))
 			{
-				const FBox LandscapeBounds = Landscape->GetComponentsBoundingBox(bNonColliding, bIncludeChildActors);
+				const FBox LandscapeBounds = Landscape->GetCompleteBounds();
 				if (LandscapeBounds.Intersect(ActorBounds))
 				{
 					FoundLandscapes.Add(Landscape);
@@ -190,60 +190,74 @@ void FWaterEditorModule::OnLevelActorAddedToWorld(AActor* Actor)
 				}
 			}
 
-			// Spawn a Water brush for every landscape this actor overlaps with.
+			const UWaterEditorSettings* WaterEditorSettings = GetDefault<UWaterEditorSettings>();
+			check(WaterEditorSettings != nullptr);
+
+			TSubclassOf<AWaterLandscapeBrush> WaterBrushClass = WaterEditorSettings->GetWaterManagerClass();
+			UClass* WaterBrushClassPtr = WaterBrushClass.Get();
+
 			FBox Bounds;
-			for (ALandscape* FoundLandscape : FoundLandscapes)
+			if (WaterBrushClassPtr)
 			{
-				if (!IsValid(FoundLandscape))
+				// Spawn a Water brush for every landscape this actor overlaps with.
+				for (ALandscape* FoundLandscape : FoundLandscapes)
 				{
-					continue;
-				}
+					if (!IsValid(FoundLandscape))
+					{
+						continue;
+					}
 
-				const FBox LandscapeBounds = FoundLandscape->GetComponentsBoundingBox(bNonColliding, bIncludeChildActors);
-				Bounds += LandscapeBounds;
+					const FBox LandscapeBounds = FoundLandscape->GetCompleteBounds();
+					Bounds += LandscapeBounds;
 
-				bool bHasWaterManager = false;
-				FoundLandscape->ForEachLayer([&bHasWaterManager](FLandscapeLayer& Layer)
-				{
-					for (const FLandscapeLayerBrush& Brush : Layer.Brushes)
+					bool bHasWaterManager = false;
+					FoundLandscape->ForEachLayer([&bHasWaterManager](FLandscapeLayer& Layer)
 					{
-						bHasWaterManager |= Cast<AWaterLandscapeBrush>(Brush.GetBrush()) != nullptr;
-					}
-				});
+						for (const FLandscapeLayerBrush& Brush : Layer.Brushes)
+						{
+							bHasWaterManager |= Cast<AWaterLandscapeBrush>(Brush.GetBrush()) != nullptr;
+						}
+					});
 
-				if (!bHasWaterManager)
-				{
-					const UWaterEditorSettings* WaterEditorSettings = GetDefault<UWaterEditorSettings>();
-					check(WaterEditorSettings != nullptr);
-					TSubclassOf<AWaterLandscapeBrush> WaterBrushClass = WaterEditorSettings->GetWaterManagerClass();
-					if (UClass* WaterBrushClassPtr = WaterBrushClass.Get())
+					if (!bHasWaterManager)
 					{
 						UActorFactory* WaterBrushActorFactory = GEditor->FindActorFactoryForActorClass(WaterBrushClassPtr);
+						
+						// Attempt to find an existing water layer, else create one before spawning the water brush
+						const FName WaterLayerName = FName("Water");
+						int32 ExistingWaterLayerIndex = FoundLandscape->GetLayerIndex(WaterLayerName);
+						if (ExistingWaterLayerIndex == INDEX_NONE)
+						{
+							ExistingWaterLayerIndex = FoundLandscape->CreateLayer(WaterLayerName);
+						}
 
-						if (FoundLandscape)
+						FString BrushActorString = FString::Format(TEXT("{0}_{1}"), { FoundLandscape->GetActorLabel(), WaterBrushClassPtr->GetName() } );
+						FName BrushActorName = MakeUniqueObjectName(FoundLandscape->GetOuter(), WaterBrushClassPtr, FName(BrushActorString));
+						FActorSpawnParameters SpawnParams;
+						SpawnParams.Name = BrushActorName;
+						SpawnParams.bAllowDuringConstructionScript = true; // This can be called by construction script if the actor being added to the world is part of a blueprint, for example : 
+						AWaterLandscapeBrush* NewBrush = (WaterBrushActorFactory != nullptr) 
+							? Cast<AWaterLandscapeBrush>(WaterBrushActorFactory->CreateActor(ActorWorld, FoundLandscape->GetLevel(), FTransform(LandscapeBounds.GetCenter()), SpawnParams))
+							: ActorWorld->SpawnActor<AWaterLandscapeBrush>(WaterBrushClassPtr, SpawnParams);
+						if (NewBrush)
 						{
-							FString BrushActorString = FString::Format(TEXT("{0}_{1}"), { FoundLandscape->GetActorLabel(), WaterBrushClassPtr->GetName() } );
-							FName BrushActorName = MakeUniqueObjectName(FoundLandscape->GetOuter(), WaterBrushClassPtr, FName(BrushActorString));
-							FActorSpawnParameters SpawnParams;
-							SpawnParams.Name = BrushActorName;
-							SpawnParams.bAllowDuringConstructionScript = true; // This can be called by construction script if the actor being added to the world is part of a blueprint, for example : 
-							AWaterLandscapeBrush* NewBrush = (WaterBrushActorFactory != nullptr) 
-								? Cast<AWaterLandscapeBrush>(WaterBrushActorFactory->CreateActor(ActorWorld, FoundLandscape->GetLevel(), FTransform(LandscapeBounds.GetCenter()), SpawnParams))
-								: ActorWorld->SpawnActor<AWaterLandscapeBrush>(WaterBrushClassPtr, SpawnParams);
-							if (NewBrush)
+							bHasWaterManager = true;
+							NewBrush->SetActorLabel(BrushActorString);
+							NewBrush->SetTargetLandscape(FoundLandscape);
+							const int32 CurrentBrushLayer = FoundLandscape->GetBrushLayer(NewBrush);
+							if (CurrentBrushLayer != ExistingWaterLayerIndex)
 							{
-								bHasWaterManager = true;
-								NewBrush->SetActorLabel(BrushActorString);
-								NewBrush->SetTargetLandscape(FoundLandscape);
+								FoundLandscape->RemoveBrush(NewBrush);
+								FoundLandscape->AddBrushToLayer(ExistingWaterLayerIndex, NewBrush);
 							}
 						}
 					}
-					else
-					{
-						UE_LOG(LogWaterEditor, Warning, TEXT("Could not find Water Manager class %s to spawn"), *WaterEditorSettings->GetWaterManagerClassPath().GetAssetPathString());
-					}
 				}
 			}
+			else
+			{
+				UE_LOG(LogWaterEditor, Warning, TEXT("Could not find Water Manager class %s to spawn"), *WaterEditorSettings->GetWaterManagerClassPath().GetAssetPathString());
+			}
 
 			const bool bHasMeshActor = !!TActorIterator<AWaterZone>(ActorWorld);
 
