Added profile scope for water info texture update

diff --git a/Engine/Plugins/Experimental/Water/Source/Runtime/Private/WaterZoneActor.cpp b/Engine/Plugins/Experimental/Water/Source/Runtime/Private/WaterZoneActor.cpp
--- a/Engine/Plugins/Experimental/Water/Source/Runtime/Private/WaterZoneActor.cpp
+++ b/Engine/Plugins/Experimental/Water/Source/Runtime/Private/WaterZoneActor.cpp
@@ -276,10 +276,11 @@ bool AWaterZone::UpdateWaterInfoTexture()
 {
 	if (UWorld* World = GetWorld(); World && FApp::CanEverRender())
 	{
+		TRACE_CPUPROFILER_EVENT_SCOPE(AWaterZone::UpdateWaterInfoTexture);
+
 		float WaterZMin(TNumericLimits<float>::Max());
 		float WaterZMax(TNumericLimits<float>::Lowest());
 	
-		
 		// Collect a list of all materials used in the water info render to ensure they have complete shaders maps.
 		// If they do not, we must submit compile jobs for them and wait until they are finished before re-rendering.
 		TArray<UMaterialInterface*> UsedMaterials;
