World Partition - Minimap: Fix water not rendering into minimap

diff --git a/Engine/Plugins/Experimental/Water/Source/Runtime/Private/WaterViewExtension.cpp b/Engine/Plugins/Experimental/Water/Source/Runtime/Private/WaterViewExtension.cpp
--- a/Engine/Plugins/Experimental/Water/Source/Runtime/Private/WaterViewExtension.cpp
+++ b/Engine/Plugins/Experimental/Water/Source/Runtime/Private/WaterViewExtension.cpp
@@ -46,8 +46,11 @@ void FWaterViewExtension::SetupView(FSceneViewFamily& InViewFamily, FSceneView&
 	// #todo_water: hacky way to prevent updating from within the water info passes. Alternative solution might be to set a flag from within the info render
 	// similarly to the SetWithinWaterInfoPass functions for water body components 
 	//This won't be necessary when waterinfo passes don't utilize scene renderers.
-	if (!InView.bIsSceneCapture)
+	static bool bUpdatingWaterInfo = false;
+	if (!bUpdatingWaterInfo)
 	{
+		bUpdatingWaterInfo = true;
+
 		for (const TPair<AWaterZone*, UE::WaterInfo::FRenderingContext>& Pair : WaterInfoContextsToRender)
 		{
 			AWaterZone* WaterZone = Pair.Key;
@@ -96,6 +99,8 @@ void FWaterViewExtension::SetupView(FSceneViewFamily& InViewFamily, FSceneView&
 				}
 			}
 		}
+
+		bUpdatingWaterInfo = false;
 	}
 
 }
