Cosmetic improvements to water-related details panels

diff --git a/Engine/Plugins/Experimental/Water/Source/Runtime/Public/WaterBodyComponent.h b/Engine/Plugins/Experimental/Water/Source/Runtime/Public/WaterBodyComponent.h
--- a/Engine/Plugins/Experimental/Water/Source/Runtime/Public/WaterBodyComponent.h
+++ b/Engine/Plugins/Experimental/Water/Source/Runtime/Public/WaterBodyComponent.h
@@ -106,7 +106,7 @@ public:
 
 // ----------------------------------------------------------------------------------
 
-UCLASS(Abstract, HideCategories = (Tags, Activation, Cooking, Replication, Input, AssetUserData, Mesh))
+UCLASS(Abstract, HideCategories = (Tags, Activation, Cooking, Physics, Replication, Input, AssetUserData, Mesh))
 class WATER_API UWaterBodyComponent : public UPrimitiveComponent
 {
 	GENERATED_UCLASS_BODY()
@@ -574,7 +574,7 @@ protected:
 	 * When this is set to true, the water mesh will always generate tiles for this water body. 
 	 * For example, this can be useful to generate water tiles even when the water material is invalid, for the case where "empty" water tiles are actually desirable.
 	 */
-	UPROPERTY(Category = Rendering, EditAnywhere, BlueprintReadOnly)
+	UPROPERTY(Category = Rendering, EditAnywhere, BlueprintReadOnly, AdvancedDisplay)
 	bool bAlwaysGenerateWaterMeshTiles = false;
 
 	/** Higher number is higher priority. If two water bodies overlap and they don't have a transition material specified, this will be used to determine which water body to use the material from. Valid range is -8192 to 8191 */
diff --git a/Engine/Plugins/Experimental/Water/Source/Runtime/Public/WaterMeshComponent.h b/Engine/Plugins/Experimental/Water/Source/Runtime/Public/WaterMeshComponent.h
--- a/Engine/Plugins/Experimental/Water/Source/Runtime/Public/WaterMeshComponent.h
+++ b/Engine/Plugins/Experimental/Water/Source/Runtime/Public/WaterMeshComponent.h
@@ -15,7 +15,7 @@
  * Water Mesh Component responsible for generating and rendering a continuous water mesh on top of all the existing water body actors in the world
  * The component contains a quadtree which defines where there are water tiles. A function for traversing the quadtree and outputing a list of instance data for each tile to be rendered from a point of view is included
  */
-UCLASS(ClassGroup = (Rendering, Water), hidecategories = (Object, Activation, "Components|Activation"), editinlinenew)
+UCLASS(ClassGroup = (Rendering, Water), hidecategories = (Object, Activation, "Components|Activation", Collision, Lighting, HLOD, Navigation, Replication, Input, MaterialParameters, TextureStreaming), editinlinenew)
 class WATER_API UWaterMeshComponent : public UMeshComponent
 {
 	GENERATED_BODY()
diff --git a/Engine/Plugins/Experimental/Water/Source/Runtime/Public/WaterZoneActor.h b/Engine/Plugins/Experimental/Water/Source/Runtime/Public/WaterZoneActor.h
--- a/Engine/Plugins/Experimental/Water/Source/Runtime/Public/WaterZoneActor.h
+++ b/Engine/Plugins/Experimental/Water/Source/Runtime/Public/WaterZoneActor.h
@@ -23,7 +23,7 @@ enum class EWaterZoneRebuildFlags
 };
 ENUM_CLASS_FLAGS(EWaterZoneRebuildFlags);
 
-UCLASS(Blueprintable)
+UCLASS(Blueprintable, HideCategories=(Physics, Replication, Input, Collision))
 class WATER_API AWaterZone : public AActor
 {
 	GENERATED_UCLASS_BODY()
@@ -55,7 +55,7 @@ public:
 	FVector2f GetWaterHeightExtents() const { return WaterHeightExtents; }
 	float GetGroundZMin() const { return GroundZMin; }
 
-	UPROPERTY(Transient, DuplicateTransient, VisibleAnywhere, BlueprintReadOnly, Category = Texture)
+	UPROPERTY(Transient, DuplicateTransient, VisibleAnywhere, BlueprintReadOnly, Category = Water)
 	TObjectPtr<UTextureRenderTarget2D> WaterInfoTexture;
 
 	FVector GetTessellatedWaterMeshCenter() const;
@@ -92,10 +92,12 @@ private:
 	void OnBoundsComponentModified();
 #endif // WITH_EDITOR
 
+private:
+
 	UPROPERTY(Transient, Category = Water, VisibleAnywhere)
 	TArray<TWeakObjectPtr<UWaterBodyComponent>> OwnedWaterBodies;
 
-	UPROPERTY(EditAnywhere, BlueprintReadOnly, Category = Texture, meta = (AllowPrivateAccess = "true"))
+	UPROPERTY(EditAnywhere, BlueprintReadOnly, Category = Water, meta = (AllowPrivateAccess = "true"))
 	FIntPoint RenderTargetResolution;
 
 	/** The water mesh component */
@@ -103,19 +105,19 @@ private:
 	TObjectPtr<UWaterMeshComponent> WaterMesh;
 
 	/** Width of the zone bounding box */
-	UPROPERTY(Category = Shape, EditAnywhere, BlueprintReadWrite, meta = (AllowPrivateAccess = "true"))
+	UPROPERTY(Category = Water, EditAnywhere, BlueprintReadWrite, meta = (AllowPrivateAccess = "true"))
 	FVector2D ZoneExtent;
 
 	/** Offsets the height above the water zone at which the WaterInfoTexture is rendered. This is applied after computing the maximum Z of all the water bodies within the zone. */
-	UPROPERTY(Category = Shape, EditAnywhere, AdvancedDisplay, meta = (AllowPrivateAccess = "true"))
+	UPROPERTY(Category = Water, EditAnywhere, AdvancedDisplay, meta = (AllowPrivateAccess = "true"))
 	float CaptureZOffset = 64.f;
 
 	/** Determines if the WaterInfoTexture should be 16 or 32 bits per channel */
-	UPROPERTY(Category = Texture, EditAnywhere, AdvancedDisplay, meta = (AllowPrivateAccess = "true"))
+	UPROPERTY(Category = Water, EditAnywhere, AdvancedDisplay, meta = (AllowPrivateAccess = "true"))
 	bool bHalfPrecisionTexture = true;
 
 	/** Radius of the velocity blur in the finalize water info pass */
-	UPROPERTY(Category = Texture, EditAnywhere, AdvancedDisplay)
+	UPROPERTY(Category = Water, EditAnywhere, AdvancedDisplay)
 	int32 VelocityBlurRadius = 1;
 
 	/** Area around the camera covered by the tessellated water mesh when LOD is enabled. */
