Fixed WaterMesh not being added automatically when adding a water body that doesn't affect landscape into the map

diff --git a/Engine/Plugins/Experimental/Water/Source/Editor/Private/WaterEditorModule.cpp b/Engine/Plugins/Experimental/Water/Source/Editor/Private/WaterEditorModule.cpp
--- a/Engine/Plugins/Experimental/Water/Source/Editor/Private/WaterEditorModule.cpp
+++ b/Engine/Plugins/Experimental/Water/Source/Editor/Private/WaterEditorModule.cpp
@@ -156,59 +156,66 @@ void FWaterEditorModule::RegisterComponentVisualizer(FName ComponentClassName, T
 
 void FWaterEditorModule::OnLevelActorAddedToWorld(AActor* Actor)
 {
+	UWorld* ActorWorld = Actor->GetWorld();
 	IWaterBrushActorInterface* WaterBrushActor = Cast<IWaterBrushActorInterface>(Actor);
-	if (WaterBrushActor && !Actor->bIsEditorPreviewActor && !Actor->HasAnyFlags(RF_Transient) && WaterBrushActor->AffectsLandscape())
+	AWaterBody* WaterBodyActor = Cast<AWaterBody>(Actor);
+
+	if (!Actor->bIsEditorPreviewActor && !Actor->HasAnyFlags(RF_Transient)
+		&& ActorWorld && ActorWorld->IsEditorWorld()
+		&& ((WaterBrushActor != nullptr) || (WaterBodyActor != nullptr)))
 	{
-		UWorld* ActorWorld = Actor->GetWorld();
-		if (ActorWorld && ActorWorld->IsEditorWorld())
-		{
-			TArray<ALandscape*> FoundLandscapes;
+		// Search for all overlapping landscapes 
+		// If we cannot find a suitable landscape via this method, default to using the first landscape in the world.
 
-			// Search for all overlapping landscapes and add a water brush to them.
-			// If we cannot find a suitable landscape via this method, default to using the first landscape in the world.
-			
-			const bool bNonColliding = true;
-			const bool bIncludeChildActors = false;
-			const FBox ActorBounds = Actor->GetComponentsBoundingBox(bNonColliding, bIncludeChildActors);
+		FBox WaterZoneBounds(ForceInit);
+		TArray<ALandscape*> FoundLandscapes;
+		bool bFoundIntersectingLandscape = false;
 
-			for (ALandscape* Landscape : TActorRange<ALandscape>(ActorWorld))
+		const bool bNonColliding = true;
+		const bool bIncludeChildActors = false;
+		const FBox ActorBounds = Actor->GetComponentsBoundingBox(bNonColliding, bIncludeChildActors);
+
+		for (ALandscape* Landscape : TActorRange<ALandscape>(ActorWorld))
+		{
+			const FBox LandscapeBounds = Landscape->GetCompleteBounds();
+			if (LandscapeBounds.Intersect(ActorBounds))
 			{
-				const FBox LandscapeBounds = Landscape->GetCompleteBounds();
-				if (LandscapeBounds.Intersect(ActorBounds))
-				{
-					FoundLandscapes.Add(Landscape);
-				}
+				FoundLandscapes.Add(Landscape);
+				// Make sure the water zone's bounds is large enough to fit all landscapes that intersect with this water body :
+				WaterZoneBounds += LandscapeBounds;
+				bFoundIntersectingLandscape = true;
 			}
+		}
 
-			if (FoundLandscapes.Num() == 0)
+		// If no intersectiong landscape was found, use the first one available, if any :
+		if (!bFoundIntersectingLandscape)
+		{
+			const TActorIterator<ALandscape> It(ActorWorld);
+			if (ALandscape* Landscape = It ? *It : nullptr)
 			{
-				UE_LOG(LogWaterEditor, Warning, TEXT("Could not find a suitable landscape to which to assign the water brush! Defaulting to the first landscape."));
-				const TActorIterator<ALandscape> It(ActorWorld);
-				if (ALandscape* Landscape = It ? *It : nullptr)
-				{
-					FoundLandscapes.Add(Landscape);
-				}
+				FoundLandscapes.Add(Landscape);
+				const FBox LandscapeBounds = Landscape->GetCompleteBounds();
+				WaterZoneBounds += LandscapeBounds;
 			}
+		}
 
+		// Automatically setup landscape-affecting features (water brush) if needed : 
+		if ((WaterBrushActor != nullptr) && WaterBrushActor->AffectsLandscape() && !FoundLandscapes.IsEmpty())
+		{
 			const UWaterEditorSettings* WaterEditorSettings = GetDefault<UWaterEditorSettings>();
 			check(WaterEditorSettings != nullptr);
-
 			TSubclassOf<AWaterLandscapeBrush> WaterBrushClass = WaterEditorSettings->GetWaterManagerClass();
-			UClass* WaterBrushClassPtr = WaterBrushClass.Get();
-
-			FBox Bounds;
-			if (WaterBrushClassPtr)
+			if (UClass* WaterBrushClassPtr = WaterBrushClass.Get())
 			{
+				if (!bFoundIntersectingLandscape)
+				{
+					UE_LOG(LogWaterEditor, Warning, TEXT("Could not find a suitable landscape to which to assign the water brush! Defaulting to the first landscape."));
+				}
+
 				// Spawn a Water brush for every landscape this actor overlaps with.
 				for (ALandscape* FoundLandscape : FoundLandscapes)
 				{
-					if (!IsValid(FoundLandscape))
-					{
-						continue;
-					}
-
-					const FBox LandscapeBounds = FoundLandscape->GetCompleteBounds();
-					Bounds += LandscapeBounds;
+					check(IsValid(FoundLandscape));
 
 					bool bHasWaterManager = false;
 					FoundLandscape->ForEachLayer([&bHasWaterManager](FLandscapeLayer& Layer)
@@ -222,7 +229,7 @@ void FWaterEditorModule::OnLevelActorAddedToWorld(AActor* Actor)
 					if (!bHasWaterManager)
 					{
 						UActorFactory* WaterBrushActorFactory = GEditor->FindActorFactoryForActorClass(WaterBrushClassPtr);
-						
+
 						// Attempt to find an existing water layer, else create one before spawning the water brush
 						const FName WaterLayerName = FName("Water");
 						int32 ExistingWaterLayerIndex = FoundLandscape->GetLayerIndex(WaterLayerName);
@@ -231,13 +238,13 @@ void FWaterEditorModule::OnLevelActorAddedToWorld(AActor* Actor)
 							ExistingWaterLayerIndex = FoundLandscape->CreateLayer(WaterLayerName);
 						}
 
-						FString BrushActorString = FString::Format(TEXT("{0}_{1}"), { FoundLandscape->GetActorLabel(), WaterBrushClassPtr->GetName() } );
+						FString BrushActorString = FString::Format(TEXT("{0}_{1}"), { FoundLandscape->GetActorLabel(), WaterBrushClassPtr->GetName() });
 						FName BrushActorName = MakeUniqueObjectName(FoundLandscape->GetOuter(), WaterBrushClassPtr, FName(BrushActorString));
 						FActorSpawnParameters SpawnParams;
 						SpawnParams.Name = BrushActorName;
 						SpawnParams.bAllowDuringConstructionScript = true; // This can be called by construction script if the actor being added to the world is part of a blueprint, for example : 
-						AWaterLandscapeBrush* NewBrush = (WaterBrushActorFactory != nullptr) 
-							? Cast<AWaterLandscapeBrush>(WaterBrushActorFactory->CreateActor(ActorWorld, FoundLandscape->GetLevel(), FTransform(LandscapeBounds.GetCenter()), SpawnParams))
+						AWaterLandscapeBrush* NewBrush = (WaterBrushActorFactory != nullptr)
+							? Cast<AWaterLandscapeBrush>(WaterBrushActorFactory->CreateActor(ActorWorld, FoundLandscape->GetLevel(), FTransform(WaterZoneBounds.GetCenter()), SpawnParams))
 							: ActorWorld->SpawnActor<AWaterLandscapeBrush>(WaterBrushClassPtr, SpawnParams);
 						if (NewBrush)
 						{
@@ -258,29 +265,30 @@ void FWaterEditorModule::OnLevelActorAddedToWorld(AActor* Actor)
 			{
 				UE_LOG(LogWaterEditor, Warning, TEXT("Could not find Water Manager class %s to spawn"), *WaterEditorSettings->GetWaterManagerClassPath().GetAssetPathString());
 			}
+		}
 
-			const bool bHasMeshActor = !!TActorIterator<AWaterZone>(ActorWorld);
-
-			if (!bHasMeshActor)
+		// Setup the water zone actor for this water body : 
+		const bool bHasZoneActor = !!TActorIterator<AWaterZone>(ActorWorld);
+		if ((WaterBodyActor != nullptr) && !bHasZoneActor)
+		{
+			FActorSpawnParameters SpawnParams;
+			SpawnParams.OverrideLevel = ActorWorld->PersistentLevel;
+			SpawnParams.bAllowDuringConstructionScript = true; // This can be called by construction script if the actor being added to the world is part of a blueprint, for example : 
+			AWaterZone* WaterZoneActor = ActorWorld->SpawnActor<AWaterZone>(AWaterZone::StaticClass(), SpawnParams);
+
+			// TODO [jonathan.bard] : when we can tag static meshes as "water ground", add these to the bounds
+			// Set a more sensible default location and extent so that the zone fully encapsulates the landscape if one exists.
+			if (WaterZoneBounds.IsValid)
 			{
-				FActorSpawnParameters SpawnParams;
-				SpawnParams.OverrideLevel = ActorWorld->PersistentLevel;
-				SpawnParams.bAllowDuringConstructionScript = true; // This can be called by construction script if the actor being added to the world is part of a blueprint, for example : 
-				AWaterZone* WaterZoneActor = ActorWorld->SpawnActor<AWaterZone>(AWaterZone::StaticClass(), SpawnParams);
-
-				// Set a more sensible default location and extent so that the zone fully encapsulates the landscape if one exists.
-				if (FoundLandscapes.Num() > 0)
-				{
-					WaterZoneActor->SetActorLocation(Bounds.GetCenter());
-					// FBox::GetExtent returns the radius, SetZoneExtent expects diameter.
-					WaterZoneActor->SetZoneExtent(2 * FVector2D(Bounds.GetExtent()));
-				}
-
-				// Set the defaults here because the actor factory isn't triggered on manual SpawnActor.
-				const FWaterZoneActorDefaults& WaterMeshActorDefaults = GetDefault<UWaterEditorSettings>()->WaterZoneActorDefaults;
-				WaterZoneActor->GetWaterMeshComponent()->FarDistanceMaterial = WaterMeshActorDefaults.GetFarDistanceMaterial();
-				WaterZoneActor->GetWaterMeshComponent()->FarDistanceMeshExtent = WaterMeshActorDefaults.FarDistanceMeshExtent;
+				WaterZoneActor->SetActorLocation(WaterZoneBounds.GetCenter());
+				// FBox::GetExtent returns the radius, SetZoneExtent expects diameter.
+				WaterZoneActor->SetZoneExtent(2 * FVector2D(WaterZoneBounds.GetExtent()));
 			}
+
+			// Set the defaults here because the actor factory isn't triggered on manual SpawnActor.
+			const FWaterZoneActorDefaults& WaterMeshActorDefaults = GetDefault<UWaterEditorSettings>()->WaterZoneActorDefaults;
+			WaterZoneActor->GetWaterMeshComponent()->FarDistanceMaterial = WaterMeshActorDefaults.GetFarDistanceMaterial();
+			WaterZoneActor->GetWaterMeshComponent()->FarDistanceMeshExtent = WaterMeshActorDefaults.FarDistanceMeshExtent;
 		}
 	}
 }
