Water: added hooks to update WaterZones' OwnedWaterBodies when the zone is moved/resized

diff --git a/Engine/Plugins/Experimental/Water/Source/Runtime/Private/WaterZoneActor.cpp b/Engine/Plugins/Experimental/Water/Source/Runtime/Private/WaterZoneActor.cpp
--- a/Engine/Plugins/Experimental/Water/Source/Runtime/Private/WaterZoneActor.cpp
+++ b/Engine/Plugins/Experimental/Water/Source/Runtime/Private/WaterZoneActor.cpp
@@ -179,9 +179,28 @@ void AWaterZone::PostEditMove(bool bFinished)
 	{
 		RebuildFlags |= EWaterZoneRebuildFlags::UpdateWaterBodyLODSections;
 	}
+
+	UpdateOverlappingWaterBodies();
+
 	MarkForRebuild(RebuildFlags);
 }
 
+void AWaterZone::PostEditUndo()
+{
+	Super::PostEditUndo();
+
+	UpdateOverlappingWaterBodies();
+	MarkForRebuild(EWaterZoneRebuildFlags::All);
+}
+
+void AWaterZone::PostEditImport()
+{
+	Super::PostEditImport();
+
+	UpdateOverlappingWaterBodies();
+	MarkForRebuild(EWaterZoneRebuildFlags::All);
+}
+
 void AWaterZone::PostEditChangeProperty(FPropertyChangedEvent& PropertyChangedEvent)
 {
 	Super::PostEditChangeProperty(PropertyChangedEvent);
@@ -258,9 +277,21 @@ void AWaterZone::OnExtentChanged()
 	BoundsComponent->SetBoxExtent(FVector(ZoneHalfExtent, 8192.f));
 #endif // WITH_EDITOR
 
+	UpdateOverlappingWaterBodies();
+
 	MarkForRebuild(EWaterZoneRebuildFlags::All);
 }
 
+void AWaterZone::UpdateOverlappingWaterBodies()
+{
+	FWaterBodyManager::ForEachWaterBodyComponent(GetWorld(), [](UWaterBodyComponent* WaterBodyComponent)
+		{
+			WaterBodyComponent->UpdateWaterZones();
+			return true;
+		});
+
+}
+
 
 #if WITH_EDITOR
 void AWaterZone::OnBoundsComponentModified()
diff --git a/Engine/Plugins/Experimental/Water/Source/Runtime/Public/WaterBodyComponent.h b/Engine/Plugins/Experimental/Water/Source/Runtime/Public/WaterBodyComponent.h
--- a/Engine/Plugins/Experimental/Water/Source/Runtime/Public/WaterBodyComponent.h
+++ b/Engine/Plugins/Experimental/Water/Source/Runtime/Public/WaterBodyComponent.h
@@ -363,6 +363,9 @@ public:
 	/** Gets the water zone to which this component belongs */
 	AWaterZone* GetWaterZone() const { return OwningWaterZone; }
 
+	/** Registers or this water body with corresponding overlapping water zones and unregisters it from any old zones if they are no longer overlapping. */
+	void UpdateWaterZones();
+
 	void UpdateWaterBodyRenderData();
 	void UpdateNonTessellatedMeshSections();
 
@@ -445,7 +448,6 @@ protected:
 	EWaterBodyQueryFlags CheckAndAjustQueryFlags(EWaterBodyQueryFlags InQueryFlags) const;
 	void UpdateSplineComponent();
 	void UpdateExclusionVolumes();
-	void UpdateWaterZones();
 	AWaterZone* FindWaterZone() const;
 	bool UpdateWaterHeight();
 	virtual void CreateOrUpdateWaterMID();
diff --git a/Engine/Plugins/Experimental/Water/Source/Runtime/Public/WaterZoneActor.h b/Engine/Plugins/Experimental/Water/Source/Runtime/Public/WaterZoneActor.h
--- a/Engine/Plugins/Experimental/Water/Source/Runtime/Public/WaterZoneActor.h
+++ b/Engine/Plugins/Experimental/Water/Source/Runtime/Public/WaterZoneActor.h
@@ -77,6 +77,7 @@ private:
 	bool UpdateWaterInfoTexture();
 
 	void OnExtentChanged();
+	void UpdateOverlappingWaterBodies();
 
 #if WITH_EDITOR
 	void OnActorSelectionChanged(const TArray<UObject*>& NewSelection, bool bForceRefresh);
@@ -85,6 +86,8 @@ private:
 	void ForceUpdateWaterInfoTexture();
 
 	virtual void PostEditMove(bool bFinished) override;
+	virtual void PostEditUndo() override;
+	virtual void PostEditImport() override;
 
 	virtual void PostEditChangeProperty(FPropertyChangedEvent& PropertyChangedEvent) override;
 
