Water: Fixed issue where water body is not unregistered from water zone when it unloads.

diff --git a/Engine/Plugins/Experimental/Water/Source/Runtime/Private/WaterBodyComponent.cpp b/Engine/Plugins/Experimental/Water/Source/Runtime/Private/WaterBodyComponent.cpp
--- a/Engine/Plugins/Experimental/Water/Source/Runtime/Private/WaterBodyComponent.cpp
+++ b/Engine/Plugins/Experimental/Water/Source/Runtime/Private/WaterBodyComponent.cpp
@@ -754,6 +754,11 @@ void UWaterBodyComponent::OnUnregister()
 	RegisterOnChangeWaterSplineData(/*bRegister = */false);
 #endif // WITH_EDITOR
 
+	if (AWaterZone* WaterZone = GetWaterZone(); IsValid(WaterZone))
+	{
+		WaterZone->RemoveWaterBodyComponent(this);
+	}
+
 	Super::OnUnregister();
 }
 
diff --git a/Engine/Plugins/Experimental/Water/Source/Runtime/Private/WaterZoneActor.cpp b/Engine/Plugins/Experimental/Water/Source/Runtime/Private/WaterZoneActor.cpp
--- a/Engine/Plugins/Experimental/Water/Source/Runtime/Private/WaterZoneActor.cpp
+++ b/Engine/Plugins/Experimental/Water/Source/Runtime/Private/WaterZoneActor.cpp
@@ -125,7 +125,7 @@ void AWaterZone::MarkForRebuild(EWaterZoneRebuildFlags Flags)
 	}
 }
 
-void AWaterZone::ForEachWaterBodyComponent(TFunctionRef<bool(UWaterBodyComponent*)> Predicate)
+void AWaterZone::ForEachWaterBodyComponent(TFunctionRef<bool(UWaterBodyComponent*)> Predicate) const
 {
 	for (const TWeakObjectPtr<UWaterBodyComponent>& WaterBodyComponent : OwnedWaterBodies)
 	{
@@ -139,11 +139,13 @@ void AWaterZone::ForEachWaterBodyComponent(TFunctionRef<bool(UWaterBodyComponent
 void AWaterZone::AddWaterBodyComponent(UWaterBodyComponent* WaterBodyComponent)
 {
 	OwnedWaterBodies.AddUnique(WaterBodyComponent);
+	MarkForRebuild(EWaterZoneRebuildFlags::All);
 }
 
 void AWaterZone::RemoveWaterBodyComponent(UWaterBodyComponent* WaterBodyComponent)
 {
 	OwnedWaterBodies.RemoveSwap(WaterBodyComponent);
+	MarkForRebuild(EWaterZoneRebuildFlags::All);
 }
 
 void AWaterZone::Update()
diff --git a/Engine/Plugins/Experimental/Water/Source/Runtime/Public/WaterZoneActor.h b/Engine/Plugins/Experimental/Water/Source/Runtime/Public/WaterZoneActor.h
--- a/Engine/Plugins/Experimental/Water/Source/Runtime/Public/WaterZoneActor.h
+++ b/Engine/Plugins/Experimental/Water/Source/Runtime/Public/WaterZoneActor.h
@@ -35,7 +35,7 @@ public:
 	void Update();
 		
 	/** Execute a predicate function on each valid water body within the water zone. Predicate should return false for early exit. */
-	void ForEachWaterBodyComponent(TFunctionRef<bool(UWaterBodyComponent*)> Predicate);
+	void ForEachWaterBodyComponent(TFunctionRef<bool(UWaterBodyComponent*)> Predicate) const;
 
 	void AddWaterBodyComponent(UWaterBodyComponent* WaterBodyComponent);
 	void RemoveWaterBodyComponent(UWaterBodyComponent* WaterBodyComponent);
