Improved CPU traces for water body component

#rnx

diff --git a/Engine/Plugins/Experimental/Water/Source/Runtime/Private/WaterBodyComponent.cpp b/Engine/Plugins/Experimental/Water/Source/Runtime/Private/WaterBodyComponent.cpp
--- a/Engine/Plugins/Experimental/Water/Source/Runtime/Private/WaterBodyComponent.cpp
+++ b/Engine/Plugins/Experimental/Water/Source/Runtime/Private/WaterBodyComponent.cpp
@@ -340,7 +340,7 @@ void UWaterBodyComponent::RemoveIsland(AWaterBodyIsland* Island)
 
 void UWaterBodyComponent::UpdateIslands()
 {
-	QUICK_SCOPE_CYCLE_COUNTER(STAT_Water_UpdateIslands);
+	TRACE_CPUPROFILER_EVENT_SCOPE(UWaterBodyComponent::UpdateIslands);
 
 	// For now, islands are not detected dynamically
 #if WITH_EDITOR
@@ -366,7 +366,7 @@ void UWaterBodyComponent::RemoveExclusionVolume(AWaterBodyExclusionVolume* InExc
 
 void UWaterBodyComponent::UpdateExclusionVolumes()
 {
-	QUICK_SCOPE_CYCLE_COUNTER(STAT_Water_UpdateExclusionVolumes);
+	TRACE_CPUPROFILER_EVENT_SCOPE(UWaterBodyComponent::UpdateExclusionVolumes);
 	if (GetWorld())
 	{
 		for (AWaterBodyExclusionVolume* ExclusionVolume : TActorRange<AWaterBodyExclusionVolume>(GetWorld()))
@@ -1255,7 +1255,7 @@ void UWaterBodyComponent::UpdateAll(bool bShapeOrPositionChanged)
 	
 	if (GIsEditor || IsBodyDynamic())
 	{
-		QUICK_SCOPE_CYCLE_COUNTER(STAT_Water_UpdateAll);
+		TRACE_CPUPROFILER_EVENT_SCOPE(UWaterBodyComponent::UpdateAll);
 
 		bShapeOrPositionChanged |= UpdateWaterHeight();
 
@@ -1310,7 +1310,7 @@ void UWaterBodyComponent::UpdateSplineComponent()
 
 void UWaterBodyComponent::OnWaterBodyChanged(bool bShapeOrPositionChanged, bool bWeightmapSettingsChanged)
 {
-	TRACE_CPUPROFILER_EVENT_SCOPE("UWaterBodyComponent::OnWaterBodyChanged")
+	TRACE_CPUPROFILER_EVENT_SCOPE(UWaterBodyComponent::OnWaterBodyChanged)
 	// It's possible to get called without a water spline after the Redo of a water body deletion (i.e. the water body actor gets deleted again, hence its SplineComp is restored to nullptr)
 	//  This is a very-edgy case that needs to be checked everywhere that UpdateAll might hook into so it's simpler to just skip it all. The actor is in limbo by then anyway (it only survives because
 	//  of the editor transaction) :
@@ -1709,6 +1709,8 @@ static inline FColor PackFlowData(float VelocityMagnitude, float DirectionAngle,
 
 void UWaterBodyComponent::RebuildWaterBodyInfoMesh()
 {
+	TRACE_CPUPROFILER_EVENT_SCOPE(UWaterBodyComponent::RebuildWaterBodyInfoMesh);
+
 	using namespace UE::Geometry;
 
 	WaterBodyMeshVertices.Empty();
@@ -1764,7 +1766,7 @@ void UWaterBodyComponent::RebuildWaterBodyLODSections()
 
 	WaterBodyMeshSections.Empty();
 
-	TRACE_CPUPROFILER_EVENT_SCOPE("ComputeWaterBodyMeshSections");
+	TRACE_CPUPROFILER_EVENT_SCOPE(UWaterBodyComponent::RebuildWaterBodyLODSections);
 
 	const AWaterZone* WaterZone = GetWaterZone();
 	check(WaterZone);
@@ -1795,7 +1797,7 @@ void UWaterBodyComponent::RebuildWaterBodyLODSections()
 
 	// Slice the mesh along the grid defined by the Water Mesh
 	{
-		TRACE_CPUPROFILER_EVENT_SCOPE("ComputeWaterBodyMeshSections::PartitionWaterBodyMesh");
+		TRACE_CPUPROFILER_EVENT_SCOPE(ComputeWaterBodyMeshSections::PartitionWaterBodyMesh);
 
 		// We start at an index of -1 and go to +1 here to ensure the bordering sections don't contain triangles that could extend outside the bounds of the cell.
 		for (int32 SectionIndexX = -1; SectionIndexX < WaterZoneExtentInLODSections.X + 1; ++SectionIndexX)
@@ -1833,7 +1835,7 @@ void UWaterBodyComponent::RebuildWaterBodyLODSections()
 	const float MaxFlowVelocity = FWaterUtils::GetWaterMaxFlowVelocity(false);
 	ParallelFor(NumberOfSections, [&](int32 SectionIndex)
 	{
-		TRACE_CPUPROFILER_EVENT_SCOPE("ComputeWaterBodyMeshSections::BuildMeshSection");
+		TRACE_CPUPROFILER_EVENT_SCOPE(ComputeWaterBodyMeshSections::BuildMeshSection);
 		const int32 SectionIndexX = SectionIndex % WaterZoneExtentInLODSections.X;
 		const int32 SectionIndexY = SectionIndex / WaterZoneExtentInLODSections.X;
 
