Water: fixed unnecessary water updates when components are unregistered.

diff --git a/Engine/Plugins/Experimental/Water/Source/Runtime/Private/WaterZoneActor.cpp b/Engine/Plugins/Experimental/Water/Source/Runtime/Private/WaterZoneActor.cpp
--- a/Engine/Plugins/Experimental/Water/Source/Runtime/Private/WaterZoneActor.cpp
+++ b/Engine/Plugins/Experimental/Water/Source/Runtime/Private/WaterZoneActor.cpp
@@ -119,7 +119,7 @@ void AWaterZone::MarkForRebuild(EWaterZoneRebuildFlags Flags)
 	{
 		ForEachWaterBodyComponent([](UWaterBodyComponent* WaterBodyComponent)
 		{
-			WaterBodyComponent->UpdateWaterBodyRenderData();
+			WaterBodyComponent->UpdateNonTessellatedMeshSections();
 			return true;
 		});
 	}
@@ -139,13 +139,13 @@ void AWaterZone::ForEachWaterBodyComponent(TFunctionRef<bool(UWaterBodyComponent
 void AWaterZone::AddWaterBodyComponent(UWaterBodyComponent* WaterBodyComponent)
 {
 	OwnedWaterBodies.AddUnique(WaterBodyComponent);
-	MarkForRebuild(EWaterZoneRebuildFlags::All);
+	MarkForRebuild(EWaterZoneRebuildFlags::UpdateWaterInfoTexture | EWaterZoneRebuildFlags::UpdateWaterMesh);
 }
 
 void AWaterZone::RemoveWaterBodyComponent(UWaterBodyComponent* WaterBodyComponent)
 {
 	OwnedWaterBodies.RemoveSwap(WaterBodyComponent);
-	MarkForRebuild(EWaterZoneRebuildFlags::All);
+	MarkForRebuild(EWaterZoneRebuildFlags::UpdateWaterInfoTexture | EWaterZoneRebuildFlags::UpdateWaterMesh);
 }
 
 void AWaterZone::Update()
