Fixed incorrect CPU perf tags for water plugin

#skipci
#rnx

diff --git a/Engine/Plugins/Experimental/Water/Source/Runtime/Private/WaterBodySceneProxy.cpp b/Engine/Plugins/Experimental/Water/Source/Runtime/Private/WaterBodySceneProxy.cpp
--- a/Engine/Plugins/Experimental/Water/Source/Runtime/Private/WaterBodySceneProxy.cpp
+++ b/Engine/Plugins/Experimental/Water/Source/Runtime/Private/WaterBodySceneProxy.cpp
@@ -17,7 +17,7 @@ FWaterBodySceneProxy::FWaterBodySceneProxy(UWaterBodyComponent* Component, const
 	, WaterBodyInfoMesh(GetScene().GetFeatureLevel())
 	, WaterBodyInfoDilatedMesh(GetScene().GetFeatureLevel())
 {
-	TRACE_CPUPROFILER_EVENT_SCOPE("FWaterBodySceneProxy::FWaterBodySceneProxy");
+	TRACE_CPUPROFILER_EVENT_SCOPE(FWaterBodySceneProxy::FWaterBodySceneProxy);
 	WaterBodySectionedLODMesh.InitFromSections(Component->WaterBodyMeshSections);
 	WaterBodyInfoMesh.Init(Component->WaterBodyMeshVertices, Component->WaterBodyMeshIndices);
 	WaterBodyInfoDilatedMesh.Init(Component->DilatedWaterBodyMeshVertices, Component->DilatedWaterBodyMeshIndices);
@@ -37,7 +37,7 @@ FWaterBodySceneProxy::FWaterBodySceneProxy(UWaterBodyComponent* Component, const
 
 void FWaterBodySceneProxy::FWaterBodySectionedLODMesh::RebuildIndexBuffer(const FBox2D& TessellatedWaterMeshBounds)
 {
-	TRACE_CPUPROFILER_EVENT_SCOPE("FWaterBodySceneProxy::RebuildIndexBuffer");
+	TRACE_CPUPROFILER_EVENT_SCOPE(FWaterBodySceneProxy::RebuildIndexBuffer);
 
 	TArray<uint32> Indices;
 	for (const FWaterBodySectionedMeshProxy& MeshSectionProxy : Sections)
@@ -199,7 +199,7 @@ FWaterBodySceneProxy::~FWaterBodySceneProxy()
 
 void FWaterBodySceneProxy::GetDynamicMeshElements(const TArray<const FSceneView*>& Views, const FSceneViewFamily& ViewFamily, uint32 VisibilityMap, FMeshElementCollector& Collector) const
 {
-	TRACE_CPUPROFILER_EVENT_SCOPE("FWaterBodySceneProxy::GetDynamicMeshElements");
+	TRACE_CPUPROFILER_EVENT_SCOPE(FWaterBodySceneProxy::GetDynamicMeshElements);
 
 	const bool bWithinWaterInfoPasses = CurrentWaterInfoPass != EWaterInfoPass::None;
 
@@ -233,7 +233,7 @@ void FWaterBodySceneProxy::GetDynamicMeshElements(const TArray<const FSceneView*
 	const bool bHasDilatedSections = WaterBodyInfoDilatedMesh.IndexBuffer.Indices.Num() > 0;
 	for (int32 ViewIndex = 0; ViewIndex < Views.Num(); ++ViewIndex)
 	{
-		TRACE_CPUPROFILER_EVENT_SCOPE("FWaterBodySceneProxy::AddSectionsForView");
+		TRACE_CPUPROFILER_EVENT_SCOPE(FWaterBodySceneProxy::AddSectionsForView);
 
 		if (CurrentWaterInfoPass == EWaterInfoPass::Color)
 		{
