Water: fixed an issue in non-editor builds where water bodies were not generating their geometry on load.

diff --git a/Engine/Plugins/Experimental/Water/Source/Runtime/Private/WaterBodyComponent.cpp b/Engine/Plugins/Experimental/Water/Source/Runtime/Private/WaterBodyComponent.cpp
--- a/Engine/Plugins/Experimental/Water/Source/Runtime/Private/WaterBodyComponent.cpp
+++ b/Engine/Plugins/Experimental/Water/Source/Runtime/Private/WaterBodyComponent.cpp
@@ -1298,6 +1298,12 @@ void UWaterBodyComponent::UpdateAll(bool bShapeOrPositionChanged)
 		UpdateWaterSpriteComponent();
 #endif
 	}
+
+	// #todo_water: this should be reconsidered along with the entire concept of non-dynamic water bodies now that evening is runtime-generated.
+	if (bShapeOrPositionChanged)
+	{
+		UpdateWaterBodyRenderData();
+	}
 }
 
 void UWaterBodyComponent::UpdateSplineComponent()
