Water: fixed issue where visualization components were being included in navigation on custom water body components.

diff --git a/Engine/Plugins/Experimental/Water/Source/Runtime/Private/WaterBodyCustomComponent.cpp b/Engine/Plugins/Experimental/Water/Source/Runtime/Private/WaterBodyCustomComponent.cpp
--- a/Engine/Plugins/Experimental/Water/Source/Runtime/Private/WaterBodyCustomComponent.cpp
+++ b/Engine/Plugins/Experimental/Water/Source/Runtime/Private/WaterBodyCustomComponent.cpp
@@ -76,9 +76,14 @@ void UWaterBodyCustomComponent::OnUpdateBody(bool bWithExclusionVolumes)
 	TInlineComponentArray<UPrimitiveComponent*> PrimitiveComponents;
 	OwnerActor->GetComponents(PrimitiveComponents);
 
-	// Make no assumptions for custom meshes.  Add all components with collision to the list of collision components
+	// Make no assumptions for custom water bodies: all (non-visualization) primitive components will be included in navigation
 	for (UPrimitiveComponent* Comp : PrimitiveComponents)
 	{
+		if (Comp->IsVisualizationComponent())
+		{
+			continue;
+		}
+
 		CopySharedNavigationSettingsToComponent(Comp);
 
 		Comp->SetMobility(Mobility);
