Fixed non-editor compile in WaterBodyCustomComponent.cpp

diff --git a/Engine/Plugins/Experimental/Water/Source/Runtime/Private/WaterBodyCustomComponent.cpp b/Engine/Plugins/Experimental/Water/Source/Runtime/Private/WaterBodyCustomComponent.cpp
--- a/Engine/Plugins/Experimental/Water/Source/Runtime/Private/WaterBodyCustomComponent.cpp
+++ b/Engine/Plugins/Experimental/Water/Source/Runtime/Private/WaterBodyCustomComponent.cpp
@@ -79,10 +79,12 @@ void UWaterBodyCustomComponent::OnUpdateBody(bool bWithExclusionVolumes)
 	// Make no assumptions for custom water bodies: all (non-visualization) primitive components will be included in navigation
 	for (UPrimitiveComponent* Comp : PrimitiveComponents)
 	{
+#if WITH_EDITORONLY_DATA
 		if (Comp->IsVisualizationComponent())
 		{
 			continue;
 		}
+#endif // WITH_EDITORONLY_DATA
 
 		CopySharedNavigationSettingsToComponent(Comp);
 
