* Removed dead code from water brush
* Removed old perf tags

#rnx

diff --git a/Engine/Plugins/Experimental/Water/Source/Editor/Private/WaterBrushManager.cpp b/Engine/Plugins/Experimental/Water/Source/Editor/Private/WaterBrushManager.cpp
--- a/Engine/Plugins/Experimental/Water/Source/Editor/Private/WaterBrushManager.cpp
+++ b/Engine/Plugins/Experimental/Water/Source/Editor/Private/WaterBrushManager.cpp
@@ -344,8 +344,6 @@ void AWaterBrushManager::ForceUpdate()
 	bKillCache = true;	
 	ClearCurveCache();
 	ALandscapeBlueprintBrushBase::RequestLandscapeUpdate();
-	// Regenerate the water depth velocity RT : 
-	MarkRenderTargetsDirty();	
 }
 
 void AWaterBrushManager::SingleJumpStep()
@@ -442,38 +440,6 @@ void AWaterBrushManager::CacheBrushDistanceField(const FBrushActorRenderContext&
 	BrushActorRenderContext.CacheContainer->Cache.CacheIsValid = true;
 }
 
-void AWaterBrushManager::BlueprintGetRenderTargets_Native(UTextureRenderTarget2D* InHeightRenderTarget, /*out*/ UTextureRenderTarget2D*& OutVelocityRenderTarget)
-{
-	// Use the last velocity RT rendered to as input and the other one as input :
-	FBrushRenderContext RenderContext;
-	RenderContext.VelocityRTIndex = LastRenderedVelocityRTIndex;
-
-	UTextureRenderTarget2D* VelocityRead = VelocityPingPongRead(RenderContext);
-	UTextureRenderTarget2D* VelocityWrite = VelocityPingPongWrite(RenderContext);
-
-	if ((VelocityRead == nullptr) || (VelocityWrite == nullptr))
-	{
-		UE_LOG(LogWaterEditor, Error, TEXT("Invalid Velocity Render Target for Water Brush. Aborting BlueprintGetRenderTargets."));
-		return;
-	}
-
-	if (FinalizeVelocityHeightMID == nullptr)
-	{
-		UE_LOG(LogWaterEditor, Error, TEXT("Invalid Finalize Velocity Height material for Water Brush. Aborting BlueprintGetRenderTargets."));
-		return;
-	}
-
-	UE_LOG(LogWaterEditor, Verbose, TEXT("Blueprint Get Render Targets called : finalizing VelocityHeightTexture (read: %s, write: %s)"), *VelocityRead->GetName(), *VelocityWrite->GetName());
-	FinalizeVelocityHeightMID->UMaterialInstanceDynamic::SetTextureParameterValue(FName(TEXT("CombinedVelocityAndHeight")), VelocityRead);
-	FinalizeVelocityHeightMID->UMaterialInstanceDynamic::SetTextureParameterValue(FName(TEXT("LandscapeHeight")), InHeightRenderTarget);
-
-	UKismetRenderingLibrary::ClearRenderTarget2D(this, VelocityWrite, FLinearColor(0.000000, 0.000000, 0.000000, 1.000000));
-	UKismetRenderingLibrary::DrawMaterialToRenderTarget(this, VelocityWrite, FinalizeVelocityHeightMID);
-
-	// TODO [jonathan.bard] : if SkipTextureGenerationProcess was implemented, this would be OutVelocityRenderTarget = SkipTextureGenerationProcess ? VelocityPingPongRead(); : nullptr;
-	OutVelocityRenderTarget = VelocityWrite;
-}
-
 void AWaterBrushManager::GetRenderDependencies(TSet<UObject*>& OutDependencies)
 {
 	Super::GetRenderDependencies(OutDependencies);
@@ -488,7 +454,6 @@ void AWaterBrushManager::GetRenderDependencies(TSet<UObject*>& OutDependencies)
 	AddDependencyIfValid(DrawCanvasMaterial, OutDependencies);
 	AddDependencyIfValid(CompositeWaterBodyTextureMaterial, OutDependencies);
 	AddDependencyIfValid(IslandFalloffMaterial, OutDependencies);
-	AddDependencyIfValid(FinalizeVelocityHeightMaterial, OutDependencies);
 	AddDependencyIfValid(JumpStepMaterial, OutDependencies);
 	AddDependencyIfValid(FindEdgesMaterial, OutDependencies);
 	AddDependencyIfValid(BlurEdgesMaterial, OutDependencies);
@@ -516,8 +481,6 @@ void AWaterBrushManager::UpdateTransform(const FTransform& Transform)
 
 		// The landscape transform has changed, let's re-draw everything (no need to request a landscape update because we're in the middle of one) :
 		bKillCache = true;
-		// Of course we also need to regenerate the water depth velocity RT : 
-		MarkRenderTargetsDirty();
 	}
 }
 
@@ -1319,7 +1282,6 @@ bool AWaterBrushManager::CreateMIDs()
 	WeightmapMID = FWaterUtils::GetOrCreateTransientMID(WeightmapMID, TEXT("WeightmapMID"), WeightmapMaterial);
 	DistanceFieldCacheMID = FWaterUtils::GetOrCreateTransientMID(DistanceFieldCacheMID, TEXT("DistanceFieldCacheMID"), DistanceFieldCacheMaterial);
 	CompositeWaterBodyTextureMID = FWaterUtils::GetOrCreateTransientMID(CompositeWaterBodyTextureMID, TEXT("CompositeWaterBodyTextureMID"), CompositeWaterBodyTextureMaterial);
-	FinalizeVelocityHeightMID = FWaterUtils::GetOrCreateTransientMID(FinalizeVelocityHeightMID, TEXT("FinalizeVelocityHeightMID"), FinalizeVelocityHeightMaterial);
 	DrawCanvasMID = FWaterUtils::GetOrCreateTransientMID(DrawCanvasMID, TEXT("DrawCanvasMID"), DrawCanvasMaterial);
 
 	if ((BrushAngleFalloffMID == nullptr)
@@ -1328,7 +1290,6 @@ bool AWaterBrushManager::CreateMIDs()
 		|| (WeightmapMID == nullptr)
 		|| (DistanceFieldCacheMID == nullptr)
 		|| (CompositeWaterBodyTextureMID == nullptr)
-		|| (FinalizeVelocityHeightMID == nullptr)
 		|| (DrawCanvasMID == nullptr))
 	{
 		UE_LOG(LogWaterEditor, Error, TEXT("Invalid water brush materials."));
@@ -1376,7 +1337,6 @@ void AWaterBrushManager::SetupDefaultMaterials()
 	DrawCanvasMaterial = WaterEditorSettings->GetDefaultDrawCanvasMaterial();
 	CompositeWaterBodyTextureMaterial = WaterEditorSettings->GetDefaultCompositeWaterBodyTextureMaterial();
 	IslandFalloffMaterial = WaterEditorSettings->GetDefaultBrushIslandFalloffMaterial();
-	FinalizeVelocityHeightMaterial = WaterEditorSettings->GetDefaultFinalizeVelocityHeightMaterial();
 
 	JumpStepMaterial = WaterEditorSettings->GetDefaultJumpFloodStepMaterial();
 	BlurEdgesMaterial = WaterEditorSettings->GetDefaultBlurEdgesMaterial();
@@ -1385,7 +1345,6 @@ void AWaterBrushManager::SetupDefaultMaterials()
 
 UTextureRenderTarget2D* AWaterBrushManager::Render_Native(bool InIsHeightmap, UTextureRenderTarget2D* InCombinedResult, FName const& InWeightmapLayerName)
 {
-	QUICK_SCOPE_CYCLE_COUNTER(AWaterBrushManager_Render_Native);
 	TRACE_CPUPROFILER_EVENT_SCOPE(AWaterBrushManager::Render_Native);
 
 	LandscapeRTRef = InCombinedResult;
@@ -1425,7 +1384,7 @@ UTextureRenderTarget2D* AWaterBrushManager::Render_Native(bool InIsHeightmap, UT
 
 	for (AWaterBody* WaterBody : WaterBodies)
 	{
-		QUICK_SCOPE_CYCLE_COUNTER(RenderWaterBody);
+		TRACE_CPUPROFILER_EVENT_SCOPE(RenderWaterBody);
 		FBrushActorRenderContext BrushActorRenderContext(WaterBody);
 		RenderBrushActorContext(BrushRenderContext, BrushActorRenderContext);
 	}
@@ -1436,7 +1395,7 @@ UTextureRenderTarget2D* AWaterBrushManager::Render_Native(bool InIsHeightmap, UT
 
 	for (AWaterBodyIsland* WaterBodyIsland : WaterbodyIslands)
 	{
-		QUICK_SCOPE_CYCLE_COUNTER(RenderWaterIsland);
+		TRACE_CPUPROFILER_EVENT_SCOPE(RenderWaterIsland);
 		FBrushActorRenderContext BrushActorRenderContext(WaterBodyIsland);
 		RenderBrushActorContext(BrushRenderContext, BrushActorRenderContext);
 	}
@@ -1444,9 +1403,6 @@ UTextureRenderTarget2D* AWaterBrushManager::Render_Native(bool InIsHeightmap, UT
 	// Render Completed
 	UTextureRenderTarget2D* ReturnRT = InIsHeightmap ? HeightPingPongRead(BrushRenderContext) : WeightPingPongRead(BrushRenderContext);
 
-	// Remember the last rendered velocity texture so that we can return the right version when the brush asks for it : 
-	LastRenderedVelocityRTIndex = BrushRenderContext.VelocityRTIndex;
-
 	bKillCache = false;
 
 	bNeedsForceUpdate = false;
diff --git a/Engine/Plugins/Experimental/Water/Source/Editor/Private/WaterBrushManager.h b/Engine/Plugins/Experimental/Water/Source/Editor/Private/WaterBrushManager.h
--- a/Engine/Plugins/Experimental/Water/Source/Editor/Private/WaterBrushManager.h
+++ b/Engine/Plugins/Experimental/Water/Source/Editor/Private/WaterBrushManager.h
@@ -102,9 +102,6 @@ public:
 	UPROPERTY(EditAnywhere, AdvancedDisplay, BlueprintReadWrite, meta = (Category = "Brush Materials"))
 	TObjectPtr<UMaterialInterface> IslandFalloffMaterial = nullptr;
 
-	UPROPERTY(EditAnywhere, AdvancedDisplay, BlueprintReadWrite, meta = (Category = "Brush Materials"))
-	TObjectPtr<UMaterialInterface> FinalizeVelocityHeightMaterial = nullptr;
-
 	UPROPERTY(EditAnywhere, AdvancedDisplay, BlueprintReadWrite, meta = (Category = "Brush Materials"))
 	TObjectPtr<UMaterialInterface> JumpStepMaterial = nullptr;
 
@@ -142,9 +139,6 @@ public:
 
 	UPROPERTY(EditInstanceOnly, AdvancedDisplay, BlueprintReadWrite, Transient, meta = (Category = "Debug MIDs"))
 	TObjectPtr<UMaterialInstanceDynamic> IslandFalloffMID = nullptr;
-
-	UPROPERTY(EditInstanceOnly, AdvancedDisplay, BlueprintReadWrite, Transient, meta = (Category = "Debug MIDs"))
-	TObjectPtr<UMaterialInstanceDynamic> FinalizeVelocityHeightMID = nullptr;
 	// MIDs End
 
 	UPROPERTY(VisibleAnywhere, BlueprintReadWrite, Transient, meta=(Category="Debug"))
@@ -192,6 +186,16 @@ public:
 	UPROPERTY()
 	bool bNeedsForceUpdate = false;
 
+#if WITH_EDITORONLY_DATA
+	UE_DEPRECATED(5.1, "This material is now useless, the WaterVelocityTexture is now regenerated at runtime (WaterInfoTexture in AWaterZone).")
+	UPROPERTY()
+	TObjectPtr<UMaterialInterface> FinalizeVelocityHeightMaterial_DEPRECATED = nullptr;
+
+	UE_DEPRECATED(5.1, "This material instance is now useless, the WaterVelocityTexture is now regenerated at runtime (WaterInfoTexture in AWaterZone).")
+	UPROPERTY()
+	TObjectPtr<UMaterialInstanceDynamic> FinalizeVelocityHeightMID_DEPRECATED = nullptr;
+#endif // WITH_EDITORONLY_DATA
+
 	AWaterBrushManager(const FObjectInitializer& ObjectInitializer = FObjectInitializer::Get());
 
 	virtual void Serialize(FArchive& Ar) override;
@@ -205,8 +209,6 @@ public:
 	virtual void Initialize_Native(FTransform const& InLandscapeTransform, FIntPoint const& InLandscapeSize, FIntPoint const& InLandscapeRenderTargetSize) override;
 	virtual UTextureRenderTarget2D*  Render_Native(bool InIsHeightmap, UTextureRenderTarget2D* InCombinedResult, FName const& InWeightmapLayerName) override;
 		
-	virtual void BlueprintGetRenderTargets_Native(UTextureRenderTarget2D* InHeightRenderTarget, /*out*/ UTextureRenderTarget2D*& OutVelocityRenderTarget) override;
-		
 	virtual void GetRenderDependencies(TSet<UObject*>& OutDependencies) override;
 
 	// Debug Buttons
@@ -320,7 +322,6 @@ private:
 
 private:
 	bool bKillCache = false;
-	int32 LastRenderedVelocityRTIndex = 0;
 	// HACK [jonathan.bard] : shouldn't be needed anymore once deprecation is done : 
 	FDelegateHandle OnWorldPostInitHandle;
 	FDelegateHandle OnLevelAddedToWorldHandle;
diff --git a/Engine/Plugins/Experimental/Water/Source/Editor/Private/WaterEditorSettings.cpp b/Engine/Plugins/Experimental/Water/Source/Editor/Private/WaterEditorSettings.cpp
--- a/Engine/Plugins/Experimental/Water/Source/Editor/Private/WaterEditorSettings.cpp
+++ b/Engine/Plugins/Experimental/Water/Source/Editor/Private/WaterEditorSettings.cpp
@@ -140,7 +140,6 @@ UWaterEditorSettings::UWaterEditorSettings()
 	, DefaultBrushWeightmapMaterial(FSoftObjectPath(TEXT("/Water/Materials/Brushes/MeshBrush_Weightmap.MeshBrush_Weightmap")))
 	, DefaultCacheDistanceFieldCacheMaterial(FSoftObjectPath(TEXT("/Landmass/Landscape/BlueprintBrushes/Materials/CacheDistanceField.CacheDistanceField")))
 	, DefaultCompositeWaterBodyTextureMaterial(FSoftObjectPath(TEXT("/Water/Materials/Brushes/CompositeWaterBodyTexture.CompositeWaterBodyTexture")))
-	, DefaultFinalizeVelocityHeightMaterial(FSoftObjectPath(TEXT("/Water/Materials/Brushes/FinalizeWaterTexture.FinalizeWaterTexture")))
 	, DefaultJumpFloodStepMaterial(FSoftObjectPath(TEXT("/Landmass/DistanceFields/Materials/JumpFloodStep.JumpFloodStep")))
 	, DefaultBlurEdgesMaterial(FSoftObjectPath(TEXT("/Landmass/DistanceFields/Materials/BlurEdges.BlurEdges")))
 	, DefaultFindEdgesMaterial(FSoftObjectPath(TEXT("/Landmass/DistanceFields/Materials/DetectEdges.DetectEdges")))
@@ -200,11 +199,6 @@ UMaterialInterface* UWaterEditorSettings::GetDefaultCompositeWaterBodyTextureMat
 	return DefaultCompositeWaterBodyTextureMaterial.LoadSynchronous();
 }
 
-UMaterialInterface* UWaterEditorSettings::GetDefaultFinalizeVelocityHeightMaterial() const
-{
-	return DefaultFinalizeVelocityHeightMaterial.LoadSynchronous();
-}
-
 UMaterialInterface* UWaterEditorSettings::GetDefaultJumpFloodStepMaterial() const
 {
 	return DefaultJumpFloodStepMaterial.LoadSynchronous();
diff --git a/Engine/Plugins/Experimental/Water/Source/Editor/Private/WaterLandscapeBrush.cpp b/Engine/Plugins/Experimental/Water/Source/Editor/Private/WaterLandscapeBrush.cpp
--- a/Engine/Plugins/Experimental/Water/Source/Editor/Private/WaterLandscapeBrush.cpp
+++ b/Engine/Plugins/Experimental/Water/Source/Editor/Private/WaterLandscapeBrush.cpp
@@ -198,7 +198,6 @@ void AWaterLandscapeBrush::OnActorChanged(AActor* Actor, bool bWeightmapSettings
 	if (bAffectsLandscape || bForceUpdateBrush)
 	{
 		RequestLandscapeUpdate();
-		MarkRenderTargetsDirty();
 	}
 
 	if (bRebuildWaterMesh)
@@ -210,11 +209,6 @@ void AWaterLandscapeBrush::OnActorChanged(AActor* Actor, bool bWeightmapSettings
 	}
 }
 
-void AWaterLandscapeBrush::MarkRenderTargetsDirty()
-{
-	bRenderTargetsDirty = true;
-}
-
 void AWaterLandscapeBrush::OnWaterBrushActorChanged(const IWaterBrushActorInterface::FWaterBrushActorChangedEventParams& InParams)
 {
 	AActor* Actor = CastChecked<AActor>(InParams.WaterBrushActor);
@@ -229,8 +223,6 @@ void AWaterLandscapeBrush::OnActorsAffectingLandscapeChanged()
 	{
 		WaterSubsystem->MarkAllWaterZonesForRebuild();
 	}
-
-	MarkRenderTargetsDirty();
 }
 
 bool AWaterLandscapeBrush::IsActorAffectingLandscape(AActor* Actor) const
@@ -445,7 +437,7 @@ void AWaterLandscapeBrush::ClearActorCache(AActor* InActor)
 
 void AWaterLandscapeBrush::BlueprintGetRenderTargets_Implementation(UTextureRenderTarget2D* InHeightRenderTarget, UTextureRenderTarget2D*& OutVelocityRenderTarget)
 {
-	BlueprintGetRenderTargets_Native(InHeightRenderTarget, OutVelocityRenderTarget);
+	// Deprecated
 }
 
 void AWaterLandscapeBrush::SetTargetLandscape(ALandscape* InTargetLandscape)
@@ -475,7 +467,6 @@ void AWaterLandscapeBrush::SetTargetLandscape(ALandscape* InTargetLandscape)
 
 void AWaterLandscapeBrush::OnFullHeightmapRenderDone(UTextureRenderTarget2D* InHeightmapRenderTarget)
 {
-	bRenderTargetsDirty = false;
 	// #todo_water [roey]: This needs to be changed when the WaterZone can maintain it's own list of "ground actors" so that we don't needlessly update all water zones.
 	if (UWaterSubsystem* WaterSubsystem = UWaterSubsystem::GetWaterSubsystem(GetWorld()))
 	{
@@ -519,12 +510,12 @@ void AWaterLandscapeBrush::ForceUpdate()
 
 void AWaterLandscapeBrush::BlueprintOnRenderTargetTexturesUpdated_Implementation(UTexture2D* VelocityTexture)
 {
-	BlueprintOnRenderTargetTexturesUpdated_Native(VelocityTexture);
+	// Deprecated
 }
 
 void AWaterLandscapeBrush::ForceWaterTextureUpdate()
 {
-	MarkRenderTargetsDirty();
+	// Deprecated
 }
 
 #if WITH_EDITOR
diff --git a/Engine/Plugins/Experimental/Water/Source/Editor/Private/WaterLandscapeBrush.h b/Engine/Plugins/Experimental/Water/Source/Editor/Private/WaterLandscapeBrush.h
--- a/Engine/Plugins/Experimental/Water/Source/Editor/Private/WaterLandscapeBrush.h
+++ b/Engine/Plugins/Experimental/Water/Source/Editor/Private/WaterLandscapeBrush.h
@@ -67,15 +67,18 @@ public:
 	UFUNCTION(BlueprintCallable, Category = "Cache")
 	void ClearActorCache(AActor* InActor);
 
-	UFUNCTION(BlueprintNativeEvent, Category = "Cache", meta = (CallInEditor = "true"))
+	UFUNCTION(BlueprintNativeEvent, Category = "Cache", meta = (CallInEditor = "true", DeprecatedFunction, DeprecationMessage = "This event isn't called anymore, the WaterVelocityTexture is now regenerated at runtime (WaterInfoTexture in AWaterZone)."))
 	void BlueprintGetRenderTargets(UTextureRenderTarget2D* InHeightRenderTarget, UTextureRenderTarget2D*& OutVelocityRenderTarget);
+	UE_DEPRECATED(5.1, "This function isn't called anymore, the WaterVelocityTexture is now regenerated at runtime (WaterInfoTexture in AWaterZone)")
 	virtual void BlueprintGetRenderTargets_Native(UTextureRenderTarget2D* InHeightRenderTarget, UTextureRenderTarget2D*& OutVelocityRenderTarget) {}
 
-	UFUNCTION(BlueprintNativeEvent, Category = "Cache", meta = (CallInEditor = "true"))
+	UFUNCTION(BlueprintNativeEvent, Category = "Cache", meta = (CallInEditor = "true", DeprecatedFunction, DeprecationMessage = "This event isn't called anymore, the WaterVelocityTexture is now regenerated at runtime (WaterInfoTexture in AWaterZone)."))
 	void BlueprintOnRenderTargetTexturesUpdated(UTexture2D* VelocityTexture);
+	UE_DEPRECATED(5.1, "This function isn't called anymore, the WaterVelocityTexture is now regenerated at runtime (WaterInfoTexture in AWaterZone)")
 	virtual void BlueprintOnRenderTargetTexturesUpdated_Native(UTexture2D* VelocityTexture) {}
 
-	UFUNCTION(BlueprintCallable, Category = "Cache")
+	UE_DEPRECATED(5.1, "This function is now useless, the WaterVelocityTexture is now regenerated at runtime (WaterInfoTexture in AWaterZone).")
+	UFUNCTION(BlueprintCallable, Category = "Cache", meta = (CallInEditor = "true", DeprecatedFunction, DeprecationMessage = "This function is now useless, the WaterVelocityTexture is now regenerated at runtime (WaterInfoTexture in AWaterZone)."))
 	void ForceWaterTextureUpdate();
 
 	void SetTargetLandscape(ALandscape* InTargetLandscape);
@@ -103,14 +106,6 @@ public:
 	virtual bool CanEditChange(const FProperty* InProperty) const override;
 #endif // WITH_EDITOR
 
-#if WITH_EDITORONLY_DATA
-	UPROPERTY(Transient)
-	TObjectPtr<UBillboardComponent> ActorIcon;
-#endif // WITH_EDITORONLY_DATA
-
-protected:
-	void MarkRenderTargetsDirty();
-
 private:
 	template<class T>
 	friend class FGetActorsOfType;
@@ -129,6 +124,13 @@ private:
 	void OnLevelActorAdded(AActor* InActor);
 	void OnLevelActorRemoved(AActor* InActor);
 
+public:
+#if WITH_EDITORONLY_DATA
+	UPROPERTY(Transient)
+	TObjectPtr<UBillboardComponent> ActorIcon;
+#endif // WITH_EDITORONLY_DATA
+
+private:
 	TArray<TWeakInterfacePtr<IWaterBrushActorInterface>> ActorsAffectingLandscape;
 	FDelegateHandle OnWorldPostInitHandle;
 	FDelegateHandle OnLevelAddedToWorldHandle;
@@ -143,6 +145,4 @@ private:
 
 	UPROPERTY(Transient, DuplicateTransient, VisibleAnywhere, AdvancedDisplay, meta = (Category = "Debug"))
 	TMap<TWeakObjectPtr<AActor>, TObjectPtr<UObject>> Cache;
-
-	bool bRenderTargetsDirty = false;
 };
\ No newline at end of file
diff --git a/Engine/Plugins/Experimental/Water/Source/Editor/Public/WaterEditorSettings.h b/Engine/Plugins/Experimental/Water/Source/Editor/Public/WaterEditorSettings.h
--- a/Engine/Plugins/Experimental/Water/Source/Editor/Public/WaterEditorSettings.h
+++ b/Engine/Plugins/Experimental/Water/Source/Editor/Public/WaterEditorSettings.h
@@ -202,8 +202,10 @@ public:
 	UMaterialInterface* GetDefaultCompositeWaterBodyTextureMaterial() const;
 	FSoftObjectPath GetDefaultCompositeWaterBodyTextureMaterialPath() const { return DefaultCompositeWaterBodyTextureMaterial.ToSoftObjectPath(); }
 
-	UMaterialInterface* GetDefaultFinalizeVelocityHeightMaterial() const;
-	FSoftObjectPath GetDefaultFinalizeVelocityHeightMaterialPath() const { return DefaultFinalizeVelocityHeightMaterial.ToSoftObjectPath(); }
+	UE_DEPRECATED(5.1, "This material is not needed anymore : the WaterVelocityTexture is now regenerated at runtime (WaterInfoTexture in AWaterZone)")
+	UMaterialInterface* GetDefaultFinalizeVelocityHeightMaterial() const { return nullptr; }
+	UE_DEPRECATED(5.1, "This material is not needed anymore : the WaterVelocityTexture is now regenerated at runtime (WaterInfoTexture in AWaterZone)")
+	FSoftObjectPath GetDefaultFinalizeVelocityHeightMaterialPath() const { return FSoftObjectPath(); }
 
 	UMaterialInterface* GetDefaultJumpFloodStepMaterial() const;
 	FSoftObjectPath GetDefaultJumpFloodStepMaterialPath() const { return DefaultJumpFloodStepMaterial.ToSoftObjectPath(); }
@@ -283,9 +285,6 @@ private:
 	UPROPERTY(EditAnywhere, config, Category = Brush)
 	TSoftObjectPtr<UMaterialInterface> DefaultCompositeWaterBodyTextureMaterial;
 
-	UPROPERTY(EditAnywhere, config, Category = Brush)
-	TSoftObjectPtr<UMaterialInterface> DefaultFinalizeVelocityHeightMaterial;
-
 	UPROPERTY(EditAnywhere, config, Category = Brush)
 	TSoftObjectPtr<UMaterialInterface> DefaultJumpFloodStepMaterial;
 
