Fixed merge problem

diff --git a/Engine/Plugins/Experimental/Water/Source/Editor/Private/WaterLandscapeBrush.cpp b/Engine/Plugins/Experimental/Water/Source/Editor/Private/WaterLandscapeBrush.cpp
--- a/Engine/Plugins/Experimental/Water/Source/Editor/Private/WaterLandscapeBrush.cpp
+++ b/Engine/Plugins/Experimental/Water/Source/Editor/Private/WaterLandscapeBrush.cpp
@@ -27,7 +27,7 @@
 #include "Modules/ModuleManager.h"
 #include "WaterModule.h"
 #include "WaterEditorSettings.h"
-"LandscapeModule.h"
+#include "LandscapeModule.h"
 #include "LandscapeEditorServices.h"
 
 #define LOCTEXT_NAMESPACE "WaterLandscapeBrush"
