- Precache PSOs for WaterMeshComponents
- Fix for template specialization declares of vertex factory static types and templates

diff --git a/Engine/Plugins/Experimental/Water/Source/Runtime/Private/WaterMeshComponent.cpp b/Engine/Plugins/Experimental/Water/Source/Runtime/Private/WaterMeshComponent.cpp
--- a/Engine/Plugins/Experimental/Water/Source/Runtime/Private/WaterMeshComponent.cpp
+++ b/Engine/Plugins/Experimental/Water/Source/Runtime/Private/WaterMeshComponent.cpp
@@ -3,6 +3,7 @@
 #include "WaterMeshComponent.h"
 #include "WaterBodyActor.h"
 #include "WaterSplineComponent.h"
+#include "WaterVertexFactory.h"
 #include "EngineUtils.h"
 #include "DrawDebugHelpers.h"
 #include "PhysicsEngine/AggregateGeom.h"
@@ -74,6 +75,29 @@ void UWaterMeshComponent::PostLoad()
 	Super::PostLoad();
 }
 
+void UWaterMeshComponent::PrecachePSOs()
+{
+	if (!FApp::CanEverRender())
+	{
+		return;
+	}
+
+	FPSOPrecacheParams PrecachePSOParams;
+	SetupPrecachePSOParams(PrecachePSOParams);
+
+	if (FarDistanceMaterial)
+	{
+		FarDistanceMaterial->PrecachePSOs(&TWaterVertexFactory</*bWithWaterSelectionSupport = */ false>::StaticType, PrecachePSOParams);
+	}
+	for (UMaterialInterface* MaterialInterface : UsedMaterials)
+	{
+		if (MaterialInterface)
+		{
+			MaterialInterface->PrecachePSOs(&TWaterVertexFactory</*bWithWaterSelectionSupport = */ false>::StaticType, PrecachePSOParams);
+		}
+	}
+}
+
 void UWaterMeshComponent::PostInitProperties()
 {
 	Super::PostInitProperties();
@@ -567,6 +591,7 @@ void UWaterMeshComponent::Update()
 		LODScaleBiasScalability = NewLODScaleBias;
 		const float LODCountBiasFactor = FMath::Pow(2.0f, (float)LODCountBiasScalability);
 		RebuildWaterMesh(TileSize / LODCountBiasFactor, FIntPoint(FMath::CeilToInt(ExtentInTiles.X * LODCountBiasFactor), FMath::CeilToInt(ExtentInTiles.Y * LODCountBiasFactor)));
+		PrecachePSOs();
 		bNeedsRebuild = false;
 	}
 }
diff --git a/Engine/Plugins/Experimental/Water/Source/Runtime/Public/WaterMeshComponent.h b/Engine/Plugins/Experimental/Water/Source/Runtime/Public/WaterMeshComponent.h
--- a/Engine/Plugins/Experimental/Water/Source/Runtime/Public/WaterMeshComponent.h
+++ b/Engine/Plugins/Experimental/Water/Source/Runtime/Public/WaterMeshComponent.h
@@ -41,6 +41,8 @@ public:
 #endif // WITH_EDITOR
 	//~ End UPrimitiveComponent Interface
 
+	virtual void PrecachePSOs() override;
+
 	void Update();
 
 	/** Use this instead of GetMaterialRelevance, since this one will go over all materials from all tiles */
diff --git a/Engine/Plugins/Experimental/Water/Source/Runtime/Public/WaterVertexFactory.h b/Engine/Plugins/Experimental/Water/Source/Runtime/Public/WaterVertexFactory.h
--- a/Engine/Plugins/Experimental/Water/Source/Runtime/Public/WaterVertexFactory.h
+++ b/Engine/Plugins/Experimental/Water/Source/Runtime/Public/WaterVertexFactory.h
@@ -216,6 +216,11 @@ private:
 	const float LODScale = 0.0f;
 };
 
+DECLARE_TEMPLATE_SPECIALIZATION_VERTEX_FACTORY_TYPE(template<>, TWaterVertexFactory</*bWithWaterSelectionSupport = */ false>);
+
+#if WITH_WATER_SELECTION_SUPPORT
+DECLARE_TEMPLATE_SPECIALIZATION_VERTEX_FACTORY_TYPE(template<>, TWaterVertexFactory</*bWithWaterSelectionSupport = */ true>);
+#endif // WITH_WATER_SELECTION_SUPPORT
 
 /**
  * Water user data provided to FMeshBatchElement(s)
diff --git a/Engine/Source/Runtime/RenderCore/Public/VertexFactory.h b/Engine/Source/Runtime/RenderCore/Public/VertexFactory.h
--- a/Engine/Source/Runtime/RenderCore/Public/VertexFactory.h
+++ b/Engine/Source/Runtime/RenderCore/Public/VertexFactory.h
@@ -548,6 +548,12 @@ extern RENDERCORE_API FVertexFactoryType* FindVertexFactoryType(const FHashedNam
 	FactoryClass::ModifyCompilationEnvironment, \
 	FactoryClass::ValidateCompiledResult
 
+ /**
+  * A macro for declaring a template specialization vertex factory type so it can be used everywhere
+  */
+#define DECLARE_TEMPLATE_SPECIALIZATION_VERTEX_FACTORY_TYPE(TemplatePrefix, FactoryClass) \
+	PREPROCESSOR_REMOVE_OPTIONAL_PARENS(TemplatePrefix) FVertexFactoryType FactoryClass::StaticType;
+
 /**
  * A macro for implementing the static vertex factory type object, and specifying parameters used by the type.
  */
