Precache PSOs for WaterMeshComponents

diff --git a/Engine/Plugins/Experimental/Water/Source/Runtime/Private/WaterMeshComponent.cpp b/Engine/Plugins/Experimental/Water/Source/Runtime/Private/WaterMeshComponent.cpp
--- a/Engine/Plugins/Experimental/Water/Source/Runtime/Private/WaterMeshComponent.cpp
+++ b/Engine/Plugins/Experimental/Water/Source/Runtime/Private/WaterMeshComponent.cpp
@@ -74,6 +74,30 @@ void UWaterMeshComponent::PostLoad()
 	Super::PostLoad();
 }
 
+void UWaterMeshComponent::PrecachePSOs()
+{
+	if (!FApp::CanEverRender())
+	{
+		return;
+	}
+
+	FPSOPrecacheParams PrecachePSOParams;
+	SetupPrecachePSOParams(PrecachePSOParams);
+
+	const FVertexFactoryType* WaterVertexFactoryType = GetWaterVertexFactoryType(/*bWithWaterSelectionSupport = */ false);
+	if (FarDistanceMaterial)
+	{
+		FarDistanceMaterial->PrecachePSOs(WaterVertexFactoryType, PrecachePSOParams);
+	}
+	for (UMaterialInterface* MaterialInterface : UsedMaterials)
+	{
+		if (MaterialInterface)
+		{
+			MaterialInterface->PrecachePSOs(WaterVertexFactoryType, PrecachePSOParams);
+		}
+	}
+}
+
 void UWaterMeshComponent::PostInitProperties()
 {
 	Super::PostInitProperties();
@@ -567,6 +591,7 @@ void UWaterMeshComponent::Update()
 		LODScaleBiasScalability = NewLODScaleBias;
 		const float LODCountBiasFactor = FMath::Pow(2.0f, (float)LODCountBiasScalability);
 		RebuildWaterMesh(TileSize / LODCountBiasFactor, FIntPoint(FMath::CeilToInt(ExtentInTiles.X * LODCountBiasFactor), FMath::CeilToInt(ExtentInTiles.Y * LODCountBiasFactor)));
+		PrecachePSOs();
 		bNeedsRebuild = false;
 	}
 }
diff --git a/Engine/Plugins/Experimental/Water/Source/Runtime/Private/WaterVertexFactory.cpp b/Engine/Plugins/Experimental/Water/Source/Runtime/Private/WaterVertexFactory.cpp
--- a/Engine/Plugins/Experimental/Water/Source/Runtime/Private/WaterVertexFactory.cpp
+++ b/Engine/Plugins/Experimental/Water/Source/Runtime/Private/WaterVertexFactory.cpp
@@ -116,3 +116,18 @@ IMPLEMENT_TEMPLATE_VERTEX_FACTORY_TYPE(template<>, TWaterVertexFactory</*bWithWa
 );
 
 #endif // WITH_WATER_SELECTION_SUPPORT
+
+const FVertexFactoryType* GetWaterVertexFactoryType(bool bWithWaterSelectionSupport)
+{
+#if WITH_WATER_SELECTION_SUPPORT
+	if (bWithWaterSelectionSupport)
+	{
+		return &TWaterVertexFactory</*bWithWaterSelectionSupport = */ true>::StaticType;
+	}
+	else
+#endif // WITH_WATER_SELECTION_SUPPORT
+	{
+		check(!bWithWaterSelectionSupport);
+		return &TWaterVertexFactory</*bWithWaterSelectionSupport = */ false>::StaticType;
+	}
+}
diff --git a/Engine/Plugins/Experimental/Water/Source/Runtime/Public/WaterMeshComponent.h b/Engine/Plugins/Experimental/Water/Source/Runtime/Public/WaterMeshComponent.h
--- a/Engine/Plugins/Experimental/Water/Source/Runtime/Public/WaterMeshComponent.h
+++ b/Engine/Plugins/Experimental/Water/Source/Runtime/Public/WaterMeshComponent.h
@@ -41,6 +41,8 @@ public:
 #endif // WITH_EDITOR
 	//~ End UPrimitiveComponent Interface
 
+	virtual void PrecachePSOs() override;
+
 	void Update();
 
 	/** Use this instead of GetMaterialRelevance, since this one will go over all materials from all tiles */
diff --git a/Engine/Plugins/Experimental/Water/Source/Runtime/Public/WaterVertexFactory.h b/Engine/Plugins/Experimental/Water/Source/Runtime/Public/WaterVertexFactory.h
--- a/Engine/Plugins/Experimental/Water/Source/Runtime/Public/WaterVertexFactory.h
+++ b/Engine/Plugins/Experimental/Water/Source/Runtime/Public/WaterVertexFactory.h
@@ -216,6 +216,8 @@ private:
 	const float LODScale = 0.0f;
 };
 
+extern const FVertexFactoryType* GetWaterVertexFactoryType(bool bWithWaterSelectionSupport);
+
 
 /**
  * Water user data provided to FMeshBatchElement(s)
