Fix for crash when copy/pasting landscape and water bodies across levels.

When copy/pasting landscape and water actors simultaneously, the newly pasted landscape may not yet be fully initialized at the point where the water actors try to populate the level landscapes with Water layers and brushes.  When this occurs, we now ignore the pasted landscape in that operation.  This may mean that the auto-population won't work in certain circumstances when copy/pasting landscapes and water bodies simultaneously.

This doesn't fix all of the issues:
   You may be left with water zone pointers that belong the wrong package, making the level un-saveable until fixed.
   Roey will be looking to fix that separately by making those pointers soft.

diff --git a/Engine/Plugins/Experimental/Water/Source/Editor/Private/WaterEditorModule.cpp b/Engine/Plugins/Experimental/Water/Source/Editor/Private/WaterEditorModule.cpp
--- a/Engine/Plugins/Experimental/Water/Source/Editor/Private/WaterEditorModule.cpp
+++ b/Engine/Plugins/Experimental/Water/Source/Editor/Private/WaterEditorModule.cpp
@@ -175,28 +175,37 @@ void FWaterEditorModule::OnLevelActorAddedToWorld(AActor* Actor)
 		const bool bIncludeChildActors = false;
 		const FBox ActorBounds = Actor->GetComponentsBoundingBox(bNonColliding, bIncludeChildActors);
 
+		ALandscape* FallbackLandscape = nullptr;
 		for (ALandscape* Landscape : TActorRange<ALandscape>(ActorWorld))
 		{
-			const FBox LandscapeBounds = Landscape->GetCompleteBounds();
-			if (LandscapeBounds.Intersect(ActorBounds))
+			// This function is called while copy-pasting landscapes, before the landscape properties have been imported
+			// in which case the Guid will be invalid, and trying to GetCompleteBounds will assert.
+			// Note that skipping this logic means that we won't automatically create Water layers / Water brushes,
+			// and any automatically created Water Zone may not be initialized to the full bounds of the landscapes
+			if (Landscape->GetLandscapeGuid().IsValid())
 			{
-				FoundLandscapes.Add(Landscape);
-				// Make sure the water zone's bounds is large enough to fit all landscapes that intersect with this water body :
-				WaterZoneBounds += LandscapeBounds;
-				bFoundIntersectingLandscape = true;
+				if (FallbackLandscape == nullptr)
+				{
+					FallbackLandscape = Landscape;
+				}
+
+				const FBox LandscapeBounds = Landscape->GetCompleteBounds();
+				if (LandscapeBounds.Intersect(ActorBounds))
+				{
+					FoundLandscapes.Add(Landscape);
+					// Make sure the water zone's bounds is large enough to fit all landscapes that intersect with this water body :
+					WaterZoneBounds += LandscapeBounds;
+					bFoundIntersectingLandscape = true;
+				}
 			}
 		}
 
-		// If no intersectiong landscape was found, use the first one available, if any :
-		if (!bFoundIntersectingLandscape)
+		// If no intersecting landscape was found, use the first valid one as a fallback
+		if (!bFoundIntersectingLandscape && FallbackLandscape != nullptr)
 		{
-			const TActorIterator<ALandscape> It(ActorWorld);
-			if (ALandscape* Landscape = It ? *It : nullptr)
-			{
-				FoundLandscapes.Add(Landscape);
-				const FBox LandscapeBounds = Landscape->GetCompleteBounds();
-				WaterZoneBounds += LandscapeBounds;
-			}
+			FoundLandscapes.Add(FallbackLandscape);
+			const FBox LandscapeBounds = FallbackLandscape->GetCompleteBounds();
+			WaterZoneBounds += LandscapeBounds;
 		}
 
 		// Automatically setup landscape-affecting features (water brush) if needed : 
diff --git a/Engine/Source/Runtime/Landscape/Private/LandscapeEdit.cpp b/Engine/Source/Runtime/Landscape/Private/LandscapeEdit.cpp
--- a/Engine/Source/Runtime/Landscape/Private/LandscapeEdit.cpp
+++ b/Engine/Source/Runtime/Landscape/Private/LandscapeEdit.cpp
@@ -4248,6 +4248,15 @@ void ALandscape::PostEditImport()
 		}
 	}
 
+	// We need to reparent brushes that may have been part of the copy/pasted actors :
+	for (FLandscapeLayer& Layer : LandscapeLayers)
+	{
+		for (FLandscapeLayerBrush& Brush : Layer.Brushes)
+		{
+			Brush.SetOwner(this);
+		}
+	}
+
 	// Some edit layers could be affected by BP brushes, which might need to be updated when the landscape is transformed :
 	RequestLayersContentUpdate(ELandscapeLayerUpdateMode::Update_All);
 
