Fix static analysis warning in WaterEditorModule.cpp (MSVC)

diff --git a/Engine/Plugins/Experimental/Water/Source/Editor/Private/WaterEditorModule.cpp b/Engine/Plugins/Experimental/Water/Source/Editor/Private/WaterEditorModule.cpp
--- a/Engine/Plugins/Experimental/Water/Source/Editor/Private/WaterEditorModule.cpp
+++ b/Engine/Plugins/Experimental/Water/Source/Editor/Private/WaterEditorModule.cpp
@@ -178,6 +178,8 @@ void FWaterEditorModule::OnLevelActorAddedToWorld(AActor* Actor)
 		ALandscape* FallbackLandscape = nullptr;
 		for (ALandscape* Landscape : TActorRange<ALandscape>(ActorWorld))
 		{
+			check(Landscape);
+
 			// This function is called while copy-pasting landscapes, before the landscape properties have been imported
 			// in which case the Guid will be invalid, and trying to GetCompleteBounds will assert.
 			// Note that skipping this logic means that we won't automatically create Water layers / Water brushes,
