Quick speculative workaround for crash when cooking content workers.

diff --git a/Engine/Plugins/Experimental/Water/Source/Editor/Private/WaterLandscapeBrush.cpp b/Engine/Plugins/Experimental/Water/Source/Editor/Private/WaterLandscapeBrush.cpp
--- a/Engine/Plugins/Experimental/Water/Source/Editor/Private/WaterLandscapeBrush.cpp
+++ b/Engine/Plugins/Experimental/Water/Source/Editor/Private/WaterLandscapeBrush.cpp
@@ -281,8 +281,22 @@ void AWaterLandscapeBrush::PostInitProperties()
 		{
 			if (World->PersistentLevel)
 			{
-				OnLoadedActorAddedToLevelEventHandle = World->PersistentLevel->OnLoadedActorAddedToLevelEvent.AddLambda([this](AActor& InActor) { OnLevelActorAdded(&InActor); });
-				OnLoadedActorRemovedFromLevelEventHandle = World->PersistentLevel->OnLoadedActorRemovedFromLevelEvent.AddLambda([this](AActor& InActor) { OnLevelActorRemoved(&InActor); });
+				TWeakObjectPtr<AWaterLandscapeBrush> WeakThis = this;
+
+				OnLoadedActorAddedToLevelEventHandle = World->PersistentLevel->OnLoadedActorAddedToLevelEvent.AddLambda([WeakThis](AActor& InActor) { 
+
+					if (AWaterLandscapeBrush* StrongThis = WeakThis.Get())
+					{
+						StrongThis->OnLevelActorAdded(&InActor);
+					}
+				});
+
+				OnLoadedActorRemovedFromLevelEventHandle = World->PersistentLevel->OnLoadedActorRemovedFromLevelEvent.AddLambda([WeakThis](AActor& InActor) { 
+					if (AWaterLandscapeBrush* StrongThis = WeakThis.Get())
+					{
+						StrongThis->OnLevelActorRemoved(&InActor);
+					}
+				});
 			}
 		}
 #endif // WITH_EDITOR
