Water landscape brush fix : moved level callbacks registration/unregistration to OnPostRegisterAllComponents/PostUnregisterAllComponents, where the world / level should always be valid so that we don't end up with lingering callbacks in the level

#rnx

diff --git a/Engine/Plugins/Experimental/Water/Source/Editor/Private/WaterLandscapeBrush.cpp b/Engine/Plugins/Experimental/Water/Source/Editor/Private/WaterLandscapeBrush.cpp
--- a/Engine/Plugins/Experimental/Water/Source/Editor/Private/WaterLandscapeBrush.cpp
+++ b/Engine/Plugins/Experimental/Water/Source/Editor/Private/WaterLandscapeBrush.cpp
@@ -241,7 +241,7 @@ void AWaterLandscapeBrush::PostInitProperties()
 {
 	Super::PostInitProperties();
 
-	if (!HasAnyFlags(RF_ClassDefaultObject | RF_ArchetypeObject | RF_Transient))
+	if (!IsTemplate())
 	{
 		OnWorldPostInitHandle = FWorldDelegates::OnPostWorldInitialization.AddLambda([this](UWorld* World, const UWorld::InitializationValues IVS)
 		{
@@ -277,37 +277,11 @@ void AWaterLandscapeBrush::PostInitProperties()
 		OnLevelActorAddedHandle = GEngine->OnLevelActorAdded().AddUObject(this, &AWaterLandscapeBrush::OnLevelActorAdded);
 		OnLevelActorDeletedHandle = GEngine->OnLevelActorDeleted().AddUObject(this, &AWaterLandscapeBrush::OnLevelActorRemoved);
 
-#if WITH_EDITOR
-		// in world partition, actors don't belong to levels and the on loaded/removed callbacks are different :
-		if (UWorld* World = GetWorld())
-		{
-			if (World->PersistentLevel)
-			{
-				TWeakObjectPtr<AWaterLandscapeBrush> WeakThis = this;
-
-				OnLoadedActorAddedToLevelEventHandle = World->PersistentLevel->OnLoadedActorAddedToLevelEvent.AddLambda([WeakThis](AActor& InActor) { 
-
-					if (AWaterLandscapeBrush* StrongThis = WeakThis.Get())
-					{
-						StrongThis->OnLevelActorAdded(&InActor);
-					}
-				});
-
-				OnLoadedActorRemovedFromLevelEventHandle = World->PersistentLevel->OnLoadedActorRemovedFromLevelEvent.AddLambda([WeakThis](AActor& InActor) { 
-					if (AWaterLandscapeBrush* StrongThis = WeakThis.Get())
-					{
-						StrongThis->OnLevelActorRemoved(&InActor);
-					}
-				});
-			}
-		}
-#endif // WITH_EDITOR
-
 		IWaterBrushActorInterface::GetOnWaterBrushActorChangedEvent().AddUObject(this, &AWaterLandscapeBrush::OnWaterBrushActorChanged);
-	}
 
-	// If we are loading do not trigger events
-	UpdateActors(!GIsEditorLoadingPackage);
+		// If we are loading do not trigger events
+		UpdateActors(!GIsEditorLoadingPackage);
+	}
 }
 
 void AWaterLandscapeBrush::OnLevelActorAdded(AActor* InActor)
@@ -345,7 +319,7 @@ void AWaterLandscapeBrush::BeginDestroy()
 {
 	Super::BeginDestroy();
 
-	if (!HasAnyFlags(RF_ClassDefaultObject | RF_ArchetypeObject | RF_Transient))
+	if (!IsTemplate())
 	{
 		ClearActors();
 
@@ -364,22 +338,56 @@ void AWaterLandscapeBrush::BeginDestroy()
 		GEngine->OnLevelActorDeleted().Remove(OnLevelActorDeletedHandle);
 		OnLevelActorDeletedHandle.Reset();
 
+		IWaterBrushActorInterface::GetOnWaterBrushActorChangedEvent().RemoveAll(this);
+	}
+}
+
+void AWaterLandscapeBrush::PostRegisterAllComponents()
+{
+	Super::PostRegisterAllComponents();
+
+#if WITH_EDITOR
+	if (!IsTemplate())
+	{
+		if (!OnLoadedActorRemovedFromLevelEventHandle.IsValid())
+		{
+			check(!OnLoadedActorRemovedFromLevelEventHandle.IsValid());
+
+			// In world partition, actors don't belong to levels and the on loaded/removed callbacks are different : 
+			// Since these are world events, we register/unregister to it in AWaterLandscapeBrush::RegisterAllComponents() / AWaterLandscapeBrush::UnregisterAllComponents() to make sure that the world is in a valid state when it's called :
+			UWorld* World = GetWorld();
+			checkf((World != nullptr) && (World->PersistentLevel != nullptr), TEXT("This function should only be called when the world and level are accessible"));
+			OnLoadedActorAddedToLevelEventHandle = World->PersistentLevel->OnLoadedActorAddedToLevelEvent.AddLambda([this](AActor& InActor) { OnLevelActorAdded(&InActor); });
+			OnLoadedActorRemovedFromLevelEventHandle = World->PersistentLevel->OnLoadedActorRemovedFromLevelEvent.AddLambda([this](AActor& InActor) { OnLevelActorRemoved(&InActor); });
+		}
+	}
+#endif // WITH_EDITOR
+}
+
+void AWaterLandscapeBrush::UnregisterAllComponents(bool bForReregister)
+{
+	Super::UnregisterAllComponents(bForReregister);
+
 #if WITH_EDITOR
-		// Unregister callbacks
+	if (!IsTemplate())
+	{
+		if (OnLoadedActorAddedToLevelEventHandle.IsValid())
+		{
+			check(OnLoadedActorRemovedFromLevelEventHandle.IsValid());
+
+			// UnregisterAllComponents can be called when the world if getting GCed. By this time, the world won't be able to broadcast these events so it's fine if we don't unregister from them :
 		if (UWorld* World = GetWorld())
 		{
-			if (World->PersistentLevel)
-			{
+				// Since these are world events, we register/unregister to it in AWaterLandscapeBrush::RegisterAllComponents() / AWaterLandscapeBrush::UnregisterAllComponents() to make sure that the world is in a valid state when it's called :
+				checkf(World->PersistentLevel != nullptr, TEXT("This function should only be called when the world and level are accessible"));
 				World->PersistentLevel->OnLoadedActorAddedToLevelEvent.Remove(OnLoadedActorAddedToLevelEventHandle);
 				World->PersistentLevel->OnLoadedActorRemovedFromLevelEvent.Remove(OnLoadedActorRemovedFromLevelEventHandle);
 			}
-		}
 		OnLoadedActorAddedToLevelEventHandle.Reset();
 		OnLoadedActorRemovedFromLevelEventHandle.Reset();
-#endif // WITH_EDITOR
-
-		IWaterBrushActorInterface::GetOnWaterBrushActorChangedEvent().RemoveAll(this);
+		}
 	}
+#endif // WITH_EDITOR
 }
 
 void AWaterLandscapeBrush::AddReferencedObjects(UObject* InThis, FReferenceCollector& Collector)
diff --git a/Engine/Plugins/Experimental/Water/Source/Editor/Private/WaterLandscapeBrush.h b/Engine/Plugins/Experimental/Water/Source/Editor/Private/WaterLandscapeBrush.h
--- a/Engine/Plugins/Experimental/Water/Source/Editor/Private/WaterLandscapeBrush.h
+++ b/Engine/Plugins/Experimental/Water/Source/Editor/Private/WaterLandscapeBrush.h
@@ -25,6 +25,8 @@ public:
 	virtual void PostInitProperties() override;
 	virtual void OnConstruction(const FTransform& Transform) override;
 	virtual void BeginDestroy() override;
+	virtual void PostRegisterAllComponents() override;
+	virtual void UnregisterAllComponents(bool bForReregister) override;
 	static void AddReferencedObjects(UObject* InThis, FReferenceCollector& Collector);
 
 	UFUNCTION(BlueprintCallable, Category = "Water", meta = (DeterminesOutputType = "WaterBodyClass", DynamicOutputParam = "OutWaterBodies"))
