Fixed water brush not working for actors already loaded by the time PostRegisterAllComponents was called

#rnx

diff --git a/Engine/Plugins/Experimental/Water/Source/Editor/Private/WaterLandscapeBrush.cpp b/Engine/Plugins/Experimental/Water/Source/Editor/Private/WaterLandscapeBrush.cpp
--- a/Engine/Plugins/Experimental/Water/Source/Editor/Private/WaterLandscapeBrush.cpp
+++ b/Engine/Plugins/Experimental/Water/Source/Editor/Private/WaterLandscapeBrush.cpp
@@ -360,6 +360,9 @@ void AWaterLandscapeBrush::PostRegisterAllComponents()
 			OnLoadedActorAddedToLevelEventHandle = World->PersistentLevel->OnLoadedActorAddedToLevelEvent.AddLambda([this](AActor& InActor) { OnLevelActorAdded(&InActor); });
 			OnLoadedActorRemovedFromLevelEventHandle = World->PersistentLevel->OnLoadedActorRemovedFromLevelEvent.AddLambda([this](AActor& InActor) { OnLevelActorRemoved(&InActor); });
 		}
+
+		// It's possible that actors registered to us via OnLoadedActorAddedToLevelEvent were already loaded by the time PostRegisterAllComponents runs, so we need to update our list of actors now : 
+		UpdateActors();
 	}
 #endif // WITH_EDITOR
 }
