Water: converted OwningWaterZone pointer to soft object pointer.

diff --git a/Engine/Plugins/Experimental/Water/Source/Runtime/Private/WaterBodyComponent.cpp b/Engine/Plugins/Experimental/Water/Source/Runtime/Private/WaterBodyComponent.cpp
--- a/Engine/Plugins/Experimental/Water/Source/Runtime/Private/WaterBodyComponent.cpp
+++ b/Engine/Plugins/Experimental/Water/Source/Runtime/Private/WaterBodyComponent.cpp
@@ -133,11 +133,9 @@ void UWaterBodyComponent::OnHiddenInGameChanged()
 
 FPrimitiveSceneProxy* UWaterBodyComponent::CreateSceneProxy()
 {
-	if (AffectsWaterMesh() && GetWaterZone())
+	if (AffectsWaterMesh())
 	{
-		const FVector TessellatedMeshCenter = GetWaterZone()->GetTessellatedWaterMeshCenter();
-		const FVector TessellatedMeshExtent =  GetWaterZone()->GetTessellatedWaterMeshExtent();
-		return new FWaterBodySceneProxy(this, FBox2D(FVector2D(TessellatedMeshCenter - TessellatedMeshExtent/ 2), FVector2D(TessellatedMeshCenter + TessellatedMeshExtent/ 2)));
+		return new FWaterBodySceneProxy(this);
 	}
 	return nullptr;
 }
@@ -382,23 +380,32 @@ void UWaterBodyComponent::UpdateWaterZones()
 {
 	if (UWorld* World = GetWorld())
 	{
+		AWaterZone* OldWaterZone = GetWaterZone();
 		AWaterZone* FoundZone = FindWaterZone();
 
-		if (OwningWaterZone != nullptr)
+		if (OldWaterZone != FoundZone)
 		{
-			OwningWaterZone->RemoveWaterBodyComponent(this);
-			OwningWaterZone = nullptr;
-		}
+			if (OldWaterZone)
+			{
+				OldWaterZone->RemoveWaterBodyComponent(this);
+				OldWaterZone = nullptr;
+			}
 
-		if (FoundZone != nullptr)
-		{
-			FoundZone->AddWaterBodyComponent(this);
-		}
+			if (FoundZone)
+			{
+				FoundZone->AddWaterBodyComponent(this);
+			}
 
-		OwningWaterZone = FoundZone;
+			OwningWaterZone = FoundZone;
+		}
 	}
 }
 
+AWaterZone* UWaterBodyComponent::GetWaterZone() const
+{
+	return OwningWaterZone.Get();
+}
+
 FPostProcessVolumeProperties UWaterBodyComponent::GetPostProcessProperties() const
 {
 	FPostProcessVolumeProperties Ret;
@@ -751,6 +758,11 @@ void UWaterBodyComponent::OnRegister()
 
 	CreateWaterSpriteComponent();
 #endif // WITH_EDITOR
+
+	if (AWaterZone* WaterZone = GetWaterZone())
+	{
+		WaterZone->AddWaterBodyComponent(this);
+	}
 }
 
 void UWaterBodyComponent::OnUnregister()
@@ -759,7 +771,7 @@ void UWaterBodyComponent::OnUnregister()
 	RegisterOnChangeWaterSplineData(/*bRegister = */false);
 #endif // WITH_EDITOR
 
-	if (AWaterZone* WaterZone = GetWaterZone(); IsValid(WaterZone))
+	if (AWaterZone* WaterZone = GetWaterZone())
 	{
 		WaterZone->RemoveWaterBodyComponent(this);
 	}
@@ -1421,9 +1433,6 @@ void UWaterBodyComponent::PostLoad()
 #if WITH_EDITOR
 	RegisterOnUpdateWavesData(GetWaterWaves(), /* bRegister = */true);
 #endif // WITH_EDITOR
-
-	// Register the component with the water zone since the component is the authority
-	UpdateWaterZones();
 }
 
 void UWaterBodyComponent::DeprecateData()
@@ -1532,7 +1541,6 @@ bool UWaterBodyComponent::SetDynamicParametersOnMID(UMaterialInstanceDynamic* In
 
 	InMID->SetVectorParameterValue(FixedVelocityName, GetConstantVelocity());
 
-	// Use WaterZone actor of the same level
 	if (const AWaterZone* WaterZone = GetWaterZone())
 	{
 		InMID->SetTextureParameterValue(WaterVelocityAndHeightName, WaterZone->WaterInfoTexture);
diff --git a/Engine/Plugins/Experimental/Water/Source/Runtime/Private/WaterBodySceneProxy.cpp b/Engine/Plugins/Experimental/Water/Source/Runtime/Private/WaterBodySceneProxy.cpp
--- a/Engine/Plugins/Experimental/Water/Source/Runtime/Private/WaterBodySceneProxy.cpp
+++ b/Engine/Plugins/Experimental/Water/Source/Runtime/Private/WaterBodySceneProxy.cpp
@@ -10,8 +10,7 @@
 #include "PrimitiveSceneInfo.h"
 #include "Algo/Transform.h"
 
-
-FWaterBodySceneProxy::FWaterBodySceneProxy(UWaterBodyComponent* Component, const FBox2D& TessellatedMeshBounds)
+FWaterBodySceneProxy::FWaterBodySceneProxy(UWaterBodyComponent* Component)
 	: FPrimitiveSceneProxy(Component)
 	, WaterBodySectionedLODMesh(GetScene().GetFeatureLevel())
 	, WaterBodyInfoMesh(GetScene().GetFeatureLevel())
@@ -21,7 +20,6 @@ FWaterBodySceneProxy::FWaterBodySceneProxy(UWaterBodyComponent* Component, const
 	WaterBodySectionedLODMesh.InitFromSections(Component->WaterBodyMeshSections);
 	WaterBodyInfoMesh.Init(Component->WaterBodyMeshVertices, Component->WaterBodyMeshIndices);
 	WaterBodyInfoDilatedMesh.Init(Component->DilatedWaterBodyMeshVertices, Component->DilatedWaterBodyMeshIndices);
-	WaterBodySectionedLODMesh.RebuildIndexBuffer(TessellatedMeshBounds);
 
 	if (UMaterialInstance* WaterInfoMaterialInstance = Component->GetWaterInfoMaterialInstance())
 	{
diff --git a/Engine/Plugins/Experimental/Water/Source/Runtime/Private/WaterZoneActor.cpp b/Engine/Plugins/Experimental/Water/Source/Runtime/Private/WaterZoneActor.cpp
--- a/Engine/Plugins/Experimental/Water/Source/Runtime/Private/WaterZoneActor.cpp
+++ b/Engine/Plugins/Experimental/Water/Source/Runtime/Private/WaterZoneActor.cpp
@@ -107,6 +107,17 @@ void AWaterZone::PostLoad()
 #endif // WITH_EDITORONLY_DATA
 }
 
+void AWaterZone::PostRegisterAllComponents()
+{
+	Super::PostRegisterAllComponents();
+
+	if (UpdateOverlappingWaterBodies())
+	{
+		MarkForRebuild(EWaterZoneRebuildFlags::UpdateWaterBodyLODSections);
+	}
+
+}
+
 void AWaterZone::MarkForRebuild(EWaterZoneRebuildFlags Flags)
 {
 	if (EnumHasAnyFlags(Flags, EWaterZoneRebuildFlags::UpdateWaterMesh))
@@ -119,14 +130,10 @@ void AWaterZone::MarkForRebuild(EWaterZoneRebuildFlags Flags)
 		UE_LOG(LogWater, Verbose, TEXT("AWaterZone::MarkForRebuild (UpdateWaterInfoTexture)"));
 		bNeedsWaterInfoRebuild = true;
 	}
-	if (EnumHasAnyFlags(Flags, EWaterZoneRebuildFlags::UpdateWaterBodyLODSections))
+	if (EnumHasAnyFlags(Flags, EWaterZoneRebuildFlags::UpdateWaterBodyLODSections) && bEnableNonTessellatedLODMesh)
 	{
 		UE_LOG(LogWater, Verbose, TEXT("AWaterZone::MarkForRebuild (UpdateWaterBodyLODSections)"));
-		ForEachWaterBodyComponent([](UWaterBodyComponent* WaterBodyComponent)
-		{
-			WaterBodyComponent->UpdateNonTessellatedMeshSections();
-			return true;
-		});
+		bNeedsNonTessellatedMeshRebuild = true;
 	}
 }
 
@@ -143,14 +150,21 @@ void AWaterZone::ForEachWaterBodyComponent(TFunctionRef<bool(UWaterBodyComponent
 
 void AWaterZone::AddWaterBodyComponent(UWaterBodyComponent* WaterBodyComponent)
 {
-	OwnedWaterBodies.AddUnique(WaterBodyComponent);
-	MarkForRebuild(EWaterZoneRebuildFlags::UpdateWaterInfoTexture | EWaterZoneRebuildFlags::UpdateWaterMesh);
+	if (OwnedWaterBodies.Find(WaterBodyComponent) == INDEX_NONE)
+	{
+		OwnedWaterBodies.Add(WaterBodyComponent);
+		MarkForRebuild(EWaterZoneRebuildFlags::UpdateWaterInfoTexture | EWaterZoneRebuildFlags::UpdateWaterMesh);
+	}
 }
 
 void AWaterZone::RemoveWaterBodyComponent(UWaterBodyComponent* WaterBodyComponent)
 {
-	OwnedWaterBodies.RemoveSwap(WaterBodyComponent);
-	MarkForRebuild(EWaterZoneRebuildFlags::UpdateWaterInfoTexture | EWaterZoneRebuildFlags::UpdateWaterMesh);
+	int32 Index;
+	if (OwnedWaterBodies.Find(WaterBodyComponent, Index))
+	{
+		OwnedWaterBodies.RemoveAtSwap(Index);
+		MarkForRebuild(EWaterZoneRebuildFlags::UpdateWaterInfoTexture | EWaterZoneRebuildFlags::UpdateWaterMesh);
+	}
 }
 
 void AWaterZone::Update()
@@ -163,6 +177,16 @@ void AWaterZone::Update()
 			bNeedsWaterInfoRebuild = false;
 		}
 	}
+	if (bNeedsNonTessellatedMeshRebuild)
+	{
+		ForEachWaterBodyComponent([](UWaterBodyComponent* WaterBodyComponent)
+		{
+			WaterBodyComponent->UpdateNonTessellatedMeshSections();
+			return true;
+		});
+
+		bNeedsNonTessellatedMeshRebuild = false;
+	}
 	
 	if (WaterMesh)
 	{
@@ -289,14 +313,15 @@ void AWaterZone::OnExtentChanged()
 	MarkForRebuild(EWaterZoneRebuildFlags::All);
 }
 
-void AWaterZone::UpdateOverlappingWaterBodies()
+bool AWaterZone::UpdateOverlappingWaterBodies()
 {
+	TArray<TWeakObjectPtr<UWaterBodyComponent>> OldOwnedWaterBodies = OwnedWaterBodies;
 	FWaterBodyManager::ForEachWaterBodyComponent(GetWorld(), [](UWaterBodyComponent* WaterBodyComponent)
 		{
 			WaterBodyComponent->UpdateWaterZones();
 			return true;
 		});
-
+	return (OwnedWaterBodies != OldOwnedWaterBodies);
 }
 
 
diff --git a/Engine/Plugins/Experimental/Water/Source/Runtime/Public/WaterBodyComponent.h b/Engine/Plugins/Experimental/Water/Source/Runtime/Public/WaterBodyComponent.h
--- a/Engine/Plugins/Experimental/Water/Source/Runtime/Public/WaterBodyComponent.h
+++ b/Engine/Plugins/Experimental/Water/Source/Runtime/Public/WaterBodyComponent.h
@@ -384,7 +384,7 @@ public:
 	virtual void Reset() {}
 
 	/** Gets the water zone to which this component belongs */
-	AWaterZone* GetWaterZone() const { return OwningWaterZone; }
+	AWaterZone* GetWaterZone() const;
 
 	/** Registers or this water body with corresponding overlapping water zones and unregisters it from any old zones if they are no longer overlapping. */
 	void UpdateWaterZones();
@@ -640,7 +640,7 @@ protected:
 	mutable TWeakObjectPtr<ALandscapeProxy> Landscape;
 
 	UPROPERTY()
-	TObjectPtr<AWaterZone> OwningWaterZone;
+	TSoftObjectPtr<AWaterZone> OwningWaterZone;
 
 	UPROPERTY(Transient)
 	FPostProcessSettings CurrentPostProcessSettings;
diff --git a/Engine/Plugins/Experimental/Water/Source/Runtime/Public/WaterBodySceneProxy.h b/Engine/Plugins/Experimental/Water/Source/Runtime/Public/WaterBodySceneProxy.h
--- a/Engine/Plugins/Experimental/Water/Source/Runtime/Public/WaterBodySceneProxy.h
+++ b/Engine/Plugins/Experimental/Water/Source/Runtime/Public/WaterBodySceneProxy.h
@@ -27,7 +27,7 @@ enum class EWaterInfoPass
 class FWaterBodySceneProxy final : public FPrimitiveSceneProxy
 {
 public:
-	FWaterBodySceneProxy(UWaterBodyComponent* Component, const FBox2D& TessellatedMeshBounds);
+	FWaterBodySceneProxy(UWaterBodyComponent* Component);
 	virtual ~FWaterBodySceneProxy();
 	virtual void GetDynamicMeshElements(const TArray<const FSceneView*>& Views, const FSceneViewFamily& ViewFamily, uint32 VisibilityMap, FMeshElementCollector& Collector) const override;
 	virtual FPrimitiveViewRelevance GetViewRelevance(const FSceneView* View) const override;
diff --git a/Engine/Plugins/Experimental/Water/Source/Runtime/Public/WaterZoneActor.h b/Engine/Plugins/Experimental/Water/Source/Runtime/Public/WaterZoneActor.h
--- a/Engine/Plugins/Experimental/Water/Source/Runtime/Public/WaterZoneActor.h
+++ b/Engine/Plugins/Experimental/Water/Source/Runtime/Public/WaterZoneActor.h
@@ -51,6 +51,7 @@ public:
 	virtual void BeginPlay() override;
 	virtual void PostLoadSubobjects(FObjectInstancingGraph* OuterInstanceGraph) override;
 	virtual void PostLoad() override;
+	virtual void PostRegisterAllComponents() override;
 
 	FVector2f GetWaterHeightExtents() const { return WaterHeightExtents; }
 	float GetGroundZMin() const { return GroundZMin; }
@@ -68,16 +69,21 @@ public:
 
 	bool IsNonTessellatedLODMeshEnabled() const { return bEnableNonTessellatedLODMesh; }
 
-#if WITH_EDITOR
-	virtual bool CanChangeIsSpatiallyLoadedFlag() const override { return false; }
-#endif
-
 private:
 
+	/**
+	 * Enqueue's a command on the Water Scene View Extension to re-render the water info on the next frame.
+	 * Returns false if the water info cannot be rendered this frame due to one of the dependencies not being ready yet (ie. a material under On-Demand-Shader-Compilation)
+	 */
 	bool UpdateWaterInfoTexture();
 
+	/**
+	 * Updates the list of owned water bodies causing any overlapping water bodies which do no have an owning water zone to register themselves to this water zone.
+	 * Returns true if any any change to the list of owned water bodies occurred
+	 */
+	bool UpdateOverlappingWaterBodies();
+
 	void OnExtentChanged();
-	void UpdateOverlappingWaterBodies();
 
 #if WITH_EDITOR
 	void OnActorSelectionChanged(const TArray<UObject*>& NewSelection, bool bForceRefresh);
@@ -135,6 +141,8 @@ private:
 
 	bool bNeedsWaterInfoRebuild = true;
 
+	bool bNeedsNonTessellatedMeshRebuild = true;
+
 	FVector2f WaterHeightExtents;
 	float GroundZMin;
 
