Water: removed all uses of lazy object pointers and replaced them with soft object pointers.

diff --git a/Engine/Plugins/Experimental/Water/Source/Runtime/Private/WaterBodyActor.cpp b/Engine/Plugins/Experimental/Water/Source/Runtime/Private/WaterBodyActor.cpp
--- a/Engine/Plugins/Experimental/Water/Source/Runtime/Private/WaterBodyActor.cpp
+++ b/Engine/Plugins/Experimental/Water/Source/Runtime/Private/WaterBodyActor.cpp
@@ -475,8 +475,8 @@ void AWaterBody::DeprecateData()
 		WaterBodyComponent->CollisionProfileName_DEPRECATED = CollisionProfileName_DEPRECATED;
 		WaterBodyComponent->WaterMID = WaterMID_DEPRECATED;
 		WaterBodyComponent->UnderwaterPostProcessMID = UnderwaterPostProcessMID_DEPRECATED;
-		WaterBodyComponent->Islands = Islands_DEPRECATED;
-		WaterBodyComponent->ExclusionVolumes = ExclusionVolumes_DEPRECATED;
+		WaterBodyComponent->Islands_DEPRECATED = Islands_DEPRECATED;
+		WaterBodyComponent->ExclusionVolumes_DEPRECATED = ExclusionVolumes_DEPRECATED;
 		WaterBodyComponent->bCanAffectNavigation_DEPRECATED = bCanAffectNavigation_DEPRECATED;
 		WaterBodyComponent->WaterNavAreaClass = WaterNavAreaClass_DEPRECATED;
 		WaterBodyComponent->ShapeDilation = ShapeDilation_DEPRECATED;
diff --git a/Engine/Plugins/Experimental/Water/Source/Runtime/Private/WaterBodyComponent.cpp b/Engine/Plugins/Experimental/Water/Source/Runtime/Private/WaterBodyComponent.cpp
--- a/Engine/Plugins/Experimental/Water/Source/Runtime/Private/WaterBodyComponent.cpp
+++ b/Engine/Plugins/Experimental/Water/Source/Runtime/Private/WaterBodyComponent.cpp
@@ -330,12 +330,12 @@ bool UWaterBodyComponent::ShouldGenerateWaterMeshTile() const
 
 void UWaterBodyComponent::AddIsland(AWaterBodyIsland* Island)
 {
-	Islands.AddUnique(Island);
+	WaterBodyIslands.AddUnique(Island);
 }
 
 void UWaterBodyComponent::RemoveIsland(AWaterBodyIsland* Island)
 {
-	Islands.RemoveSwap(Island);
+	WaterBodyIslands.RemoveSwap(Island);
 }
 
 void UWaterBodyComponent::UpdateIslands()
@@ -356,12 +356,12 @@ void UWaterBodyComponent::UpdateIslands()
 
 void UWaterBodyComponent::AddExclusionVolume(AWaterBodyExclusionVolume* InExclusionVolume)
 {
-	ExclusionVolumes.AddUnique(InExclusionVolume);
+	WaterBodyExclusionVolumes.AddUnique(InExclusionVolume);
 }
 
 void UWaterBodyComponent::RemoveExclusionVolume(AWaterBodyExclusionVolume* InExclusionVolume)
 {
-	ExclusionVolumes.RemoveSwap(InExclusionVolume);
+	WaterBodyExclusionVolumes.RemoveSwap(InExclusionVolume);
 }
 
 void UWaterBodyComponent::UpdateExclusionVolumes()
@@ -509,7 +509,7 @@ EWaterBodyQueryFlags UWaterBodyComponent::CheckAndAjustQueryFlags(EWaterBodyQuer
 
 bool UWaterBodyComponent::IsWorldLocationInExclusionVolume(const FVector& InWorldLocation) const
 {
-	for (const TLazyObjectPtr<AWaterBodyExclusionVolume>& ExclusionVolume : ExclusionVolumes)
+	for (const TSoftObjectPtr<AWaterBodyExclusionVolume>& ExclusionVolume : WaterBodyExclusionVolumes)
 	{
 		if (ExclusionVolume.IsValid() && ExclusionVolume->EncompassesPoint(InWorldLocation))
 		{
@@ -782,12 +782,11 @@ void UWaterBodyComponent::OnUnregister()
 TArray<AWaterBodyIsland*> UWaterBodyComponent::GetIslands() const
 {
 	TArray<AWaterBodyIsland*> IslandActors;
-	IslandActors.Reserve(Islands.Num());
+	IslandActors.Reserve(WaterBodyIslands.Num());
 
-	for (const TLazyObjectPtr<AWaterBodyIsland>& IslandPtr : Islands)
+	for (const TSoftObjectPtr<AWaterBodyIsland>& IslandPtr : WaterBodyIslands)
 	{
-		AWaterBodyIsland* Island = IslandPtr.Get();
-		if (Island)
+		if (AWaterBodyIsland* Island = IslandPtr.Get())
 		{
 			IslandActors.Add(Island);
 		}
@@ -799,9 +798,9 @@ TArray<AWaterBodyIsland*> UWaterBodyComponent::GetIslands() const
 TArray<AWaterBodyExclusionVolume*> UWaterBodyComponent::GetExclusionVolumes() const
 {
 	TArray<AWaterBodyExclusionVolume*> Result;
-	Result.Reserve(ExclusionVolumes.Num());
+	Result.Reserve(WaterBodyExclusionVolumes.Num());
 
-	for (const TLazyObjectPtr<AWaterBodyExclusionVolume>& VolumePtr : ExclusionVolumes)
+	for (const TSoftObjectPtr<AWaterBodyExclusionVolume>& VolumePtr : WaterBodyExclusionVolumes)
 	{
 		if (AWaterBodyExclusionVolume* Volume = VolumePtr.Get())
 		{
diff --git a/Engine/Plugins/Experimental/Water/Source/Runtime/Private/WaterBodyExclusionVolume.cpp b/Engine/Plugins/Experimental/Water/Source/Runtime/Private/WaterBodyExclusionVolume.cpp
--- a/Engine/Plugins/Experimental/Water/Source/Runtime/Private/WaterBodyExclusionVolume.cpp
+++ b/Engine/Plugins/Experimental/Water/Source/Runtime/Private/WaterBodyExclusionVolume.cpp
@@ -37,12 +37,12 @@ void AWaterBodyExclusionVolume::UpdateOverlappingWaterBodies()
 	TSet<UWaterBodyComponent*> ExistingOverlappingBodies;
 	TSet<TWeakObjectPtr<UWaterBodyComponent>> NewOverlappingBodies;
 
-	TLazyObjectPtr<AWaterBodyExclusionVolume> LazyThis(this);
+	TSoftObjectPtr<AWaterBodyExclusionVolume> SoftThis(this);
 
 	// Fixup overlapping bodies (iterating on actors on post-load will fail, but this is fine as this exclusion volume should not yet be referenced by an existing water body upon loading) : 
-	FWaterBodyManager::ForEachWaterBodyComponent(GetWorld(), [LazyThis, &ExistingOverlappingBodies](UWaterBodyComponent* WaterBodyComponent)
+	FWaterBodyManager::ForEachWaterBodyComponent(GetWorld(), [SoftThis, &ExistingOverlappingBodies](UWaterBodyComponent* WaterBodyComponent)
 	{
-		if (WaterBodyComponent->ContainsExclusionVolume(LazyThis))
+		if (WaterBodyComponent->ContainsExclusionVolume(SoftThis))
 		{
 			ExistingOverlappingBodies.Add(WaterBodyComponent);
 		}
diff --git a/Engine/Plugins/Experimental/Water/Source/Runtime/Private/WaterBodyIslandActor.cpp b/Engine/Plugins/Experimental/Water/Source/Runtime/Private/WaterBodyIslandActor.cpp
--- a/Engine/Plugins/Experimental/Water/Source/Runtime/Private/WaterBodyIslandActor.cpp
+++ b/Engine/Plugins/Experimental/Water/Source/Runtime/Private/WaterBodyIslandActor.cpp
@@ -193,12 +193,12 @@ void AWaterBodyIsland::UpdateOverlappingWaterBodyComponents()
 	TSet<UWaterBodyComponent*> ExistingOverlappingBodies;
 	TSet<TWeakObjectPtr<UWaterBodyComponent>> NewOverlappingBodies;
 
-	TLazyObjectPtr<AWaterBodyIsland> LazyThis(this);
+	TSoftObjectPtr<AWaterBodyIsland> SoftThis(this);
 
 	// Fixup overlapping bodies
-	FWaterBodyManager::ForEachWaterBodyComponent(GetWorld(), [LazyThis, &ExistingOverlappingBodies](UWaterBodyComponent* WaterBodyComponent)
+	FWaterBodyManager::ForEachWaterBodyComponent(GetWorld(), [SoftThis, &ExistingOverlappingBodies](UWaterBodyComponent* WaterBodyComponent)
 	{
-		if (WaterBodyComponent->ContainsIsland(LazyThis))
+		if (WaterBodyComponent->ContainsIsland(SoftThis))
 		{
 			ExistingOverlappingBodies.Add(WaterBodyComponent);
 		}
diff --git a/Engine/Plugins/Experimental/Water/Source/Runtime/Public/WaterBodyComponent.h b/Engine/Plugins/Experimental/Water/Source/Runtime/Public/WaterBodyComponent.h
--- a/Engine/Plugins/Experimental/Water/Source/Runtime/Public/WaterBodyComponent.h
+++ b/Engine/Plugins/Experimental/Water/Source/Runtime/Public/WaterBodyComponent.h
@@ -311,7 +311,7 @@ public:
 	UFUNCTION(BlueprintCallable, Category = Water)
 	TArray<AWaterBodyIsland*> GetIslands() const;
 
-	bool ContainsIsland(TLazyObjectPtr<AWaterBodyIsland> Island) const { return Islands.Contains(Island); }
+	bool ContainsIsland(TSoftObjectPtr<AWaterBodyIsland> Island) const { return WaterBodyIslands.Contains(Island); }
 
 	/**
 	 * Gets the exclusion volume that influence this water body
@@ -319,7 +319,7 @@ public:
 	UFUNCTION(BlueprintCallable, Category = Water)
 	TArray<AWaterBodyExclusionVolume*> GetExclusionVolumes() const;
 
-	bool ContainsExclusionVolume(TLazyObjectPtr<AWaterBodyExclusionVolume> InExclusionVolume) const { return ExclusionVolumes.Contains(InExclusionVolume); }
+	bool ContainsExclusionVolume(TSoftObjectPtr<AWaterBodyExclusionVolume> InExclusionVolume) const { return WaterBodyExclusionVolumes.Contains(InExclusionVolume); }
 
 	/** Component interface */
 	virtual void OnRegister() override;
@@ -631,10 +631,10 @@ protected:
 
 	/** Islands in this water body*/
 	UPROPERTY(Category = Water, VisibleAnywhere, AdvancedDisplay)
-	TArray<TLazyObjectPtr<AWaterBodyIsland>> Islands;
+	TArray<TSoftObjectPtr<AWaterBodyIsland>> WaterBodyIslands;
 
 	UPROPERTY(Category = Water, VisibleAnywhere, AdvancedDisplay)
-	TArray<TLazyObjectPtr<AWaterBodyExclusionVolume>> ExclusionVolumes;
+	TArray<TSoftObjectPtr<AWaterBodyExclusionVolume>> WaterBodyExclusionVolumes;
 
 	UPROPERTY(Transient)
 	mutable TWeakObjectPtr<ALandscapeProxy> Landscape;
@@ -659,6 +659,11 @@ protected:
 	TArray<uint32> DilatedWaterBodyMeshIndices;
 
 #if WITH_EDITORONLY_DATA
+	UPROPERTY()
+	TArray<TLazyObjectPtr<AWaterBodyIsland>> Islands_DEPRECATED;
+	UPROPERTY()
+	TArray<TLazyObjectPtr<AWaterBodyExclusionVolume>> ExclusionVolumes_DEPRECATED;
+
 	UPROPERTY()
 	bool bFillCollisionUnderWaterBodiesForNavmesh_DEPRECATED;
 
