Water: Exposed a manual water zone override

diff --git a/Engine/Plugins/Experimental/Water/Source/Runtime/Private/WaterBodyComponent.cpp b/Engine/Plugins/Experimental/Water/Source/Runtime/Private/WaterBodyComponent.cpp
--- a/Engine/Plugins/Experimental/Water/Source/Runtime/Private/WaterBodyComponent.cpp
+++ b/Engine/Plugins/Experimental/Water/Source/Runtime/Private/WaterBodyComponent.cpp
@@ -1709,6 +1709,11 @@ AWaterZone* UWaterBodyComponent::FindWaterZone() const
 {
 	TRACE_CPUPROFILER_EVENT_SCOPE(UWaterBodyComponent::FindWaterZone);
 
+	if (!WaterZoneOverride.IsNull())
+	{
+		return WaterZoneOverride.Get();
+	}
+
 	TArray<AWaterZone*, TInlineAllocator<4>> ViableZones;
 
 	if (const UWorld* World = GetWorld())
diff --git a/Engine/Plugins/Experimental/Water/Source/Runtime/Public/WaterBodyComponent.h b/Engine/Plugins/Experimental/Water/Source/Runtime/Public/WaterBodyComponent.h
--- a/Engine/Plugins/Experimental/Water/Source/Runtime/Public/WaterBodyComponent.h
+++ b/Engine/Plugins/Experimental/Water/Source/Runtime/Public/WaterBodyComponent.h
@@ -642,6 +642,9 @@ protected:
 	UPROPERTY()
 	TSoftObjectPtr<AWaterZone> OwningWaterZone;
 
+	UPROPERTY(Category = Water, EditAnywhere, BlueprintReadOnly, AdvancedDisplay)
+	TSoftObjectPtr<AWaterZone> WaterZoneOverride;
+
 	UPROPERTY(Transient)
 	FPostProcessSettings CurrentPostProcessSettings;
 
