Water: disable editing of water terrain settings when AffectsLandscape is set to false to convey that the settings would have no effect.

diff --git a/Engine/Plugins/Experimental/Water/Source/Runtime/Public/WaterBodyComponent.h b/Engine/Plugins/Experimental/Water/Source/Runtime/Public/WaterBodyComponent.h
--- a/Engine/Plugins/Experimental/Water/Source/Runtime/Public/WaterBodyComponent.h
+++ b/Engine/Plugins/Experimental/Water/Source/Runtime/Public/WaterBodyComponent.h
@@ -556,7 +556,7 @@ public:
 	UPROPERTY(Category = Rendering, EditAnywhere, BlueprintReadWrite, AdvancedDisplay, meta = (EditCondition = "bGenerateCollisions && UnderwaterPostProcessMaterial != nullptr", DisplayAfter = "UnderwaterPostProcessMaterial"))
 	FUnderwaterPostProcessSettings UnderwaterPostProcessSettings;
 
-	UPROPERTY(Category = Terrain, EditAnywhere, BlueprintReadWrite)
+	UPROPERTY(Category = Terrain, EditAnywhere, BlueprintReadWrite, meta = (EditCondition = "bAffectsLandscape"))
 	FWaterCurveSettings CurveSettings;
 
 	UPROPERTY(Category = Rendering, EditAnywhere, BlueprintReadOnly)
@@ -575,11 +575,11 @@ public:
 	UPROPERTY(Category = Rendering, EditAnywhere, BlueprintReadOnly, meta = (DisplayAfter = "WaterMaterial"))
 	TObjectPtr<UMaterialInterface> WaterInfoMaterial;
 	
-	UPROPERTY(Category = Terrain, EditAnywhere, BlueprintReadWrite)
+	UPROPERTY(Category = Terrain, EditAnywhere, BlueprintReadWrite, meta = (EditCondition = "bAffectsLandscape"))
 	FWaterBodyHeightmapSettings WaterHeightmapSettings;
 
 #if WITH_EDITORONLY_DATA
-	UPROPERTY(Category = Terrain, EditAnywhere, BlueprintReadWrite)
+	UPROPERTY(Category = Terrain, EditAnywhere, BlueprintReadWrite, meta = (EditCondition = "bAffectsLandscape"))
 	TMap<FName, FWaterBodyWeightmapSettings> LayerWeightmapSettings;
 #endif
 
