Water: FixedWaterDepth now has it's own parameter separate from the spline point depths so that oceans/lakes may set it.

diff --git a/Engine/Plugins/Experimental/Water/Source/Runtime/Private/WaterBodyComponent.cpp b/Engine/Plugins/Experimental/Water/Source/Runtime/Private/WaterBodyComponent.cpp
--- a/Engine/Plugins/Experimental/Water/Source/Runtime/Private/WaterBodyComponent.cpp
+++ b/Engine/Plugins/Experimental/Water/Source/Runtime/Private/WaterBodyComponent.cpp
@@ -463,7 +463,10 @@ float UWaterBodyComponent::GetConstantDepth() const
 {
 	// Only makes sense when you consider the water depth to be constant for the whole water body, in which case we just use the first spline key's : 
 	const UWaterSplineComponent* const WaterSpline = GetWaterSpline();
-	return WaterSpline ? WaterSpline->GetFloatPropertyAtSplineInputKey(0.0f, GET_MEMBER_NAME_CHECKED(UWaterSplineMetadata, Depth)) : 0.0f;
+	float SplinePointDepth = WaterSpline ? WaterSpline->GetFloatPropertyAtSplineInputKey(0.0f, GET_MEMBER_NAME_CHECKED(UWaterSplineMetadata, Depth)) : 0.0f;
+
+	// Ensure that the existing spline depth is used if it is non-zero, otherwise use the new FixedWaterDepth parameter.
+	return SplinePointDepth != 0.f ? SplinePointDepth : FixedWaterDepth;
 }
 
 FVector UWaterBodyComponent::GetConstantVelocity() const
diff --git a/Engine/Plugins/Experimental/Water/Source/Runtime/Public/WaterBodyComponent.h b/Engine/Plugins/Experimental/Water/Source/Runtime/Public/WaterBodyComponent.h
--- a/Engine/Plugins/Experimental/Water/Source/Runtime/Public/WaterBodyComponent.h
+++ b/Engine/Plugins/Experimental/Water/Source/Runtime/Public/WaterBodyComponent.h
@@ -661,6 +661,10 @@ protected:
 	TArray<FDynamicMeshVertex> DilatedWaterBodyMeshVertices;
 	TArray<uint32> DilatedWaterBodyMeshIndices;
 
+	/** If the Water Material assigned to this component has Fixed Depth enabled, this is the depth that is passed. */
+	UPROPERTY(Category = Water, EditAnywhere, AdvancedDisplay)
+	double FixedWaterDepth = 512.0;
+
 #if WITH_EDITORONLY_DATA
 	UPROPERTY()
 	TArray<TLazyObjectPtr<AWaterBodyIsland>> Islands_DEPRECATED;
