Fix for AWaterZone::GroundZMin being set to zero when the extent of a LandscapeProxy's bounds is zero. This should also improve precision in the water info texture.

#rnx

diff --git a/Engine/Plugins/Experimental/Water/Source/Runtime/Private/WaterZoneActor.cpp b/Engine/Plugins/Experimental/Water/Source/Runtime/Private/WaterZoneActor.cpp
--- a/Engine/Plugins/Experimental/Water/Source/Runtime/Private/WaterZoneActor.cpp
+++ b/Engine/Plugins/Experimental/Water/Source/Runtime/Private/WaterZoneActor.cpp
@@ -386,8 +386,8 @@ bool AWaterZone::UpdateWaterInfoTexture()
 		for (ALandscapeProxy* LandscapeProxy : TActorRange<ALandscapeProxy>(World))
 		{
 			const FBox LandscapeBox = LandscapeProxy->GetComponentsBoundingBox();
-			// Only consider landscapes which this zone intersects with in XY
-			if (WaterZoneBounds.IntersectXY(LandscapeBox))
+			// Only consider landscapes which this zone intersects with in XY and if the landscape volume is not zero sized
+			if (WaterZoneBounds.IntersectXY(LandscapeBox) && LandscapeBox.GetVolume() > 0.0)
 			{
 				GroundZMin = FMath::Min(GroundZMin, LandscapeBox.Min.Z);
 				GroundZMax = FMath::Max(GroundZMax, LandscapeBox.Max.Z);
