Artificially clamp the spline arrive/leave tangents when editing scale to prevent scaling them past some safe threshold. Since the tangent values get scaled multiplicatively it is very easy to reach values which can produce NaNs and crash various portions of the editor. In the future, this should be fixed by scaling additively rather than multiplicatively as is done in the rest of the engine.

diff --git a/Engine/Plugins/Experimental/Water/Source/Runtime/Private/WaterSplineComponent.cpp b/Engine/Plugins/Experimental/Water/Source/Runtime/Private/WaterSplineComponent.cpp
--- a/Engine/Plugins/Experimental/Water/Source/Runtime/Private/WaterSplineComponent.cpp
+++ b/Engine/Plugins/Experimental/Water/Source/Runtime/Private/WaterSplineComponent.cpp
@@ -231,6 +231,21 @@ bool UWaterSplineComponent::SynchronizeWaterProperties()
 				// Set the splines local scale.x to the depth and ensure it has some small positive value. (non-zero scale required for collision to work)
 				Scale.Y = DepthAtPoint = FMath::Max(DepthAtPoint, KINDA_SMALL_NUMBER);
 			}
+
+			// #hack: temporarily clamp the tangents to a sensible range to prevent multiplicatively scaling until they hit infinity and cause a crash
+			auto ClampVec3 = [](const FVector& Vec3, double Min, double Max) {
+				FVector Result;
+				Result.X = FMath::Clamp(Vec3.X, Min, Max);
+				Result.Y = FMath::Clamp(Vec3.Y, Min, Max);
+				Result.Z = FMath::Clamp(Vec3.Z, Min, Max);
+				return Result;
+			};
+			FVector& ArriveTangent = SplineCurves.Position.Points[Point].ArriveTangent;
+			FVector& LeaveTangent = SplineCurves.Position.Points[Point].LeaveTangent;
+			
+			constexpr double MaxTangentValue = 1.e10L;
+			ArriveTangent = ClampVec3(ArriveTangent, -MaxTangentValue, MaxTangentValue);
+			LeaveTangent = ClampVec3(LeaveTangent, -MaxTangentValue, MaxTangentValue);
 		}
 	}
 
