Add NewWaterZoneScale so that seam around landscapes less visible

diff --git a/Engine/Plugins/Experimental/Water/Source/Editor/Private/WaterEditorModule.cpp b/Engine/Plugins/Experimental/Water/Source/Editor/Private/WaterEditorModule.cpp
--- a/Engine/Plugins/Experimental/Water/Source/Editor/Private/WaterEditorModule.cpp
+++ b/Engine/Plugins/Experimental/Water/Source/Editor/Private/WaterEditorModule.cpp
@@ -45,6 +45,11 @@ DEFINE_LOG_CATEGORY(LogWaterEditor);
 
 EAssetTypeCategories::Type FWaterEditorModule::WaterAssetCategory;
 
+namespace WaterEditorModule
+{
+	static TAutoConsoleVariable<float> CVarOverrideNewWaterZoneScale(TEXT("r.Water.WaterZoneActor.OverrideNewWaterZoneScale"), 0, TEXT("Multiply WaterZone actor extent beyond landscape by this amount. 0 means do override."));
+}
+
 void FWaterEditorModule::StartupModule()
 {
 	FWaterUIStyle::Initialize();
@@ -292,8 +297,22 @@ void FWaterEditorModule::OnLevelActorAddedToWorld(AActor* Actor)
 			if (WaterZoneBounds.IsValid)
 			{
 				WaterZoneActor->SetActorLocation(WaterZoneBounds.GetCenter());
+
 				// FBox::GetExtent returns the radius, SetZoneExtent expects diameter.
-				WaterZoneActor->SetZoneExtent(2 * FVector2D(WaterZoneBounds.GetExtent()));
+				FVector2D NewExtent = 2 * FVector2D(WaterZoneBounds.GetExtent());
+
+				float ZoneExtentScale = WaterEditorModule::CVarOverrideNewWaterZoneScale.GetValueOnAnyThread();
+				if (ZoneExtentScale == 0)
+				{
+					ZoneExtentScale = GetDefault<UWaterEditorSettings>()->WaterZoneActorDefaults.NewWaterZoneScale;
+				}
+
+				if (ZoneExtentScale != 0)
+				{
+					NewExtent = FMath::Abs(ZoneExtentScale) * NewExtent;
+				}
+
+				WaterZoneActor->SetZoneExtent(NewExtent);
 			}
 
 			// Set the defaults here because the actor factory isn't triggered on manual SpawnActor.
diff --git a/Engine/Plugins/Experimental/Water/Source/Editor/Public/WaterEditorSettings.h b/Engine/Plugins/Experimental/Water/Source/Editor/Public/WaterEditorSettings.h
--- a/Engine/Plugins/Experimental/Water/Source/Editor/Public/WaterEditorSettings.h
+++ b/Engine/Plugins/Experimental/Water/Source/Editor/Public/WaterEditorSettings.h
@@ -48,6 +48,10 @@ struct WATEREDITOR_API FWaterZoneActorDefaults
 
 	UPROPERTY(EditAnywhere, config, Category = Mesh)
 	float FarDistanceMeshExtent = 4000000.0f;
+
+	UPROPERTY(EditAnywhere, config, Category = Mesh)
+	float NewWaterZoneScale = 1.0f;
+
 protected:
 	UPROPERTY(EditAnywhere, Config, Category = Mesh)
 	TSoftObjectPtr<UMaterialInterface> FarDistanceMaterial;
