Change WaterEditorModule to use GetValueOnGameThread in one location

diff --git a/Engine/Plugins/Experimental/Water/Source/Editor/Private/WaterEditorModule.cpp b/Engine/Plugins/Experimental/Water/Source/Editor/Private/WaterEditorModule.cpp
--- a/Engine/Plugins/Experimental/Water/Source/Editor/Private/WaterEditorModule.cpp
+++ b/Engine/Plugins/Experimental/Water/Source/Editor/Private/WaterEditorModule.cpp
@@ -301,7 +301,7 @@ void FWaterEditorModule::OnLevelActorAddedToWorld(AActor* Actor)
 				// FBox::GetExtent returns the radius, SetZoneExtent expects diameter.
 				FVector2D NewExtent = 2 * FVector2D(WaterZoneBounds.GetExtent());
 
-				float ZoneExtentScale = WaterEditorModule::CVarOverrideNewWaterZoneScale.GetValueOnAnyThread();
+				float ZoneExtentScale = WaterEditorModule::CVarOverrideNewWaterZoneScale.GetValueOnGameThread();
 				if (ZoneExtentScale == 0)
 				{
 					ZoneExtentScale = GetDefault<UWaterEditorSettings>()->WaterZoneActorDefaults.NewWaterZoneScale;
