Fixed water body custom component registering static mesh components when the body component isn't registered

diff --git a/Engine/Plugins/Experimental/Water/Source/Runtime/Private/WaterBodyCustomComponent.cpp b/Engine/Plugins/Experimental/Water/Source/Runtime/Private/WaterBodyCustomComponent.cpp
--- a/Engine/Plugins/Experimental/Water/Source/Runtime/Private/WaterBodyCustomComponent.cpp
+++ b/Engine/Plugins/Experimental/Water/Source/Runtime/Private/WaterBodyCustomComponent.cpp
@@ -72,7 +72,11 @@ void UWaterBodyCustomComponent::OnUpdateBody(bool bWithExclusionVolumes)
 		MeshComp = NewObject<UStaticMeshComponent>(OwnerActor, TEXT("CustomMeshComponent"), RF_Transactional);
 		MeshComp->SetNetAddressable(); // it's deterministically named so it's addressable over network (needed for collision)
 		MeshComp->SetupAttachment(this);
-		MeshComp->RegisterComponent();
+
+		if(IsRegistered())
+		{
+			MeshComp->RegisterComponent();
+		}
 	}
 
 	TInlineComponentArray<UPrimitiveComponent*> PrimitiveComponents;
