Fixed: Most of the water texture is floating above the ground all over the map

diff --git a/Engine/Plugins/Experimental/Water/Shaders/Private/WaterInfoMerge.usf b/Engine/Plugins/Experimental/Water/Shaders/Private/WaterInfoMerge.usf
--- a/Engine/Plugins/Experimental/Water/Shaders/Private/WaterInfoMerge.usf
+++ b/Engine/Plugins/Experimental/Water/Shaders/Private/WaterInfoMerge.usf
@@ -4,8 +4,11 @@
 #include "/Engine/Private/DeferredShadingCommon.ush"
 
 Texture2D<float4> DepthTexture;
+SamplerState DepthTextureSampler;
 Texture2D<float4> ColorTexture;
+SamplerState ColorTextureSampler;
 Texture2D<float4> DilationTexture;
+SamplerState DilationTextureSampler;
 
 
 float2 WaterHeightExtents;
@@ -49,11 +52,12 @@ void Main(
 	out float4 OutColor : SV_Target0
 	)
 {
-	const float4 SVPos = Input.Position;
-	uint2 TextureCoordinates = floor(SVPos.xy);
+	float4 ColorSample = ColorTexture.SampleLevel(ColorTextureSampler, Input.UV, 0);
+	float4 DepthSample = DepthTexture.SampleLevel(DepthTextureSampler, Input.UV, 0);
+	float4 DilationSample = DilationTexture.SampleLevel(DilationTextureSampler, Input.UV, 0);
 	
-	const FWaterColorPassResult WaterInfo = MakeWaterColorPassResult(ColorTexture.Load(int3(TextureCoordinates, 0)));
-	const float GroundDeviceDepth = DecodeFloatRGB(DepthTexture.Load(int3(TextureCoordinates, 0)).rgb);
+	const FWaterColorPassResult WaterInfo = MakeWaterColorPassResult(ColorSample);
+	const float GroundDeviceDepth = DecodeFloatRGB(DepthSample.rgb);
 	const float GroundSceneDepth = ConvertFromDeviceZ(GroundDeviceDepth);
 	
 	const float WaterSceneDepth = WaterInfo.SceneDepth;
@@ -87,7 +91,7 @@ void Main(
 	}
 	else
 	{
-		const float DilationDeviceDepth = DecodeFloatRGB(DilationTexture.Load(int3(TextureCoordinates, 0)).rgb);
+		const float DilationDeviceDepth = DecodeFloatRGB(DilationSample.rgb);
 		const float DilationSceneDepth = ConvertFromDeviceZ(DilationDeviceDepth);
 		const bool bPixelHasDilationData = DilationSceneDepth != 0.0;
 		
diff --git a/Engine/Plugins/Experimental/Water/Source/Runtime/Private/WaterInfoRendering.cpp b/Engine/Plugins/Experimental/Water/Source/Runtime/Private/WaterInfoRendering.cpp
--- a/Engine/Plugins/Experimental/Water/Source/Runtime/Private/WaterInfoRendering.cpp
+++ b/Engine/Plugins/Experimental/Water/Source/Runtime/Private/WaterInfoRendering.cpp
@@ -57,9 +57,12 @@ public:
 	BEGIN_SHADER_PARAMETER_STRUCT(FParameters, )
 		SHADER_PARAMETER_STRUCT_REF(FViewUniformShaderParameters, View)
 		SHADER_PARAMETER_STRUCT_INCLUDE(FSceneTextureShaderParameters, SceneTextures)
-		SHADER_PARAMETER_RDG_TEXTURE(Texture2D, DepthTexture)
+		SHADER_PARAMETER_RDG_TEXTURE(Texture2D<float4>, DepthTexture)
+		SHADER_PARAMETER_SAMPLER(SamplerState, DepthTextureSampler)
 		SHADER_PARAMETER_RDG_TEXTURE(Texture2D<float4>, ColorTexture)
-		SHADER_PARAMETER_RDG_TEXTURE(Texture2D, DilationTexture)
+		SHADER_PARAMETER_SAMPLER(SamplerState, ColorTextureSampler)
+		SHADER_PARAMETER_RDG_TEXTURE(Texture2D<float4>, DilationTexture)
+		SHADER_PARAMETER_SAMPLER(SamplerState, DilationTextureSampler)
 		SHADER_PARAMETER(FVector2f, WaterHeightExtents)
 		SHADER_PARAMETER(float, GroundZMin)
 		SHADER_PARAMETER(float, CaptureZ)
@@ -103,8 +106,11 @@ static void MergeWaterInfoAndDepth(
 		PassParameters->RenderTargets[0] = FRenderTargetBinding(OutputTexture, ERenderTargetLoadAction::ENoAction);
 		PassParameters->SceneTextures = SceneTextures.GetSceneTextureShaderParameters(ViewFamily.GetFeatureLevel());
 		PassParameters->DepthTexture = DepthTexture;
+		PassParameters->DepthTextureSampler = TStaticSamplerState<SF_Point, AM_Clamp, AM_Clamp, AM_Clamp>::GetRHI();
 		PassParameters->ColorTexture = ColorTexture;
+		PassParameters->ColorTextureSampler = TStaticSamplerState<SF_Point, AM_Clamp, AM_Clamp, AM_Clamp>::GetRHI();
 		PassParameters->DilationTexture = DilationTexture;
+		PassParameters->DilationTextureSampler = TStaticSamplerState<SF_Point, AM_Clamp, AM_Clamp, AM_Clamp>::GetRHI();
 		PassParameters->CaptureZ = Params.CaptureZ;
 		PassParameters->WaterHeightExtents = Params.WaterHeightExtents;
 		PassParameters->GroundZMin = Params.GroundZMin;
