Fixing inadvertant water warning in the logs -- converting to verbose

diff --git a/Engine/Plugins/Experimental/Water/Source/Runtime/Private/WaterSubsystem.cpp b/Engine/Plugins/Experimental/Water/Source/Runtime/Private/WaterSubsystem.cpp
--- a/Engine/Plugins/Experimental/Water/Source/Runtime/Private/WaterSubsystem.cpp
+++ b/Engine/Plugins/Experimental/Water/Source/Runtime/Private/WaterSubsystem.cpp
@@ -368,7 +368,7 @@ void UWaterSubsystem::ApplyRuntimeSettings(const UWaterRuntimeSettings* Settings
 #if WITH_EDITOR
 void UWaterSubsystem::OnHeightmapStreamed(const FOnHeightmapStreamedContext& InContext)
 {
-	UE_LOG(LogWater, Warning, TEXT("OnHeightmapStreamed!  Rebuilding Water Info Texture..."));
+	UE_LOG(LogWater, Verbose, TEXT("UWaterSubsystem::OnHeightmapStreamed() -- Rebuilding Water Info Texture..."));
 	MarkWaterZonesInRegionForRebuild(InContext.GetUpdateRegion(), EWaterZoneRebuildFlags::UpdateWaterInfoTexture);
 }
 #endif // WITH_EDITOR
