Fix for water disappearing at LWC distances in TM-VSM. Replaced some float calculations with double calculations.

#rnx

diff --git a/Engine/Plugins/Experimental/Water/Source/Runtime/Private/WaterQuadTree.cpp b/Engine/Plugins/Experimental/Water/Source/Runtime/Private/WaterQuadTree.cpp
--- a/Engine/Plugins/Experimental/Water/Source/Runtime/Private/WaterQuadTree.cpp
+++ b/Engine/Plugins/Experimental/Water/Source/Runtime/Private/WaterQuadTree.cpp
@@ -828,18 +828,18 @@ static void SplitPolyWithLine(const TArray<FVector2D>& InPoly, const FVector2D&
 	int32 NumPVerts = InPoly.Num();
 
 	// Calculate distance of verts from clipping line
-	TArray<float> PlaneDist;
+	TArray<double> PlaneDist;
 	PlaneDist.AddZeroed(NumPVerts);
 	for (int32 i = 0; i < NumPVerts; i++)
 	{
 		const FVector2D PointDiff = InPoly[i] - LinePoint;
-		PlaneDist[i] = LineNormal.X > 0.0f ? PointDiff.X : PointDiff.Y;
+		PlaneDist[i] = LineNormal.X > 0.0 ? PointDiff.X : PointDiff.Y;
 	}
 
 	for (int32 ThisVert = 0; ThisVert < NumPVerts; ThisVert++)
 	{
 		// Vert is on positive side of line, add to Poly0
-		if (PlaneDist[ThisVert] > 0.0f)
+		if (PlaneDist[ThisVert] > 0.0)
 		{
 			OutPoly0.Add(InPoly[ThisVert]);
 		}
@@ -851,10 +851,10 @@ static void SplitPolyWithLine(const TArray<FVector2D>& InPoly, const FVector2D&
 		// If start and next vert are on opposite sides, add intersection
 		int32 NextVert = (ThisVert + 1) % NumPVerts;
 
-		if (PlaneDist[ThisVert] * PlaneDist[NextVert] < 0.0f)
+		if (PlaneDist[ThisVert] * PlaneDist[NextVert] < 0.0)
 		{
 			// Find distance along edge that plane is
-			float Alpha = -PlaneDist[ThisVert] / (PlaneDist[NextVert] - PlaneDist[ThisVert]);
+			double Alpha = -PlaneDist[ThisVert] / (PlaneDist[NextVert] - PlaneDist[ThisVert]);
 			FVector2D NewVertPos = FMath::Lerp(InPoly[ThisVert], InPoly[NextVert], Alpha);
 
 			// Save vert
@@ -864,24 +864,24 @@ static void SplitPolyWithLine(const TArray<FVector2D>& InPoly, const FVector2D&
 	}
 }
 
-static float CalcPoly2DArea(const TArray<FVector2D>& InPoly)
+static double CalcPoly2DArea(const TArray<FVector2D>& InPoly)
 {
-	float ResultArea = 0.0f;
+	double ResultArea = 0.0;
 	for (int i = 0, j = InPoly.Num() - 1; i < InPoly.Num(); j = i++)
 	{
 		const FVector2D& Vert0 = InPoly[i];
 		const FVector2D& Vert1 = InPoly[j];
 		ResultArea += Vert1.X * Vert0.Y - Vert0.X * Vert1.Y;
 	}
-	return ResultArea * 0.5f;
+	return ResultArea * 0.5;
 }
 
 void FWaterQuadTree::AddOceanRecursive(const TArray<FVector2D>& InPoly, const FBox2D& InBox, const FVector2D& InZBounds, bool HSplit, int32 InDepth, uint32 InWaterBodyIndex)
 {
 	// Some value to guard against false positives, based on the max area
-	const float BoxArea = InBox.GetArea();
-	const float AreaEpsilon = BoxArea * 0.0001f;
-	const FVector2D LeafSizeShrink(LeafSize * 0.25f, LeafSize * 0.25f);
+	const double BoxArea = InBox.GetArea();
+	const double AreaEpsilon = BoxArea * 0.0001;
+	const FVector2D LeafSizeShrink(LeafSize * 0.25, LeafSize * 0.25);
 
 	//We've reached the bottom, figure out if this poly is filling out its box
 	if (InDepth == 0)
@@ -940,9 +940,9 @@ void FWaterQuadTree::AddOceanRecursive(const TArray<FVector2D>& InPoly, const FB
 void FWaterQuadTree::AddLakeRecursive(const TArray<FVector2D>& InPoly, const FBox2D& InBox, const FVector2D& InZBounds, bool HSplit, int32 InDepth, uint32 InWaterBodyIndex)
 {
 	// Some value to guard against false positives, based on the max area
-	const float BoxArea = InBox.GetArea();
-	const float AreaEpsilon = BoxArea * 0.0001f;
-	const FVector2D LeafSizeShrink(LeafSize * 0.01f, LeafSize * 0.01f);
+	const double BoxArea = InBox.GetArea();
+	const double AreaEpsilon = BoxArea * 0.0001;
+	const FVector2D LeafSizeShrink(LeafSize * 0.01, LeafSize * 0.01);
 
 	//We've reached the bottom, figure out if this poly is filling out its box
 	if (InDepth == 0)
@@ -973,7 +973,7 @@ void FWaterQuadTree::AddLakeRecursive(const TArray<FVector2D>& InPoly, const FBo
 	// Recurse split the two new polys if they have any significant lake poly area
 	// Poly0 is the positive box (on the positive side of the line normal)
 	const FBox2D HalfBox0(InBox.Min + InBox.GetExtent() * LineNormal, InBox.Max);
-	if (CalcPoly2DArea(Poly0) > (BoxArea * 0.5f - AreaEpsilon))
+	if (CalcPoly2DArea(Poly0) > (BoxArea * 0.5 - AreaEpsilon))
 	{
 		// This halfbox is filled with lake poly, mark as water
 		FBox TileBounds(FVector(HalfBox0.Min + LeafSizeShrink, InZBounds.X), FVector(HalfBox0.Max - LeafSizeShrink, InZBounds.Y));
@@ -986,7 +986,7 @@ void FWaterQuadTree::AddLakeRecursive(const TArray<FVector2D>& InPoly, const FBo
 
 	// Poly1 is the negative box (on the negative side of the line normal)
 	const FBox2D HalfBox1(InBox.Min, InBox.Max - InBox.GetExtent() * LineNormal);
-	if (CalcPoly2DArea(Poly1) > (BoxArea * 0.5f - AreaEpsilon))
+	if (CalcPoly2DArea(Poly1) > (BoxArea * 0.5 - AreaEpsilon))
 	{
 		// This halfbox is filled with lake poly, mark as water
 		FBox TileBounds(FVector(HalfBox1.Min + LeafSizeShrink, InZBounds.X), FVector(HalfBox1.Max - LeafSizeShrink, InZBounds.Y));
