Landscape : Fix landscape not updating on changes to water.

Previously to this change we defined the OnWaterBrushActorChangedEvent static within a header which resulted in mutiple delegates being created depending on which translation units the header was included in, rather than a single delegtate to notify the landscapes of changes.

diff --git a/Engine/Plugins/Experimental/Water/Source/Runtime/Private/WaterBrushActorInterface.cpp b/Engine/Plugins/Experimental/Water/Source/Runtime/Private/WaterBrushActorInterface.cpp
--- a/Engine/Plugins/Experimental/Water/Source/Runtime/Private/WaterBrushActorInterface.cpp
+++ b/Engine/Plugins/Experimental/Water/Source/Runtime/Private/WaterBrushActorInterface.cpp
@@ -15,5 +15,14 @@ TArray<UPrimitiveComponent*> IWaterBrushActorInterface::GetBrushRenderableCompon
 {
 	return TArray<UPrimitiveComponent*>();
 }
+
+IWaterBrushActorInterface::FWaterBrushActorChangedEvent& IWaterBrushActorInterface::GetOnWaterBrushActorChangedEvent()
+{
+	static IWaterBrushActorInterface::FWaterBrushActorChangedEvent WaterBrushActorChangedEvent;
+	return WaterBrushActorChangedEvent;
+}
+
 #endif //WITH_EDITOR
 
+
+
diff --git a/Engine/Plugins/Experimental/Water/Source/Runtime/Public/WaterBrushActorInterface.h b/Engine/Plugins/Experimental/Water/Source/Runtime/Public/WaterBrushActorInterface.h
--- a/Engine/Plugins/Experimental/Water/Source/Runtime/Public/WaterBrushActorInterface.h
+++ b/Engine/Plugins/Experimental/Water/Source/Runtime/Public/WaterBrushActorInterface.h
@@ -103,11 +103,7 @@ class WATER_API IWaterBrushActorInterface
 	 * Event sent whenever a data change occurs on a water brush actor
 	 */
 	DECLARE_EVENT_OneParam(IWaterBrushActorInterface, FWaterBrushActorChangedEvent, const FWaterBrushActorChangedEventParams&);
-	static FWaterBrushActorChangedEvent& GetOnWaterBrushActorChangedEvent()
-	{
-		static FWaterBrushActorChangedEvent WaterBrushActorChangedEvent;
-		return WaterBrushActorChangedEvent;
-	}
+	static FWaterBrushActorChangedEvent& GetOnWaterBrushActorChangedEvent();
 
 	void BroadcastWaterBrushActorChangedEvent(const FWaterBrushActorChangedEventParams& InParams)
 	{
