Water: Added a new project setting to allow overriding the base WaterZone class spawned automatically for users.

diff --git a/Engine/Plugins/Experimental/Water/Source/Editor/Private/WaterEditorModule.cpp b/Engine/Plugins/Experimental/Water/Source/Editor/Private/WaterEditorModule.cpp
--- a/Engine/Plugins/Experimental/Water/Source/Editor/Private/WaterEditorModule.cpp
+++ b/Engine/Plugins/Experimental/Water/Source/Editor/Private/WaterEditorModule.cpp
@@ -215,11 +215,11 @@ void FWaterEditorModule::OnLevelActorAddedToWorld(AActor* Actor)
 			WaterZoneBounds += LandscapeBounds;
 		}
 
+		const UWaterEditorSettings* WaterEditorSettings = GetDefault<UWaterEditorSettings>();
+		check(WaterEditorSettings != nullptr);
 		// Automatically setup landscape-affecting features (water brush) if needed : 
 		if ((WaterBrushActor != nullptr) && WaterBrushActor->AffectsLandscape() && !FoundLandscapes.IsEmpty())
 		{
-			const UWaterEditorSettings* WaterEditorSettings = GetDefault<UWaterEditorSettings>();
-			check(WaterEditorSettings != nullptr);
 			TSubclassOf<AWaterLandscapeBrush> WaterBrushClass = WaterEditorSettings->GetWaterManagerClass();
 			if (UClass* WaterBrushClassPtr = WaterBrushClass.Get())
 			{
@@ -262,8 +262,14 @@ void FWaterEditorModule::OnLevelActorAddedToWorld(AActor* Actor)
 						AWaterLandscapeBrush* NewBrush = (WaterBrushActorFactory != nullptr)
 							? Cast<AWaterLandscapeBrush>(WaterBrushActorFactory->CreateActor(ActorWorld, FoundLandscape->GetLevel(), FTransform(WaterZoneBounds.GetCenter()), SpawnParams))
 							: ActorWorld->SpawnActor<AWaterLandscapeBrush>(WaterBrushClassPtr, SpawnParams);
+
 						if (NewBrush)
 						{
+							if (!WaterBrushActorFactory)
+							{
+								UE_LOG(LogWaterEditor, Warning, TEXT("WaterManager Actor Factory could not be found! The newly spawned %s may have incorrect defaults!"), *NewBrush->GetActorLabel());
+							}
+
 							bHasWaterManager = true;
 							NewBrush->SetActorLabel(BrushActorString);
 							NewBrush->SetTargetLandscape(FoundLandscape);
@@ -284,41 +290,58 @@ void FWaterEditorModule::OnLevelActorAddedToWorld(AActor* Actor)
 		}
 
 		// Setup the water zone actor for this water body : 
+
 		const bool bHasZoneActor = !!TActorIterator<AWaterZone>(ActorWorld);
 		if ((WaterBodyActor != nullptr) && !bHasZoneActor)
 		{
-			FActorSpawnParameters SpawnParams;
-			SpawnParams.OverrideLevel = ActorWorld->PersistentLevel;
-			SpawnParams.bAllowDuringConstructionScript = true; // This can be called by construction script if the actor being added to the world is part of a blueprint, for example : 
-			AWaterZone* WaterZoneActor = ActorWorld->SpawnActor<AWaterZone>(AWaterZone::StaticClass(), SpawnParams);
-
-			// TODO [jonathan.bard] : when we can tag static meshes as "water ground", add these to the bounds
-			// Set a more sensible default location and extent so that the zone fully encapsulates the landscape if one exists.
-			if (WaterZoneBounds.IsValid)
+			TSubclassOf<AWaterZone> WaterZoneClass = WaterEditorSettings->GetWaterZoneClass();
+			if (UClass* WaterZoneClassPtr = WaterZoneClass.Get())
 			{
-				WaterZoneActor->SetActorLocation(WaterZoneBounds.GetCenter());
+				UActorFactory* WaterZoneActorFactory = GEditor->FindActorFactoryForActorClass(WaterZoneClassPtr);
 
-				// FBox::GetExtent returns the radius, SetZoneExtent expects diameter.
-				FVector2D NewExtent = 2 * FVector2D(WaterZoneBounds.GetExtent());
+				FActorSpawnParameters SpawnParams;
+				SpawnParams.OverrideLevel = ActorWorld->PersistentLevel;
+				SpawnParams.bAllowDuringConstructionScript = true; // This can be called by construction script if the actor being added to the world is part of a blueprint, for example : 
 
-				float ZoneExtentScale = WaterEditorModule::CVarOverrideNewWaterZoneScale.GetValueOnGameThread();
-				if (ZoneExtentScale == 0)
-				{
-					ZoneExtentScale = GetDefault<UWaterEditorSettings>()->WaterZoneActorDefaults.NewWaterZoneScale;
-				}
+				AWaterZone* WaterZoneActor = (WaterZoneActorFactory != nullptr)
+					? Cast<AWaterZone>(WaterZoneActorFactory->CreateActor(ActorWorld, Actor->GetLevel(), FTransform(WaterZoneBounds.GetCenter()), SpawnParams))
+					: ActorWorld->SpawnActor<AWaterZone>(WaterZoneClassPtr, SpawnParams);
 
-				if (ZoneExtentScale != 0)
+				if (WaterZoneActor)
 				{
-					NewExtent = FMath::Abs(ZoneExtentScale) * NewExtent;
-				}
+					if (!WaterZoneActorFactory)
+					{
+						UE_LOG(LogWaterEditor, Warning, TEXT("WaterZone Actor Factory could not be found! The newly spawned %s may have incorrect defaults!"), *WaterZoneActor->GetActorLabel());
+					}
 
-				WaterZoneActor->SetZoneExtent(NewExtent);
-			}
+					// TODO [jonathan.bard] : when we can tag static meshes as "water ground", add these to the bounds
+					// Set a more sensible default location and extent so that the zone fully encapsulates the landscape if one exists.
+					if (WaterZoneBounds.IsValid)
+					{
+						WaterZoneActor->SetActorLocation(WaterZoneBounds.GetCenter());
+
+						// FBox::GetExtent returns the radius, SetZoneExtent expects diameter.
+						FVector2D NewExtent = 2 * FVector2D(WaterZoneBounds.GetExtent());
 
-			// Set the defaults here because the actor factory isn't triggered on manual SpawnActor.
-			const FWaterZoneActorDefaults& WaterMeshActorDefaults = GetDefault<UWaterEditorSettings>()->WaterZoneActorDefaults;
-			WaterZoneActor->GetWaterMeshComponent()->FarDistanceMaterial = WaterMeshActorDefaults.GetFarDistanceMaterial();
-			WaterZoneActor->GetWaterMeshComponent()->FarDistanceMeshExtent = WaterMeshActorDefaults.FarDistanceMeshExtent;
+						float ZoneExtentScale = WaterEditorModule::CVarOverrideNewWaterZoneScale.GetValueOnGameThread();
+						if (ZoneExtentScale == 0)
+						{
+							ZoneExtentScale = GetDefault<UWaterEditorSettings>()->WaterZoneActorDefaults.NewWaterZoneScale;
+						}
+
+						if (ZoneExtentScale != 0)
+						{
+							NewExtent = FMath::Abs(ZoneExtentScale) * NewExtent;
+						}
+
+						WaterZoneActor->SetZoneExtent(NewExtent);
+					}
+				}
+			}
+			else
+			{
+				UE_LOG(LogWaterEditor, Warning, TEXT("Could not find Water Zone class %s to spawn"), *WaterEditorSettings->GetWaterZoneClassPath().GetAssetPathString());
+			}
 		}
 	}
 }
diff --git a/Engine/Plugins/Experimental/Water/Source/Editor/Private/WaterEditorSettings.cpp b/Engine/Plugins/Experimental/Water/Source/Editor/Private/WaterEditorSettings.cpp
--- a/Engine/Plugins/Experimental/Water/Source/Editor/Private/WaterEditorSettings.cpp
+++ b/Engine/Plugins/Experimental/Water/Source/Editor/Private/WaterEditorSettings.cpp
@@ -6,6 +6,7 @@
 #include "Materials/MaterialParameterCollection.h"
 #include "Materials/MaterialInterface.h"
 #include "WaterLandscapeBrush.h"
+#include "WaterZoneActor.h"
 #include "WaterWaves.h"
 #include "UObject/ConstructorHelpers.h"
 #include "Curves/CurveFloat.h"
@@ -135,6 +136,7 @@ UWaterEditorSettings::UWaterEditorSettings()
 	: TextureGroupForGeneratedTextures(TEXTUREGROUP_World)
 	, MaxWaterVelocityAndHeightTextureSize(2048)
 	, LandscapeMaterialParameterCollection(FSoftObjectPath(TEXT("/Landmass/Landscape/BlueprintBrushes/MPC/MPC_Landscape.MPC_Landscape")))
+	, WaterZoneClassPath(TEXT("/Script/Water.WaterZone"))
 	, WaterManagerClassPath(TEXT("/Script/WaterEditor.WaterBrushManager"))
 	, DefaultBrushAngleFalloffMaterial(FSoftObjectPath(TEXT("/Water/Materials/Brushes/MeshBrush_Angle.MeshBrush_Angle")))
 	, DefaultBrushIslandFalloffMaterial(FSoftObjectPath(TEXT("/Water/Materials/Brushes/MeshBrush_Island.MeshBrush_Island")))
@@ -166,6 +168,11 @@ UWaterEditorSettings::UWaterEditorSettings()
 	}
 }
 
+TSubclassOf<AWaterZone> UWaterEditorSettings::GetWaterZoneClass() const
+{
+	return WaterZoneClassPath.TryLoadClass<AWaterZone>();
+}
+
 TSubclassOf<AWaterLandscapeBrush> UWaterEditorSettings::GetWaterManagerClass() const
 {
 	return WaterManagerClassPath.TryLoadClass<AWaterLandscapeBrush>();
diff --git a/Engine/Plugins/Experimental/Water/Source/Editor/Public/WaterEditorSettings.h b/Engine/Plugins/Experimental/Water/Source/Editor/Public/WaterEditorSettings.h
--- a/Engine/Plugins/Experimental/Water/Source/Editor/Public/WaterEditorSettings.h
+++ b/Engine/Plugins/Experimental/Water/Source/Editor/Public/WaterEditorSettings.h
@@ -19,6 +19,7 @@ class UStaticMesh;
 class UWaterBodyComponent;
 class UWaterWavesBase;
 class AWaterLandscapeBrush;
+class AWaterZone;
 
 USTRUCT()
 struct WATEREDITOR_API FWaterBrushActorDefaults
@@ -185,6 +186,9 @@ public:
 
 	virtual FName GetCategoryName() const override { return FName(TEXT("Plugins")); }
 
+	TSubclassOf<AWaterZone> GetWaterZoneClass() const;
+	FSoftClassPath GetWaterZoneClassPath() const { return WaterZoneClassPath; }
+
 	TSubclassOf<AWaterLandscapeBrush> GetWaterManagerClass() const;
 	FSoftClassPath GetWaterManagerClassPath() const { return WaterManagerClassPath; }
 
@@ -278,6 +282,10 @@ private:
 	UPROPERTY(EditAnywhere, config, Category = Brush)
 	bool bUpdateWaterMeshDuringInteractiveChanges = false;
 
+	/** Class of the water zone to be used*/
+	UPROPERTY(EditAnywhere, config, Category = Water, meta = (MetaClass = "/Script/Water.WaterZone"))
+	FSoftClassPath WaterZoneClassPath;
+
 	/** Class of the water brush to be used in landscape */
 	UPROPERTY(EditAnywhere, config, Category = Brush, meta = (MetaClass = "/Script/WaterEditor.WaterLandscapeBrush"), AdvancedDisplay)
 	FSoftClassPath WaterManagerClassPath;
