Adjusting landscape LOD used for water info rendering to better match the previous value (of approximately twice the desired water info texture resolution)
Before:  https://horde.devtools.epicgames.com/testreport/63474e4ab757820b7a092baa   https://horde.devtools.epicgames.com/job/63470e0cf93be0f6347fd3a0
After:    https://horde.devtools.epicgames.com/testreport/63485c76fd22d1fbde70d5a1    https://horde.devtools.epicgames.com/job/63483020ad0f7e2f20e05b92

diff --git a/Engine/Plugins/Experimental/Water/Source/Runtime/Private/WaterInfoRendering.cpp b/Engine/Plugins/Experimental/Water/Source/Runtime/Private/WaterInfoRendering.cpp
--- a/Engine/Plugins/Experimental/Water/Source/Runtime/Private/WaterInfoRendering.cpp
+++ b/Engine/Plugins/Experimental/Water/Source/Runtime/Private/WaterInfoRendering.cpp
@@ -368,8 +368,8 @@ private:
 					}
 
 					// Derived from:
-					// LandscapeComponentUnitsPerVertex * 2 ^ (LODLevel) <= WaterInfoUnitsPerPixel
-					OptimalLODLevel = FMath::FloorLog2(WaterInfoUnitsPerPixel / LandscapeComponentUnitsPerVertex);
+					// LandscapeComponentUnitsPerVertex * 2 ^ (LODLevel) <= WaterInfoUnitsPerPixel * 0.5
+					OptimalLODLevel = FMath::Max(0, int32(FMath::FloorLog2(WaterInfoUnitsPerPixel / LandscapeComponentUnitsPerVertex)) - 1);
 
 					break;
 				}
