Fixed include paths in ../Plugins/Experimental/...

diff --git a/Engine/Plugins/Experimental/BackChannel/Source/BackChannel/Private/Protocol/BackChannelOSCBundle.cpp b/Engine/Plugins/Experimental/BackChannel/Source/BackChannel/Private/Protocol/BackChannelOSCBundle.cpp
--- a/Engine/Plugins/Experimental/BackChannel/Source/BackChannel/Private/Protocol/BackChannelOSCBundle.cpp
+++ b/Engine/Plugins/Experimental/BackChannel/Source/BackChannel/Private/Protocol/BackChannelOSCBundle.cpp
@@ -1,7 +1,7 @@
 // Copyright Epic Games, Inc. All Rights Reserved.
 
 #include "BackChannel/Protocol/OSC/BackChannelOSCBundle.h"
-#include "BackChannel/Private/BackChannelCommon.h"
+#include "BackChannelCommon.h"
 
 
 const ANSICHAR* FBackChannelOSCBundle::BundleHeader = "#bundle";
diff --git a/Engine/Plugins/Experimental/BackChannel/Source/BackChannel/Private/Protocol/BackChannelOSCConnection.cpp b/Engine/Plugins/Experimental/BackChannel/Source/BackChannel/Private/Protocol/BackChannelOSCConnection.cpp
--- a/Engine/Plugins/Experimental/BackChannel/Source/BackChannel/Private/Protocol/BackChannelOSCConnection.cpp
+++ b/Engine/Plugins/Experimental/BackChannel/Source/BackChannel/Private/Protocol/BackChannelOSCConnection.cpp
@@ -1,13 +1,13 @@
 // Copyright Epic Games, Inc. All Rights Reserved.
 
 #include "BackChannel/Protocol/OSC/BackChannelOSCConnection.h"
-#include "BackChannel/Private/BackChannelCommon.h"
-#include "BackChannel/Transport/IBackChannelTransport.h"
 #include "BackChannel/Protocol/OSC/BackChannelOSC.h"
-#include "Sockets.h"
-#include "HAL/PlatformTime.h"
+#include "BackChannel/Transport/IBackChannelTransport.h"
+#include "BackChannelCommon.h"
 #include "HAL/PlatformProcess.h"
+#include "HAL/PlatformTime.h"
 #include "Misc/ScopeLock.h"
+#include "Sockets.h"
 
 FBackChannelOSCConnection::FBackChannelOSCConnection(TSharedRef<IBackChannelSocketConnection> InConnection)
 	: Connection(InConnection)
diff --git a/Engine/Plugins/Experimental/BackChannel/Source/BackChannel/Private/Protocol/BackChannelOSCMessage.cpp b/Engine/Plugins/Experimental/BackChannel/Source/BackChannel/Private/Protocol/BackChannelOSCMessage.cpp
--- a/Engine/Plugins/Experimental/BackChannel/Source/BackChannel/Private/Protocol/BackChannelOSCMessage.cpp
+++ b/Engine/Plugins/Experimental/BackChannel/Source/BackChannel/Private/Protocol/BackChannelOSCMessage.cpp
@@ -1,7 +1,7 @@
 // Copyright Epic Games, Inc. All Rights Reserved.
 
 #include "BackChannel/Protocol/OSC/BackChannelOSCMessage.h"
-#include "BackChannel/Private/BackChannelCommon.h"
+#include "BackChannelCommon.h"
 
 bool FBackChannelOSCMessage::bIsLegacyConnection = false;
 
@@ -347,4 +347,4 @@ FBackChannelOSCMessage& operator << (FBackChannelOSCMessage& Msg, FString& Value
 	return Msg;
 }
 
-*/
\ No newline at end of file
+*/
diff --git a/Engine/Plugins/Experimental/BackChannel/Source/BackChannel/Private/Protocol/BackChannelOSCPacket.cpp b/Engine/Plugins/Experimental/BackChannel/Source/BackChannel/Private/Protocol/BackChannelOSCPacket.cpp
--- a/Engine/Plugins/Experimental/BackChannel/Source/BackChannel/Private/Protocol/BackChannelOSCPacket.cpp
+++ b/Engine/Plugins/Experimental/BackChannel/Source/BackChannel/Private/Protocol/BackChannelOSCPacket.cpp
@@ -1,9 +1,9 @@
 // Copyright Epic Games, Inc. All Rights Reserved.
 
 #include "BackChannel/Protocol/OSC/BackChannelOSCPacket.h"
-#include "BackChannel/Private/BackChannelCommon.h"
 #include "BackChannel/Protocol/OSC/BackChannelOSCBundle.h"
 #include "BackChannel/Protocol/OSC/BackChannelOSCMessage.h"
+#include "BackChannelCommon.h"
 
 
 
diff --git a/Engine/Plugins/Experimental/BackChannel/Source/BackChannel/Private/Protocol/TestBackChannelOSC.cpp b/Engine/Plugins/Experimental/BackChannel/Source/BackChannel/Private/Protocol/TestBackChannelOSC.cpp
--- a/Engine/Plugins/Experimental/BackChannel/Source/BackChannel/Private/Protocol/TestBackChannelOSC.cpp
+++ b/Engine/Plugins/Experimental/BackChannel/Source/BackChannel/Private/Protocol/TestBackChannelOSC.cpp
@@ -1,8 +1,8 @@
 // Copyright Epic Games, Inc. All Rights Reserved.
 
-#include "BackChannel/Private/BackChannelCommon.h"
-#include "BackChannel/Transport/IBackChannelTransport.h"
 #include "BackChannel/Protocol/OSC/BackChannelOSC.h"
+#include "BackChannel/Transport/IBackChannelTransport.h"
+#include "BackChannelCommon.h"
 #include "Misc/AutomationTest.h"
 
 
diff --git a/Engine/Plugins/Experimental/BackChannel/Source/BackChannel/Private/Transport/BackChannelConnection.cpp b/Engine/Plugins/Experimental/BackChannel/Source/BackChannel/Private/Transport/BackChannelConnection.cpp
--- a/Engine/Plugins/Experimental/BackChannel/Source/BackChannel/Private/Transport/BackChannelConnection.cpp
+++ b/Engine/Plugins/Experimental/BackChannel/Source/BackChannel/Private/Transport/BackChannelConnection.cpp
@@ -1,13 +1,13 @@
 // Copyright Epic Games, Inc. All Rights Reserved.
 
-#include "BackChannel/Private/Transport/BackChannelConnection.h"
-#include "BackChannel/Private/BackChannelCommon.h"
+#include "Transport/BackChannelConnection.h"
 #include "BackChannel/Transport/IBackChannelTransport.h"
+#include "BackChannelCommon.h"
 #include "Common/TcpSocketBuilder.h"
-#include "HAL/IConsoleManager.h"
-#include "Misc/ScopeLock.h"
 #include "CoreGlobals.h"
+#include "HAL/IConsoleManager.h"
 #include "Misc/ConfigCacheIni.h"
+#include "Misc/ScopeLock.h"
 
 DECLARE_STATS_GROUP(TEXT("BackChannel"), STATGROUP_BackChannel, STATCAT_Advanced);
 
diff --git a/Engine/Plugins/Experimental/BackChannel/Source/BackChannel/Private/Transport/BackChannelTransport.cpp b/Engine/Plugins/Experimental/BackChannel/Source/BackChannel/Private/Transport/BackChannelTransport.cpp
--- a/Engine/Plugins/Experimental/BackChannel/Source/BackChannel/Private/Transport/BackChannelTransport.cpp
+++ b/Engine/Plugins/Experimental/BackChannel/Source/BackChannel/Private/Transport/BackChannelTransport.cpp
@@ -1,8 +1,8 @@
 // Copyright Epic Games, Inc. All Rights Reserved.
 
-#include "BackChannel/Private/BackChannelCommon.h"
 #include "BackChannel/Transport/IBackChannelTransport.h"
-#include "BackChannel/Private/Transport/BackChannelConnection.h"
+#include "BackChannelCommon.h"
+#include "Transport/BackChannelConnection.h"
 
 const TCHAR* BackChannelTransport_TCP = TEXT("BackChannelTCP");
 
diff --git a/Engine/Plugins/Experimental/BackChannel/Source/BackChannel/Private/Transport/TestBackChannelTransport.cpp b/Engine/Plugins/Experimental/BackChannel/Source/BackChannel/Private/Transport/TestBackChannelTransport.cpp
--- a/Engine/Plugins/Experimental/BackChannel/Source/BackChannel/Private/Transport/TestBackChannelTransport.cpp
+++ b/Engine/Plugins/Experimental/BackChannel/Source/BackChannel/Private/Transport/TestBackChannelTransport.cpp
@@ -1,10 +1,10 @@
 // Copyright Epic Games, Inc. All Rights Reserved.
 
-#include "BackChannel/Private/BackChannelCommon.h"
 #include "BackChannel/Transport/IBackChannelTransport.h"
 #include "BackChannel/Utils/BackChannelThreadedConnection.h"
-#include "Misc/AutomationTest.h"
+#include "BackChannelCommon.h"
 #include "HAL/PlatformProcess.h"
+#include "Misc/AutomationTest.h"
 
 PRAGMA_DISABLE_OPTIMIZATION
 
diff --git a/Engine/Plugins/Experimental/BackChannel/Source/BackChannel/Utils/BackChannelThreadedConnection.cpp b/Engine/Plugins/Experimental/BackChannel/Source/BackChannel/Utils/BackChannelThreadedConnection.cpp
--- a/Engine/Plugins/Experimental/BackChannel/Source/BackChannel/Utils/BackChannelThreadedConnection.cpp
+++ b/Engine/Plugins/Experimental/BackChannel/Source/BackChannel/Utils/BackChannelThreadedConnection.cpp
@@ -1,7 +1,7 @@
 // Copyright Epic Games, Inc. All Rights Reserved.
 
 #include "BackChannel/Utils/BackChannelThreadedConnection.h"
-#include "BackChannel/Private/BackChannelCommon.h"
+#include "BackChannelCommon.h"
 #include "HAL/RunnableThread.h"
 #include "HAL/ThreadSafeBool.h"
 #include "Misc/ScopeLock.h"
diff --git a/Engine/Plugins/Experimental/ChaosNiagara/Source/ChaosNiagara/ChaosNiagara.Build.cs b/Engine/Plugins/Experimental/ChaosNiagara/Source/ChaosNiagara/ChaosNiagara.Build.cs
--- a/Engine/Plugins/Experimental/ChaosNiagara/Source/ChaosNiagara/ChaosNiagara.Build.cs
+++ b/Engine/Plugins/Experimental/ChaosNiagara/Source/ChaosNiagara/ChaosNiagara.Build.cs
@@ -1,11 +1,17 @@
 // Copyright Epic Games, Inc. All Rights Reserved.
 
+using System.IO;
 using UnrealBuildTool;
 
 public class ChaosNiagara : ModuleRules
 {
 	public ChaosNiagara(ReadOnlyTargetRules Target) : base(Target)
 	{
+		PrivateIncludePaths.AddRange(new string[] {
+			System.IO.Path.Combine(GetModuleDirectory("ControlRig"), "Private"),
+			System.IO.Path.Combine(GetModuleDirectory("Niagara"), "Private"),
+		});
+
 		PrivateDependencyModuleNames.AddRange(
 			new string[]
 			{
diff --git a/Engine/Plugins/Experimental/ChaosNiagara/Source/ChaosNiagara/Private/NiagaraDataInterfaceChaosDestruction.cpp b/Engine/Plugins/Experimental/ChaosNiagara/Source/ChaosNiagara/Private/NiagaraDataInterfaceChaosDestruction.cpp
--- a/Engine/Plugins/Experimental/ChaosNiagara/Source/ChaosNiagara/Private/NiagaraDataInterfaceChaosDestruction.cpp
+++ b/Engine/Plugins/Experimental/ChaosNiagara/Source/ChaosNiagara/Private/NiagaraDataInterfaceChaosDestruction.cpp
@@ -1,19 +1,19 @@
 // Copyright Epic Games, Inc. All Rights Reserved.
 
 #include "NiagaraDataInterfaceChaosDestruction.h"
-#include "NiagaraTypes.h"
+#include "Chaos/CollisionResolutionTypes.h"
+#include "Chaos/PBDCollisionConstraintsUtil.h"
 #include "Misc/FileHelper.h"
+#include "NiagaraComponent.h"
 #include "NiagaraGpuComputeDispatchInterface.h"
 #include "NiagaraShader.h"
 #include "NiagaraShaderParametersBuilder.h"
-#include "ShaderParameterUtils.h"
-#include "PhysicsSolver.h"
-#include "Niagara/Private/NiagaraStats.h"
-#include "Chaos/PBDCollisionConstraintsUtil.h"
-#include "PhysicsProxy/GeometryCollectionPhysicsProxy.h"
-#include "Chaos/CollisionResolutionTypes.h"
+#include "NiagaraStats.h"
+#include "NiagaraTypes.h"
 #include "PhysicalMaterials/PhysicalMaterial.h"
-#include "NiagaraComponent.h"
+#include "PhysicsProxy/GeometryCollectionPhysicsProxy.h"
+#include "PhysicsSolver.h"
+#include "ShaderParameterUtils.h"
 
 #include <memory>
 #include "GeometryCollection/GeometryCollectionComponent.h"
diff --git a/Engine/Plugins/Experimental/CommonUI/Source/CommonUI/Private/CommonLazyWidget.cpp b/Engine/Plugins/Experimental/CommonUI/Source/CommonUI/Private/CommonLazyWidget.cpp
--- a/Engine/Plugins/Experimental/CommonUI/Source/CommonUI/Private/CommonLazyWidget.cpp
+++ b/Engine/Plugins/Experimental/CommonUI/Source/CommonUI/Private/CommonLazyWidget.cpp
@@ -1,9 +1,9 @@
 // Copyright Epic Games, Inc. All Rights Reserved.
 
 #include "CommonLazyWidget.h"
-#include "CommonWidgetPaletteCategories.h"
+#include "CommonActivatableWidget.h"
 #include "CommonUISettings.h"
-#include "../Public/CommonActivatableWidget.h"
+#include "CommonWidgetPaletteCategories.h"
 
 #include UE_INLINE_GENERATED_CPP_BY_NAME(CommonLazyWidget)
 
diff --git a/Engine/Plugins/Experimental/CommonUI/Source/CommonUIEditor/CommonUIEditor.Build.cs b/Engine/Plugins/Experimental/CommonUI/Source/CommonUIEditor/CommonUIEditor.Build.cs
--- a/Engine/Plugins/Experimental/CommonUI/Source/CommonUIEditor/CommonUIEditor.Build.cs
+++ b/Engine/Plugins/Experimental/CommonUI/Source/CommonUIEditor/CommonUIEditor.Build.cs
@@ -34,8 +34,8 @@ public class CommonUIEditor : ModuleRules
 		PrivateIncludePaths.AddRange(
 			new string[]
 			{
-				"CommonUI/Private",
-                "CommonUIEditor/Private",
+				System.IO.Path.Combine(GetModuleDirectory("CommonUI"), "Private"),
+				System.IO.Path.Combine(GetModuleDirectory("CommonUIEditor"), "Private"),
 			}
 		);
 
@@ -43,6 +43,7 @@ public class CommonUIEditor : ModuleRules
 			new string[] {
 				"AssetRegistry",
 				"AssetTools",
+				"DataTableEditor",
 			}
 		);
 
diff --git a/Engine/Plugins/Experimental/CommonUI/Source/CommonUIEditor/Private/CommonAssetTypeActions_GenericInputActionDataTable.cpp b/Engine/Plugins/Experimental/CommonUI/Source/CommonUIEditor/Private/CommonAssetTypeActions_GenericInputActionDataTable.cpp
--- a/Engine/Plugins/Experimental/CommonUI/Source/CommonUIEditor/Private/CommonAssetTypeActions_GenericInputActionDataTable.cpp
+++ b/Engine/Plugins/Experimental/CommonUI/Source/CommonUIEditor/Private/CommonAssetTypeActions_GenericInputActionDataTable.cpp
@@ -1,10 +1,10 @@
 // Copyright Epic Games, Inc. All Rights Reserved.
 
 #include "CommonAssetTypeActions_GenericInputActionDataTable.h"
-#include "ToolMenus.h"
-#include "Misc/MessageDialog.h"
-#include "Editor/DataTableEditor/Public/DataTableEditorModule.h"
+#include "DataTableEditorModule.h"
 #include "Input/CommonGenericInputActionDataTable.h"
+#include "Misc/MessageDialog.h"
+#include "ToolMenus.h"
 
 #define LOCTEXT_NAMESPACE "AssetTypeActions"
 
diff --git a/Engine/Plugins/Experimental/CommonUI/Source/CommonUIEditor/Private/CommonVideoPlayerCustomization.h b/Engine/Plugins/Experimental/CommonUI/Source/CommonUIEditor/Private/CommonVideoPlayerCustomization.h
--- a/Engine/Plugins/Experimental/CommonUI/Source/CommonUIEditor/Private/CommonVideoPlayerCustomization.h
+++ b/Engine/Plugins/Experimental/CommonUI/Source/CommonUIEditor/Private/CommonVideoPlayerCustomization.h
@@ -2,10 +2,10 @@
 
 #pragma once
 
-#include "Editor/PropertyEditor/Public/PropertyEditorModule.h"
+#include "IDetailCustomization.h"
 #include "ITransportControl.h"
+#include "PropertyEditorModule.h"
 #include "Types/SlateEnums.h"
-#include "IDetailCustomization.h"
 
 class IPropertyHandle;
 class UCommonVideoPlayer;
@@ -37,4 +37,4 @@ private:
 	FReply HandleToggleMuteClicked() const;
 
 	TWeakObjectPtr<UCommonVideoPlayer> VideoPlayer;
-};
\ No newline at end of file
+};
diff --git a/Engine/Plugins/Experimental/GeometryCollectionPlugin/Source/GeometryCollectionNodes/Private/Dataflow/GeometryCollectionNodes.cpp b/Engine/Plugins/Experimental/GeometryCollectionPlugin/Source/GeometryCollectionNodes/Private/Dataflow/GeometryCollectionNodes.cpp
--- a/Engine/Plugins/Experimental/GeometryCollectionPlugin/Source/GeometryCollectionNodes/Private/Dataflow/GeometryCollectionNodes.cpp
+++ b/Engine/Plugins/Experimental/GeometryCollectionPlugin/Source/GeometryCollectionNodes/Private/Dataflow/GeometryCollectionNodes.cpp
@@ -13,13 +13,13 @@
 #include "Templates/SharedPointer.h"
 #include "UObject/UnrealTypePrivate.h"
 
+#include "EngineGlobals.h"
 #include "GeometryCollection/GeometryCollectionAlgo.h"
 #include "GeometryCollection/GeometryCollectionClusteringUtility.h"
 #include "GeometryCollection/GeometryCollectionConvexUtility.h"
-#include "Voronoi/Voronoi.h"
-#include "PlanarCut.h"
-#include "Runtime/Engine/Public/EngineGlobals.h"
 #include "Math/UnrealMathUtility.h"
+#include "PlanarCut.h"
+#include "Voronoi/Voronoi.h"
 
 #include UE_INLINE_GENERATED_CPP_BY_NAME(GeometryCollectionNodes)
 
diff --git a/Engine/Plugins/Experimental/MeshModelingToolsetExp/Source/MeshModelingToolsEditorOnlyExp/Private/MeshTangentsTool.cpp b/Engine/Plugins/Experimental/MeshModelingToolsetExp/Source/MeshModelingToolsEditorOnlyExp/Private/MeshTangentsTool.cpp
--- a/Engine/Plugins/Experimental/MeshModelingToolsetExp/Source/MeshModelingToolsEditorOnlyExp/Private/MeshTangentsTool.cpp
+++ b/Engine/Plugins/Experimental/MeshModelingToolsetExp/Source/MeshModelingToolsEditorOnlyExp/Private/MeshTangentsTool.cpp
@@ -11,8 +11,8 @@
 #include "ToolDataVisualizer.h"
 
 #include "AssetUtils/MeshDescriptionUtil.h"
-#include "Engine/Classes/Engine/StaticMesh.h"
-#include "Engine/Classes/Components/StaticMeshComponent.h"
+#include "Components/StaticMeshComponent.h"
+#include "Engine/StaticMesh.h"
 
 #include "TargetInterfaces/MaterialProvider.h"
 #include "TargetInterfaces/MeshDescriptionCommitter.h"
diff --git a/Engine/Plugins/Experimental/MeshModelingToolsetExp/Source/MeshModelingToolsEditorOnlyExp/Public/MeshToVolumeTool.h b/Engine/Plugins/Experimental/MeshModelingToolsetExp/Source/MeshModelingToolsEditorOnlyExp/Public/MeshToVolumeTool.h
--- a/Engine/Plugins/Experimental/MeshModelingToolsetExp/Source/MeshModelingToolsEditorOnlyExp/Public/MeshToVolumeTool.h
+++ b/Engine/Plugins/Experimental/MeshModelingToolsetExp/Source/MeshModelingToolsEditorOnlyExp/Public/MeshToVolumeTool.h
@@ -4,16 +4,16 @@
 
 #include "CoreMinimal.h"
 
+#include "BaseTools/SingleSelectionMeshEditingTool.h"
 #include "ConversionUtils/DynamicMeshToVolume.h"
+#include "Drawing/LineSetComponent.h"
+#include "DynamicMesh/DynamicMesh3.h"
+#include "Engine/BlockingVolume.h"
 #include "GameFramework/Volume.h"
-#include "Engine/Classes/Engine/BlockingVolume.h"
-#include "SingleSelectionTool.h"
 #include "InteractiveToolBuilder.h"
-#include "DynamicMesh/DynamicMesh3.h"
 #include "PreviewMesh.h"
-#include "Drawing/LineSetComponent.h"
 #include "PropertySets/OnAcceptProperties.h"
-#include "BaseTools/SingleSelectionMeshEditingTool.h"
+#include "SingleSelectionTool.h"
 #include "MeshToVolumeTool.generated.h"
 
 /**
diff --git a/Engine/Plugins/Experimental/MeshModelingToolsetExp/Source/MeshModelingToolsExp/Private/BakeMultiMeshAttributeMapsTool.cpp b/Engine/Plugins/Experimental/MeshModelingToolsetExp/Source/MeshModelingToolsExp/Private/BakeMultiMeshAttributeMapsTool.cpp
--- a/Engine/Plugins/Experimental/MeshModelingToolsetExp/Source/MeshModelingToolsExp/Private/BakeMultiMeshAttributeMapsTool.cpp
+++ b/Engine/Plugins/Experimental/MeshModelingToolsetExp/Source/MeshModelingToolsExp/Private/BakeMultiMeshAttributeMapsTool.cpp
@@ -26,7 +26,7 @@
 #include "ModelingToolTargetUtil.h"
 
 // required to pass UStaticMesh asset so we can save at same location
-#include "Engine/Classes/Engine/StaticMesh.h"
+#include "Engine/StaticMesh.h"
 
 #include UE_INLINE_GENERATED_CPP_BY_NAME(BakeMultiMeshAttributeMapsTool)
 
diff --git a/Engine/Plugins/Experimental/MeshModelingToolsetExp/Source/MeshModelingToolsExp/Private/DeformMeshPolygonsTool.cpp b/Engine/Plugins/Experimental/MeshModelingToolsetExp/Source/MeshModelingToolsExp/Private/DeformMeshPolygonsTool.cpp
--- a/Engine/Plugins/Experimental/MeshModelingToolsetExp/Source/MeshModelingToolsExp/Private/DeformMeshPolygonsTool.cpp
+++ b/Engine/Plugins/Experimental/MeshModelingToolsetExp/Source/MeshModelingToolsExp/Private/DeformMeshPolygonsTool.cpp
@@ -2,15 +2,15 @@
 
 #include "DeformMeshPolygonsTool.h"
 
-#include "InteractiveToolManager.h"
+#include "Curves/RichCurve.h"
 #include "DynamicMesh/MeshNormals.h"
-#include "ModelingOperators/Public/ModelingTaskTypes.h"
+#include "InteractiveToolManager.h"
+#include "ModelingTaskTypes.h"
+#include "ModelingToolTargetUtil.h"
 #include "Solvers/ConstrainedMeshDeformer.h"
 #include "ToolBuilderUtil.h"
-#include "ToolSetupUtil.h"
 #include "ToolSceneQueriesUtil.h"
-#include "ModelingToolTargetUtil.h"
-#include "Curves/RichCurve.h"
+#include "ToolSetupUtil.h"
 
 #include UE_INLINE_GENERATED_CPP_BY_NAME(DeformMeshPolygonsTool)
 
diff --git a/Engine/Plugins/Experimental/MeshModelingToolsetExp/Source/MeshModelingToolsExp/Private/DisplaceMeshTool.cpp b/Engine/Plugins/Experimental/MeshModelingToolsetExp/Source/MeshModelingToolsExp/Private/DisplaceMeshTool.cpp
--- a/Engine/Plugins/Experimental/MeshModelingToolsetExp/Source/MeshModelingToolsExp/Private/DisplaceMeshTool.cpp
+++ b/Engine/Plugins/Experimental/MeshModelingToolsetExp/Source/MeshModelingToolsExp/Private/DisplaceMeshTool.cpp
@@ -23,8 +23,8 @@
 
 // needed to disable normals recalculation on the underlying asset
 #include "AssetUtils/MeshDescriptionUtil.h"
-#include "Engine/Classes/Engine/StaticMesh.h"
-#include "Engine/Classes/Components/StaticMeshComponent.h"
+#include "Components/StaticMeshComponent.h"
+#include "Engine/StaticMesh.h"
 
 using namespace UE::Geometry;
 
diff --git a/Engine/Plugins/Experimental/MeshModelingToolsetExp/Source/MeshModelingToolsExp/Private/Physics/CollisionPropertySets.cpp b/Engine/Plugins/Experimental/MeshModelingToolsetExp/Source/MeshModelingToolsExp/Private/Physics/CollisionPropertySets.cpp
--- a/Engine/Plugins/Experimental/MeshModelingToolsetExp/Source/MeshModelingToolsExp/Private/Physics/CollisionPropertySets.cpp
+++ b/Engine/Plugins/Experimental/MeshModelingToolsetExp/Source/MeshModelingToolsExp/Private/Physics/CollisionPropertySets.cpp
@@ -3,7 +3,7 @@
 #include "Physics/CollisionPropertySets.h"
 #include "Physics/PhysicsDataCollection.h"
 
-#include "Engine/Classes/PhysicsEngine/BodySetup.h"
+#include "PhysicsEngine/BodySetup.h"
 
 #include UE_INLINE_GENERATED_CPP_BY_NAME(CollisionPropertySets)
 
diff --git a/Engine/Plugins/Experimental/MeshModelingToolsetExp/Source/MeshModelingToolsExp/Private/Physics/PhysicsInspectorTool.cpp b/Engine/Plugins/Experimental/MeshModelingToolsetExp/Source/MeshModelingToolsExp/Private/Physics/PhysicsInspectorTool.cpp
--- a/Engine/Plugins/Experimental/MeshModelingToolsetExp/Source/MeshModelingToolsExp/Private/Physics/PhysicsInspectorTool.cpp
+++ b/Engine/Plugins/Experimental/MeshModelingToolsetExp/Source/MeshModelingToolsExp/Private/Physics/PhysicsInspectorTool.cpp
@@ -12,8 +12,8 @@
 #include "Physics/CollisionGeometryVisualization.h"
 
 // physics data
-#include "Engine/Classes/PhysicsEngine/BodySetup.h"
-#include "Engine/Classes/PhysicsEngine/AggregateGeom.h"
+#include "PhysicsEngine/AggregateGeom.h"
+#include "PhysicsEngine/BodySetup.h"
 
 #include "TargetInterfaces/PrimitiveComponentBackedTarget.h"
 #include "TargetInterfaces/PhysicsDataSource.h"
diff --git a/Engine/Plugins/Experimental/MeshModelingToolsetExp/Source/MeshModelingToolsExp/Private/Physics/SetCollisionGeometryTool.cpp b/Engine/Plugins/Experimental/MeshModelingToolsetExp/Source/MeshModelingToolsExp/Private/Physics/SetCollisionGeometryTool.cpp
--- a/Engine/Plugins/Experimental/MeshModelingToolsetExp/Source/MeshModelingToolsExp/Private/Physics/SetCollisionGeometryTool.cpp
+++ b/Engine/Plugins/Experimental/MeshModelingToolsetExp/Source/MeshModelingToolsExp/Private/Physics/SetCollisionGeometryTool.cpp
@@ -17,9 +17,9 @@
 #include "Physics/CollisionGeometryVisualization.h"
 
 // physics data
-#include "Engine/Classes/Engine/StaticMesh.h"
-#include "Engine/Classes/Components/StaticMeshComponent.h"
-#include "Engine/Classes/PhysicsEngine/BodySetup.h"
+#include "Components/StaticMeshComponent.h"
+#include "Engine/StaticMesh.h"
+#include "PhysicsEngine/BodySetup.h"
 
 #include "Async/ParallelFor.h"
 
diff --git a/Engine/Plugins/Experimental/MeshModelingToolsetExp/Source/MeshModelingToolsExp/Public/BakeMeshAttributeMapsToolBase.h b/Engine/Plugins/Experimental/MeshModelingToolsetExp/Source/MeshModelingToolsExp/Public/BakeMeshAttributeMapsToolBase.h
--- a/Engine/Plugins/Experimental/MeshModelingToolsetExp/Source/MeshModelingToolsExp/Public/BakeMeshAttributeMapsToolBase.h
+++ b/Engine/Plugins/Experimental/MeshModelingToolsetExp/Source/MeshModelingToolsExp/Public/BakeMeshAttributeMapsToolBase.h
@@ -2,18 +2,18 @@
 
 #pragma once
 
-#include "CoreMinimal.h"
-#include "Engine/Classes/Engine/Texture2D.h"
-#include "Materials/MaterialInstanceDynamic.h"
-#include "InteractiveToolManager.h"
 #include "AssetUtils/Texture2DBuilder.h"
+#include "BakeMeshAttributeTool.h"
+#include "CoreMinimal.h"
 #include "DynamicMesh/DynamicMesh3.h"
+#include "Engine/Texture2D.h"
 #include "Image/ImageDimensions.h"
-#include "Sampling/MeshMapBaker.h"
-#include "ModelingOperators.h"
+#include "InteractiveToolManager.h"
+#include "Materials/MaterialInstanceDynamic.h"
 #include "MeshOpPreviewHelpers.h"
+#include "ModelingOperators.h"
 #include "PreviewMesh.h"
-#include "BakeMeshAttributeTool.h"
+#include "Sampling/MeshMapBaker.h"
 #include "BakeMeshAttributeMapsToolBase.generated.h"
 
 /**
diff --git a/Engine/Plugins/Experimental/MeshModelingToolsetExp/Source/MeshModelingToolsExp/Public/BakeMeshAttributeTool.h b/Engine/Plugins/Experimental/MeshModelingToolsetExp/Source/MeshModelingToolsExp/Public/BakeMeshAttributeTool.h
--- a/Engine/Plugins/Experimental/MeshModelingToolsetExp/Source/MeshModelingToolsExp/Public/BakeMeshAttributeTool.h
+++ b/Engine/Plugins/Experimental/MeshModelingToolsetExp/Source/MeshModelingToolsExp/Public/BakeMeshAttributeTool.h
@@ -2,17 +2,17 @@
 
 #pragma once
 
-#include "CoreMinimal.h"
-#include "Engine/Classes/Engine/Texture2D.h"
-#include "Materials/MaterialInstanceDynamic.h"
+#include "BakeMeshAttributeToolCommon.h"
 #include "BaseTools/MultiSelectionMeshEditingTool.h"
-#include "InteractiveToolQueryInterfaces.h" // for UInteractiveToolExclusiveToolAPI
+#include "CoreMinimal.h"
 #include "DynamicMesh/DynamicMesh3.h"
+#include "Engine/Texture2D.h"
+#include "InteractiveToolQueryInterfaces.h" // for UInteractiveToolExclusiveToolAPI
+#include "Materials/MaterialInstanceDynamic.h"
 #include "PreviewMesh.h"
-#include "TargetInterfaces/StaticMeshBackedTarget.h"
-#include "TargetInterfaces/SkeletalMeshBackedTarget.h"
 #include "TargetInterfaces/DynamicMeshSource.h"
-#include "BakeMeshAttributeToolCommon.h"
+#include "TargetInterfaces/SkeletalMeshBackedTarget.h"
+#include "TargetInterfaces/StaticMeshBackedTarget.h"
 #include "BakeMeshAttributeTool.generated.h"
 
 /**
diff --git a/Engine/Plugins/Experimental/MeshModelingToolsetExp/Source/MeshModelingToolsExp/Public/DisplaceMeshTool.h b/Engine/Plugins/Experimental/MeshModelingToolsetExp/Source/MeshModelingToolsExp/Public/DisplaceMeshTool.h
--- a/Engine/Plugins/Experimental/MeshModelingToolsetExp/Source/MeshModelingToolsExp/Public/DisplaceMeshTool.h
+++ b/Engine/Plugins/Experimental/MeshModelingToolsetExp/Source/MeshModelingToolsExp/Public/DisplaceMeshTool.h
@@ -2,17 +2,17 @@
 
 #pragma once
 
+#include "BaseTools/SingleSelectionMeshEditingTool.h"
 #include "CoreMinimal.h"
-#include "UObject/NoExportTypes.h"
-#include "Engine/Classes/Engine/Texture2D.h"
-#include "SingleSelectionTool.h"
-#include "InteractiveToolBuilder.h"
 #include "DynamicMesh/DynamicMesh3.h"
 #include "DynamicMesh/MeshNormals.h"
+#include "Engine/Texture2D.h"
+#include "InteractiveToolBuilder.h"
 #include "MeshOpPreviewHelpers.h"
+#include "SingleSelectionTool.h"
 #include "Spatial/SampledScalarField2.h"
+#include "UObject/NoExportTypes.h"
 #include "WeightMapUtil.h"
-#include "BaseTools/SingleSelectionMeshEditingTool.h"
 #include "DisplaceMeshTool.generated.h"
 
 struct FMeshDescription;
diff --git a/Engine/Plugins/Experimental/MeshModelingToolsetExp/Source/MeshModelingToolsExp/Public/EditUVIslandsTool.h b/Engine/Plugins/Experimental/MeshModelingToolsetExp/Source/MeshModelingToolsExp/Public/EditUVIslandsTool.h
--- a/Engine/Plugins/Experimental/MeshModelingToolsetExp/Source/MeshModelingToolsExp/Public/EditUVIslandsTool.h
+++ b/Engine/Plugins/Experimental/MeshModelingToolsetExp/Source/MeshModelingToolsExp/Public/EditUVIslandsTool.h
@@ -2,18 +2,18 @@
 
 #pragma once
 
-#include "CoreMinimal.h"
 #include "BaseTools/MeshSurfacePointMeshEditingTool.h"
-#include "Properties/MeshMaterialProperties.h"
+#include "Changes/MeshVertexChange.h"
 #include "Components/DynamicMeshComponent.h"
+#include "CoreMinimal.h"
 #include "DynamicMesh/DynamicMeshAABBTree3.h"
-#include "ToolDataVisualizer.h"
-#include "Changes/MeshVertexChange.h"
 #include "GroupTopology.h"
+#include "ModelingTaskTypes.h"
+#include "Properties/MeshMaterialProperties.h"
 #include "Selection/GroupTopologySelector.h"
-#include "ModelingOperators/Public/ModelingTaskTypes.h"
-#include "Transforms/MultiTransformer.h"
 #include "Selection/PolygonSelectionMechanic.h"
+#include "ToolDataVisualizer.h"
+#include "Transforms/MultiTransformer.h"
 #include "EditUVIslandsTool.generated.h"
 
 class FMeshVertexChangeBuilder;
diff --git a/Engine/Plugins/Experimental/MeshModelingToolsetExp/Source/MeshModelingToolsExp/Public/Physics/CollisionPropertySets.h b/Engine/Plugins/Experimental/MeshModelingToolsetExp/Source/MeshModelingToolsExp/Public/Physics/CollisionPropertySets.h
--- a/Engine/Plugins/Experimental/MeshModelingToolsetExp/Source/MeshModelingToolsExp/Public/Physics/CollisionPropertySets.h
+++ b/Engine/Plugins/Experimental/MeshModelingToolsetExp/Source/MeshModelingToolsExp/Public/Physics/CollisionPropertySets.h
@@ -2,10 +2,10 @@
 
 #pragma once
 
+#include "BodySetupEnums.h"
 #include "CoreMinimal.h"
 #include "InteractiveTool.h"
-#include "Engine/Classes/PhysicsEngine/AggregateGeom.h"
-#include "BodySetupEnums.h"
+#include "PhysicsEngine/AggregateGeom.h"
 #include "CollisionPropertySets.generated.h"
 
 
@@ -170,4 +170,4 @@ namespace UE
 	{
 		void InitializePhysicsToolObjectPropertySet(const FPhysicsDataCollection* PhysicsData, UPhysicsObjectToolPropertySet* PropSet);
 	}
-}
\ No newline at end of file
+}
diff --git a/Engine/Plugins/Experimental/Mutable/Source/CustomizableObjectPopulationEditor/Private/MuCOPE/CustomizableObjectPopulationClassDetails.cpp b/Engine/Plugins/Experimental/Mutable/Source/CustomizableObjectPopulationEditor/Private/MuCOPE/CustomizableObjectPopulationClassDetails.cpp
--- a/Engine/Plugins/Experimental/Mutable/Source/CustomizableObjectPopulationEditor/Private/MuCOPE/CustomizableObjectPopulationClassDetails.cpp
+++ b/Engine/Plugins/Experimental/Mutable/Source/CustomizableObjectPopulationEditor/Private/MuCOPE/CustomizableObjectPopulationClassDetails.cpp
@@ -6,7 +6,6 @@
 #include "Containers/EnumAsByte.h"
 #include "Curves/CurveBase.h"
 #include "Curves/CurveLinearColor.h"
-#include "CustomizableObjectPopulation/Private/MuCOP/CustomizableObjectPopulationCharacteristic.h"
 #include "DetailCategoryBuilder.h"
 #include "DetailLayoutBuilder.h"
 #include "DetailWidgetRow.h"
@@ -27,6 +26,7 @@
 #include "MuCO/CustomizableObject.h"
 #include "MuCO/CustomizableObjectInstance.h"
 #include "MuCO/CustomizableObjectParameterTypeDefinitions.h"
+#include "MuCOP/CustomizableObjectPopulationCharacteristic.h"
 #include "MuCOP/CustomizableObjectPopulationClass.h"
 #include "PropertyCustomizationHelpers.h"
 #include "PropertyHandle.h"
diff --git a/Engine/Plugins/Experimental/Mutable/Source/CustomizableObjectPopulationEditor/Private/MuCOPE/CustomizableObjectPopulationClassDetails.h b/Engine/Plugins/Experimental/Mutable/Source/CustomizableObjectPopulationEditor/Private/MuCOPE/CustomizableObjectPopulationClassDetails.h
--- a/Engine/Plugins/Experimental/Mutable/Source/CustomizableObjectPopulationEditor/Private/MuCOPE/CustomizableObjectPopulationClassDetails.h
+++ b/Engine/Plugins/Experimental/Mutable/Source/CustomizableObjectPopulationEditor/Private/MuCOPE/CustomizableObjectPopulationClassDetails.h
@@ -8,7 +8,6 @@
 #include "Containers/Set.h"
 #include "Containers/SparseArray.h"
 #include "Containers/UnrealString.h"
-#include "CustomizableObjectPopulation/Private/MuCOP/CustomizableObjectPopulationConstraint.h"
 #include "Delegates/Delegate.h"
 #include "HAL/Platform.h"
 #include "IDetailCustomization.h"
@@ -17,6 +16,7 @@
 #include "Math/Color.h"
 #include "Math/Vector2D.h"
 #include "Misc/Optional.h"
+#include "MuCOP/CustomizableObjectPopulationConstraint.h"
 #include "Styling/SlateBrush.h"
 #include "Templates/SharedPointer.h"
 #include "Types/SlateEnums.h"
@@ -448,4 +448,4 @@ private:
 	// Pointer to the builder passed by CustomizeDetails method
 	IDetailLayoutBuilder* DetailBuilderPtr = nullptr;
 
-};
\ No newline at end of file
+};
diff --git a/Engine/Plugins/Experimental/Mutable/Source/MutableTools/Private/MuT/NodeScalarTablePrivate.h b/Engine/Plugins/Experimental/Mutable/Source/MutableTools/Private/MuT/NodeScalarTablePrivate.h
--- a/Engine/Plugins/Experimental/Mutable/Source/MutableTools/Private/MuT/NodeScalarTablePrivate.h
+++ b/Engine/Plugins/Experimental/Mutable/Source/MutableTools/Private/MuT/NodeScalarTablePrivate.h
@@ -7,7 +7,7 @@
 #include "MuT/Table.h"
 #include "MuT/AST.h"
 
-#include "../Private/MuR/MemoryPrivate.h"
+#include "MuR/MemoryPrivate.h"
 
 
 namespace mu
diff --git a/Engine/Plugins/Experimental/OpenImageDenoise/Source/OpenImageDenoise/Private/Plugin.cpp b/Engine/Plugins/Experimental/OpenImageDenoise/Source/OpenImageDenoise/Private/Plugin.cpp
--- a/Engine/Plugins/Experimental/OpenImageDenoise/Source/OpenImageDenoise/Private/Plugin.cpp
+++ b/Engine/Plugins/Experimental/OpenImageDenoise/Source/OpenImageDenoise/Private/Plugin.cpp
@@ -3,7 +3,7 @@
 #include "CoreMinimal.h"
 #include "Modules/ModuleInterface.h"
 #include "Modules/ModuleManager.h"
-#include "Renderer/Public/PathTracingDenoiser.h"
+#include "PathTracingDenoiser.h"
 #include "RHI.h"
 #include "RHIResources.h"
 
diff --git a/Engine/Plugins/Experimental/ProxyLODPlugin/Source/ProxyLOD/Private/ProxyLODSimplifier.h b/Engine/Plugins/Experimental/ProxyLODPlugin/Source/ProxyLOD/Private/ProxyLODSimplifier.h
--- a/Engine/Plugins/Experimental/ProxyLODPlugin/Source/ProxyLOD/Private/ProxyLODSimplifier.h
+++ b/Engine/Plugins/Experimental/ProxyLODPlugin/Source/ProxyLOD/Private/ProxyLODSimplifier.h
@@ -8,9 +8,9 @@
 #include "ProxyLODMeshTypes.h" // for resize array
 #include "ProxyLODThreadedWrappers.h"
 
-#include "Containers/HashTable.h"
 #include "Containers/BinaryHeap.h"
-#include "Developer/MeshSimplifier/Private/MeshSimplifyElements.h"
+#include "Containers/HashTable.h"
+#include "MeshSimplifyElements.h"
 #include "ProxyLODQuadric.h"
 
 
@@ -836,4 +836,4 @@ namespace ProxyLOD
 		return (uint32)Index;
 	}
 
-}; // end namespace ProxyLOD
\ No newline at end of file
+}; // end namespace ProxyLOD
diff --git a/Engine/Plugins/Experimental/SkeletalReduction/Source/Private/SkeletalSimplifierMeshManager.h b/Engine/Plugins/Experimental/SkeletalReduction/Source/Private/SkeletalSimplifierMeshManager.h
--- a/Engine/Plugins/Experimental/SkeletalReduction/Source/Private/SkeletalSimplifierMeshManager.h
+++ b/Engine/Plugins/Experimental/SkeletalReduction/Source/Private/SkeletalSimplifierMeshManager.h
@@ -4,9 +4,9 @@
 #include "CoreMinimal.h"
 
 #include "Components.h" // for MAX_TEXCOORDS
-#include "Containers/HashTable.h"
 #include "Containers/BinaryHeap.h"
-#include "Developer/MeshSimplifier/Private/MeshSimplifyElements.h"
+#include "Containers/HashTable.h"
+#include "MeshSimplifyElements.h"
 
 #include "SkeletalSimplifierVertex.h"
 #include "SkeletalSimplifierLinearAlgebra.h"
@@ -554,4 +554,4 @@ namespace SkeletalSimplifier
 		};
 
 	};
-}
\ No newline at end of file
+}
diff --git a/Engine/Plugins/Experimental/UVEditor/Source/UVEditor/Private/SUVEditor3DViewportToolBar.cpp b/Engine/Plugins/Experimental/UVEditor/Source/UVEditor/Private/SUVEditor3DViewportToolBar.cpp
--- a/Engine/Plugins/Experimental/UVEditor/Source/UVEditor/Private/SUVEditor3DViewportToolBar.cpp
+++ b/Engine/Plugins/Experimental/UVEditor/Source/UVEditor/Private/SUVEditor3DViewportToolBar.cpp
@@ -2,16 +2,16 @@
 
 #include "SUVEditor3DViewportToolBar.h"
 
-#include "UVEditorCommands.h"
 #include "EditorViewportCommands.h"
-#include "Styling/AppStyle.h"
 #include "Framework/MultiBox/MultiBoxBuilder.h"
+#include "SEditorViewportToolBarMenu.h"
+#include "SEditorViewportViewMenu.h"
+#include "SUVEditor3DViewport.h"
+#include "Styling/AppStyle.h"
+#include "UVEditorCommands.h"
 #include "UVEditorStyle.h"
 #include "Widgets/Layout/SBorder.h"
 #include "Widgets/SBoxPanel.h"
-#include "Editor/UnrealEd/Public/SEditorViewportToolBarMenu.h"
-#include "SEditorViewportViewMenu.h"
-#include "SUVEditor3DViewport.h"
 
 #define LOCTEXT_NAMESPACE "SUVEditor3DViewportToolBar"
 
@@ -108,4 +108,4 @@ TSharedRef<SWidget> SUVEditor3DViewportToolBar::MakeToolBar(const TSharedPtr<FEx
 	return ToolbarBuilder.MakeWidget();
 }
 
-#undef LOCTEXT_NAMESPACE
\ No newline at end of file
+#undef LOCTEXT_NAMESPACE
diff --git a/Engine/Plugins/Experimental/Water/Source/Runtime/Private/GerstnerWaterWaves.cpp b/Engine/Plugins/Experimental/Water/Source/Runtime/Private/GerstnerWaterWaves.cpp
--- a/Engine/Plugins/Experimental/Water/Source/Runtime/Private/GerstnerWaterWaves.cpp
+++ b/Engine/Plugins/Experimental/Water/Source/Runtime/Private/GerstnerWaterWaves.cpp
@@ -1,10 +1,10 @@
 // Copyright Epic Games, Inc. All Rights Reserved.
 
 #include "GerstnerWaterWaves.h"
+#include "Engine/Engine.h"
 #include "GerstnerWaterWaveSubsystem.h"
+#include "Misc/LargeWorldRenderPosition.h"
 #include "WaterModule.h"
-#include "Engine/Engine.h"
-#include "Core/Public/Misc/LargeWorldRenderPosition.h"
 
 #include UE_INLINE_GENERATED_CPP_BY_NAME(GerstnerWaterWaves)
 
diff --git a/Engine/Plugins/Experimental/Water/Source/Runtime/Private/WaterBodyTypes.cpp b/Engine/Plugins/Experimental/Water/Source/Runtime/Private/WaterBodyTypes.cpp
--- a/Engine/Plugins/Experimental/Water/Source/Runtime/Private/WaterBodyTypes.cpp
+++ b/Engine/Plugins/Experimental/Water/Source/Runtime/Private/WaterBodyTypes.cpp
@@ -1,12 +1,12 @@
 // Copyright Epic Games, Inc. All Rights Reserved.
 
 #include "WaterBodyTypes.h"
+#include "LandscapeHeightfieldCollisionComponent.h"
+#include "LandscapeInfo.h"
+#include "LandscapeProxy.h"
 #include "WaterBodyActor.h"
-#include "WaterSubsystem.h"
 #include "WaterSplineComponent.h"
-#include "Landscape/Classes/LandscapeHeightfieldCollisionComponent.h"
-#include "Landscape/Classes/LandscapeInfo.h"
-#include "Landscape/Classes/LandscapeProxy.h"
+#include "WaterSubsystem.h"
 
 #include UE_INLINE_GENERATED_CPP_BY_NAME(WaterBodyTypes)
 
diff --git a/Engine/Plugins/Experimental/Water/Source/Runtime/Private/WaterMeshComponent.cpp b/Engine/Plugins/Experimental/Water/Source/Runtime/Private/WaterMeshComponent.cpp
--- a/Engine/Plugins/Experimental/Water/Source/Runtime/Private/WaterMeshComponent.cpp
+++ b/Engine/Plugins/Experimental/Water/Source/Runtime/Private/WaterMeshComponent.cpp
@@ -1,22 +1,22 @@
 // Copyright Epic Games, Inc. All Rights Reserved.
 
 #include "WaterMeshComponent.h"
-#include "WaterBodyActor.h"
-#include "WaterSplineComponent.h"
-#include "EngineUtils.h"
+#include "Algo/Transform.h"
 #include "DrawDebugHelpers.h"
+#include "Engine/Engine.h"
+#include "EngineUtils.h"
+#include "Materials/MaterialInstanceDynamic.h"
+#include "Materials/MaterialParameterCollection.h"
+#include "Materials/MaterialParameterCollectionInstance.h"
+#include "Math/NumericLimits.h"
 #include "PhysicsEngine/AggregateGeom.h"
 #include "PhysicsEngine/BodySetup.h"
-#include "Engine/Classes/Materials/MaterialInstanceDynamic.h"
-#include "Engine/Classes/Materials/MaterialParameterCollection.h"
-#include "Engine/Classes/Materials/MaterialParameterCollectionInstance.h"
-#include "Engine/Engine.h"
+#include "WaterBodyActor.h"
 #include "WaterMeshSceneProxy.h"
-#include "WaterSubsystem.h"
 #include "WaterModule.h"
+#include "WaterSplineComponent.h"
+#include "WaterSubsystem.h"
 #include "WaterUtils.h"
-#include "Math/NumericLimits.h"
-#include "Algo/Transform.h"
 
 #include UE_INLINE_GENERATED_CPP_BY_NAME(WaterMeshComponent)
 
diff --git a/Engine/Plugins/Experimental/Water/Source/Runtime/Private/WaterQuadTree.cpp b/Engine/Plugins/Experimental/Water/Source/Runtime/Private/WaterQuadTree.cpp
--- a/Engine/Plugins/Experimental/Water/Source/Runtime/Private/WaterQuadTree.cpp
+++ b/Engine/Plugins/Experimental/Water/Source/Runtime/Private/WaterQuadTree.cpp
@@ -1,8 +1,8 @@
 // Copyright Epic Games, Inc. All Rights Reserved.
 
 #include "WaterQuadTree.h"
-#include "Engine/Public/SceneManagement.h"
 #include "Materials/MaterialInterface.h"
+#include "SceneManagement.h"
 
 #if WITH_WATER_SELECTION_SUPPORT
 #include "HitProxies.h"
diff --git a/Engine/Plugins/Experimental/Water/Source/Runtime/Private/WaterSubsystem.cpp b/Engine/Plugins/Experimental/Water/Source/Runtime/Private/WaterSubsystem.cpp
--- a/Engine/Plugins/Experimental/Water/Source/Runtime/Private/WaterSubsystem.cpp
+++ b/Engine/Plugins/Experimental/Water/Source/Runtime/Private/WaterSubsystem.cpp
@@ -1,30 +1,30 @@
 // Copyright Epic Games, Inc. All Rights Reserved.
 
 #include "WaterSubsystem.h"
-#include "WaterModule.h"
-#include "HAL/IConsoleManager.h"
-#include "Engine/World.h"
-#include "TimerManager.h"
-#include "WaterMeshComponent.h"
+#include "BuoyancyManager.h"
+#include "CollisionShape.h"
 #include "Engine/Engine.h"
+#include "Engine/StaticMesh.h"
+#include "Engine/World.h"
 #include "EngineUtils.h"
+#include "HAL/IConsoleManager.h"
+#include "Interfaces/Interface_PostProcessVolume.h"
+#include "LandscapeSubsystem.h"
 #include "Materials/MaterialInstanceDynamic.h"
 #include "Materials/MaterialParameterCollection.h"
 #include "Materials/MaterialParameterCollectionInstance.h"
+#include "Math/NumericLimits.h"
+#include "SceneView.h"
+#include "TimerManager.h"
 #include "WaterBodyActor.h"
-#include "WaterBodyIslandActor.h"
 #include "WaterBodyExclusionVolume.h"
+#include "WaterBodyIslandActor.h"
+#include "WaterMeshComponent.h"
+#include "WaterModule.h"
+#include "WaterRuntimeSettings.h"
 #include "WaterSplineComponent.h"
 #include "WaterUtils.h"
-#include "CollisionShape.h"
-#include "Interfaces/Interface_PostProcessVolume.h"
-#include "SceneView.h"
-#include "Math/NumericLimits.h"
-#include "BuoyancyManager.h"
-#include "WaterRuntimeSettings.h"
-#include "Engine/StaticMesh.h"
 #include "WaterViewExtension.h"
-#include "Landscape/Public/LandscapeSubsystem.h"
 
 #include UE_INLINE_GENERATED_CPP_BY_NAME(WaterSubsystem)
 
diff --git a/Engine/Plugins/Experimental/Water/Source/Runtime/Public/WaterMeshSceneProxy.h b/Engine/Plugins/Experimental/Water/Source/Runtime/Public/WaterMeshSceneProxy.h
--- a/Engine/Plugins/Experimental/Water/Source/Runtime/Public/WaterMeshSceneProxy.h
+++ b/Engine/Plugins/Experimental/Water/Source/Runtime/Public/WaterMeshSceneProxy.h
@@ -3,12 +3,12 @@
 #pragma once
 
 #include "CoreMinimal.h"
-#include "UObject/ObjectMacros.h"
+#include "Materials/Material.h"
 #include "PrimitiveSceneProxy.h"
+#include "UObject/ObjectMacros.h"
+#include "WaterInstanceDataBuffer.h"
 #include "WaterQuadTree.h"
-#include "Engine/Classes/Materials/Material.h"
 #include "WaterVertexFactory.h"
-#include "WaterInstanceDataBuffer.h"
 
 class UWaterMeshComponent;
 
