Water: WaterZone Overlap Priority.

diff --git a/Engine/Plugins/Experimental/Water/Source/Runtime/Private/WaterBodyComponent.cpp b/Engine/Plugins/Experimental/Water/Source/Runtime/Private/WaterBodyComponent.cpp
--- a/Engine/Plugins/Experimental/Water/Source/Runtime/Private/WaterBodyComponent.cpp
+++ b/Engine/Plugins/Experimental/Water/Source/Runtime/Private/WaterBodyComponent.cpp
@@ -12,6 +12,7 @@
 #include "AI/NavigationSystemHelpers.h"
 #include "Algo/AllOf.h"
 #include "Algo/AnyOf.h"
+#include "Algo/MaxElement.h"
 #include "Misc/SecureHash.h"
 #include "PhysicsEngine/BodySetup.h"
 #include "Components/SplineMeshComponent.h"
@@ -1720,7 +1721,8 @@ AWaterZone* UWaterBodyComponent::FindWaterZone() const
 		return WaterZoneOverride.Get();
 	}
 
-	TArray<AWaterZone*, TInlineAllocator<4>> ViableZones;
+	// Score each overlapping water zone and then pick the best.
+	TArray<TPair<int32, AWaterZone*>, TInlineAllocator<4>> ViableZones;
 
 	if (const UWorld* World = GetWorld())
 	{
@@ -1742,15 +1744,23 @@ AWaterZone* UWaterBodyComponent::FindWaterZone() const
 				{
 					if (WaterZone->GetTypedOuter<ULevel>() == PreferredLevel)
 					{
-						return WaterZone;
+						ViableZones.Add({ WaterZone->GetOverlapPriority(), WaterZone });
+						continue;
 					}
 
-					ViableZones.Add(WaterZone);
+					// Fallback to zones not in the preferred level only as a final resort.
+					ViableZones.Add({ TNumericLimits<int32>::Min(), WaterZone});
 				}
 			}
 		}
 	}
-	return ViableZones.Num() > 0 ? ViableZones[0] : nullptr;
+	
+	if (ViableZones.Num() == 0)
+	{
+		return nullptr;
+	}
+
+	return Algo::MaxElementBy(ViableZones, [](const TPair<int32, AWaterZone*>& A) { return A.Key; })->Value;
 }
 
 static inline FColor PackFlowData(float VelocityMagnitude, float DirectionAngle, float MaxVelocity)
diff --git a/Engine/Plugins/Experimental/Water/Source/Runtime/Public/WaterZoneActor.h b/Engine/Plugins/Experimental/Water/Source/Runtime/Public/WaterZoneActor.h
--- a/Engine/Plugins/Experimental/Water/Source/Runtime/Public/WaterZoneActor.h
+++ b/Engine/Plugins/Experimental/Water/Source/Runtime/Public/WaterZoneActor.h
@@ -56,6 +56,8 @@ public:
 	FVector2f GetWaterHeightExtents() const { return WaterHeightExtents; }
 	float GetGroundZMin() const { return GroundZMin; }
 
+	int32 GetOverlapPriority() const { return OverlapPriority; }
+
 	UPROPERTY(Transient, DuplicateTransient, VisibleAnywhere, BlueprintReadOnly, Category = Water)
 	TObjectPtr<UTextureRenderTarget2D> WaterInfoTexture;
 
@@ -136,6 +138,10 @@ private:
 	UPROPERTY(Category = NonTessellatedLOD, EditAnywhere, meta = (ClampMin = "1", DisplayName = "Non Tessellated LOD Section Scale", EditCondition = "bEnableNonTessellatedLODMesh"))
 	uint32 NonTessellatedLODSectionScale = 4;
 
+	/** Higher number is higher priority. If Water Zones overlap and a water body does not have a manual water zone override, this priority will be used when automatically assigning the zone. */
+	UPROPERTY(Category = Water, EditAnywhere, AdvancedDisplay)
+	int32 OverlapPriority = 0;
+
 	UPROPERTY(Category = NonTessellatedLOD, EditAnywhere, meta = (DisplayName = "Enable Non-Tessellated LOD Mesh"))
 	bool bEnableNonTessellatedLODMesh = false;
 
