Water: deprecated ocean's visual extent parameter in favor of having the ocean always fill the water zone entirely.

diff --git a/Engine/Plugins/Experimental/Water/Source/Runtime/Private/WaterBodyOceanComponent.cpp b/Engine/Plugins/Experimental/Water/Source/Runtime/Private/WaterBodyOceanComponent.cpp
--- a/Engine/Plugins/Experimental/Water/Source/Runtime/Private/WaterBodyOceanComponent.cpp
+++ b/Engine/Plugins/Experimental/Water/Source/Runtime/Private/WaterBodyOceanComponent.cpp
@@ -30,7 +30,6 @@ UWaterBodyOceanComponent::UWaterBodyOceanComponent(const FObjectInitializer& Obj
 	: Super(ObjectInitializer)
 {
 	CollisionExtents = FVector(50000.f, 50000.f, 10000.f);
-	VisualExtents = FVector2D(150000.f, 150000.f);
 
 	// @todo_water : Remove these checks (Once AWaterBody is no more Blueprintable, these methods should become PURE_VIRTUAL and this class should overload them)
 	check(IsFlatSurface());
@@ -69,16 +68,6 @@ void UWaterBodyOceanComponent::SetHeightOffset(float InHeightOffset)
 	}
 }
 
-void UWaterBodyOceanComponent::SetVisualExtents(FVector2D NewExtents)
-{
-	if (VisualExtents != NewExtents)
-	{
-		VisualExtents = NewExtents;
-		UpdateWaterBodyRenderData();
-		Modify();
-	}
-}
-
 void UWaterBodyOceanComponent::BeginUpdateWaterBody()
 {
 	Super::BeginUpdateWaterBody();
@@ -101,10 +90,6 @@ void UWaterBodyOceanComponent::OnPostEditChangeProperty(FOnWaterBodyChangedParam
 		// Affects the physics shape
 		InOutOnWaterBodyChangedParams.bShapeOrPositionChanged = true;
 	}
-	else if (PropertyChangedEvent.GetPropertyName() == GET_MEMBER_NAME_CHECKED(UWaterBodyOceanComponent, VisualExtents))
-	{
-		InOutOnWaterBodyChangedParams.bShapeOrPositionChanged = true;
-	}
 }
 const TCHAR* UWaterBodyOceanComponent::GetWaterSpriteTextureName() const
 {
@@ -193,12 +178,18 @@ bool UWaterBodyOceanComponent::GenerateWaterBodyMesh(UE::Geometry::FDynamicMesh3
 
 	const UWaterSplineComponent* SplineComp = GetWaterSpline();
 	const FVector OceanLocation = GetComponentLocation();
+	const AWaterZone* WaterZone = GetWaterZone();
 	
 	if (SplineComp->GetNumberOfSplineSegments() < 3)
 	{
 		return false;
 	}
 
+	if (!WaterZone)
+	{
+		return false;
+	}
+
 	FPolygon2d Island;
 	TArray<FVector> PolyLineVertices;
 	SplineComp->ConvertSplineToPolyLine(ESplineCoordinateSpace::Local, FMath::Square(10.f), PolyLineVertices);
@@ -214,9 +205,9 @@ bool UWaterBodyOceanComponent::GenerateWaterBodyMesh(UE::Geometry::FDynamicMesh3
 	IslandBounds.Expand(1);
 
 	// #todo: account for scale
-	const FAxisAlignedBox2d OceanBounds = FAxisAlignedBox2d(-FVector2d(VisualExtents) / 2.0, FVector2d(VisualExtents) / 2.0);
-	FPolygon2d IslandBoundingPolygon = FPolygon2d::MakeRectangle(FVector2d(IslandBounds.Center()), 2 * IslandBounds.Extents().X, 2 * IslandBounds.Extents().Y);
-
+	const FVector2d OceanExtents = WaterZone->GetZoneExtent(); // file the entire water zone with an ocean
+	const FAxisAlignedBox2d OceanBounds = FAxisAlignedBox2d(FVector2d(GetComponentTransform().InverseTransformPositionNoScale(WaterZone->GetActorLocation())), OceanExtents.X / 2.0);
+	FPolygon2d IslandBoundingPolygon = FPolygon2d::MakeRectangle(IslandBounds.Center(), IslandBounds.Extents().X * 2., IslandBounds.Extents().Y * 2.);
 
 	FConstrainedDelaunay2d Triangulation;
 	Triangulation.FillRule = FConstrainedDelaunay2d::EFillRule::Positive;
@@ -229,7 +220,7 @@ bool UWaterBodyOceanComponent::GenerateWaterBodyMesh(UE::Geometry::FDynamicMesh3
 	Triangulation.Add(Island);
 	if (!Triangulation.Triangulate())
 	{
-		UE_LOG(LogWater, Error, TEXT("Failed to triangulate Ocean mesh for %s. Ensure that the spline points are fully contained within the VisualExtents of the ocean!"), *GetOwner()->GetActorNameOrLabel());
+		UE_LOG(LogWater, Error, TEXT("Failed to triangulate Ocean mesh for %s. Ensure that the Ocean's spline does not form any loops."), *GetOwner()->GetActorNameOrLabel());
 	}
 
 	if (Triangulation.Triangles.Num() == 0)
@@ -384,27 +375,19 @@ void UWaterBodyOceanComponent::PostLoad()
 	Super::PostLoad();
 
 #ifdef WITH_EDITORONLY_DATA
-	if (GetLinkerCustomVersion(FFortniteMainBranchObjectVersion::GUID) < FFortniteMainBranchObjectVersion::WaterZonesRefactor)
-	{
-		if (AWaterZone* WaterZone = GetWaterZone())
-		{
-			SetVisualExtents(WaterZone->GetZoneExtent());
-		}
-	}
-	if (GetLinkerCustomVersion(FFortniteMainBranchObjectVersion::GUID) < FFortniteMainBranchObjectVersion::WaterNontessellatedLODSupportAdded)
-	{
-		// Prior to the refactors when introducing the non-tessellated LOD, the Ocean's visual extent was treated as 2x the actual value.
-		// To maintain visual consistency with previous versions, the value needs to now be updated to be 2x.
-		VisualExtents = VisualExtents * 2.;
-	}
 #endif // WITH_EDITORONLY_DATA
 }
 
 FBoxSphereBounds UWaterBodyOceanComponent::CalcBounds(const FTransform& LocalToWorld) const
 {
-	FVector Min(FVector(-VisualExtents / 2.f, -1 * GetChannelDepth()));
-	FVector Max(FVector(VisualExtents / 2.f, 0.f));
-	return FBoxSphereBounds(FBox(Min, Max)).TransformBy(LocalToWorld);
+	if (AWaterZone* WaterZone = GetWaterZone())
+	{
+		const FVector2D OceanExtents = WaterZone->GetZoneExtent();
+		const FVector Min(FVector(-OceanExtents / 2.0, -1.0 * GetChannelDepth()));
+		const FVector Max(FVector(OceanExtents / 2.0, 0.0));
+		return FBoxSphereBounds(FBox(Min, Max)).TransformBy(LocalToWorld);
+	}
+	return FBoxSphereBounds().TransformBy(LocalToWorld);
 }
 
 void UWaterBodyOceanComponent::OnUpdateBody(bool bWithExclusionVolumes)
diff --git a/Engine/Plugins/Experimental/Water/Source/Runtime/Private/WaterBodySceneProxy.cpp b/Engine/Plugins/Experimental/Water/Source/Runtime/Private/WaterBodySceneProxy.cpp
--- a/Engine/Plugins/Experimental/Water/Source/Runtime/Private/WaterBodySceneProxy.cpp
+++ b/Engine/Plugins/Experimental/Water/Source/Runtime/Private/WaterBodySceneProxy.cpp
@@ -10,6 +10,11 @@
 #include "PrimitiveSceneInfo.h"
 #include "Algo/Transform.h"
 
+TAutoConsoleVariable<int32> CVarVisualizeWaterInfoSceneProxies(
+		TEXT("r.Water.VisualizeWaterInfoSceneProxies"),
+		0,
+		TEXT("Enables a wireframe visualization mode for the water info scene proxy geometry. Modes: 0 to disable, 1 to show only the selected water body, 2 to show all water bodies."));
+
 FWaterBodySceneProxy::FWaterBodySceneProxy(UWaterBodyComponent* Component)
 	: FPrimitiveSceneProxy(Component)
 	, WaterBodySectionedLODMesh(GetScene().GetFeatureLevel())
@@ -201,6 +206,14 @@ void FWaterBodySceneProxy::GetDynamicMeshElements(const TArray<const FSceneView*
 
 	const bool bWithinWaterInfoPasses = CurrentWaterInfoPass != EWaterInfoPass::None;
 
+#if !UE_BUILD_SHIPPING
+	const int VisualizeWaterInfoMode = CVarVisualizeWaterInfoSceneProxies.GetValueOnRenderThread();
+#else
+	const int VisualizeWaterInfoMode = 0;
+#endif // !UE_BUILD_SHIPPING
+
+	const bool bVisualizeWaterInfoMesh = (VisualizeWaterInfoMode != 0) && !bWithinWaterInfoPasses;
+
 	FMaterialRenderProxy* MaterialToUse = bWithinWaterInfoPasses ? WaterInfoMaterial : WaterLODMaterial;
 
 	if (!MaterialToUse)
@@ -215,7 +228,7 @@ void FWaterBodySceneProxy::GetDynamicMeshElements(const TArray<const FSceneView*
 	}
 
 	// If we are not in the waterinfo pass and the cvar is not set to show opaque bodies, we should be in wireframe
-	const bool bWireframe = AllowDebugViewmodes() && ViewFamily.EngineShowFlags.Wireframe;
+	const bool bWireframe = AllowDebugViewmodes() && (ViewFamily.EngineShowFlags.Wireframe || bVisualizeWaterInfoMesh);
 
 	if (bWireframe)
 	{
@@ -264,6 +277,18 @@ void FWaterBodySceneProxy::GetDynamicMeshElements(const TArray<const FSceneView*
 			{
 				Collector.AddMesh(ViewIndex, MeshBatch);
 			}
+#if !UE_BUILD_SHIPPING
+			if (bVisualizeWaterInfoMesh && (VisualizeWaterInfoMode > 1 || IsSelected()))
+			{
+				FMeshBatch& WireframeMeshBatch = Collector.AllocateMesh();
+				WireframeMeshBatch.MaterialRenderProxy = MaterialToUse;
+				WireframeMeshBatch.bWireframe = bWireframe;
+				if (WaterBodyInfoMesh.GetMeshElements(WireframeMeshBatch, SDPG_World, IsLocalToWorldDeterminantNegative()))
+				{
+					Collector.AddMesh(ViewIndex, WireframeMeshBatch);
+				}
+			}
+#endif // !UE_BUILD_SHIPPING
 		}
 	}
 }
diff --git a/Engine/Plugins/Experimental/Water/Source/Runtime/Public/WaterBodyOceanComponent.h b/Engine/Plugins/Experimental/Water/Source/Runtime/Public/WaterBodyOceanComponent.h
--- a/Engine/Plugins/Experimental/Water/Source/Runtime/Public/WaterBodyOceanComponent.h
+++ b/Engine/Plugins/Experimental/Water/Source/Runtime/Public/WaterBodyOceanComponent.h
@@ -25,8 +25,10 @@ public:
 	virtual void SetHeightOffset(float InHeightOffset) override;
 	virtual float GetHeightOffset() const override { return HeightOffset; }
 
-	void SetVisualExtents(FVector2D NewExtents);
-	FVector2D GetVisualExtents() const { return VisualExtents; }
+	UE_DEPRECATED(5.1, "Oceans no longer rely on the visual extent parameter making this obsolete. Instead they will be guaranteed to fill the entire water zone to which they belong.")
+	void SetVisualExtents(FVector2D) {}
+	UE_DEPRECATED(5.1, "Oceans no longer rely on the visual extent parameter making this obsolete. Instead they will be guaranteed to fill the entire water zone to which they belong.")
+	FVector2D GetVisualExtents() const { return FVector2D(); }
 protected:
 	/** UWaterBodyComponent Interface */
 	virtual bool IsBodyDynamic() const override { return true; }
@@ -50,9 +52,8 @@ protected:
 	UPROPERTY(NonPIEDuplicateTransient)
 	TArray<TObjectPtr<UOceanCollisionComponent>> CollisionHullSets;
 
-	/** The area over which the ocean should be displayed, centered on the actor */
-	UPROPERTY(Category = Water, EditAnywhere, BlueprintReadOnly)
-	FVector2D VisualExtents;
+	UPROPERTY()
+	FVector2D VisualExtents_DEPRECATED;
 
 	UPROPERTY(Category = Collision, EditAnywhere, BlueprintReadOnly)
 	FVector CollisionExtents;
