Fixed wrong default name being assigned to the water zone being automatically spawned when dragging a water body into a map

#rnx

diff --git a/Engine/Plugins/Experimental/Water/Source/Editor/Private/WaterEditorModule.cpp b/Engine/Plugins/Experimental/Water/Source/Editor/Private/WaterEditorModule.cpp
--- a/Engine/Plugins/Experimental/Water/Source/Editor/Private/WaterEditorModule.cpp
+++ b/Engine/Plugins/Experimental/Water/Source/Editor/Private/WaterEditorModule.cpp
@@ -260,7 +260,7 @@ void FWaterEditorModule::OnLevelActorAddedToWorld(AActor* Actor)
 						SpawnParams.Name = BrushActorName;
 						SpawnParams.bAllowDuringConstructionScript = true; // This can be called by construction script if the actor being added to the world is part of a blueprint, for example : 
 						AWaterLandscapeBrush* NewBrush = (WaterBrushActorFactory != nullptr)
-							? Cast<AWaterLandscapeBrush>(WaterBrushActorFactory->CreateActor(ActorWorld, FoundLandscape->GetLevel(), FTransform(WaterZoneBounds.GetCenter()), SpawnParams))
+							? Cast<AWaterLandscapeBrush>(WaterBrushActorFactory->CreateActor(WaterBrushClassPtr, FoundLandscape->GetLevel(), FTransform(WaterZoneBounds.GetCenter()), SpawnParams))
 							: ActorWorld->SpawnActor<AWaterLandscapeBrush>(WaterBrushClassPtr, SpawnParams);
 
 						if (NewBrush)
@@ -304,7 +304,7 @@ void FWaterEditorModule::OnLevelActorAddedToWorld(AActor* Actor)
 				SpawnParams.bAllowDuringConstructionScript = true; // This can be called by construction script if the actor being added to the world is part of a blueprint, for example : 
 
 				AWaterZone* WaterZoneActor = (WaterZoneActorFactory != nullptr)
-					? Cast<AWaterZone>(WaterZoneActorFactory->CreateActor(ActorWorld, Actor->GetLevel(), FTransform(WaterZoneBounds.GetCenter()), SpawnParams))
+					? Cast<AWaterZone>(WaterZoneActorFactory->CreateActor(WaterZoneClassPtr, Actor->GetLevel(), FTransform(WaterZoneBounds.GetCenter()), SpawnParams))
 					: ActorWorld->SpawnActor<AWaterZone>(WaterZoneClassPtr, SpawnParams);
 
 				if (WaterZoneActor)
