Water: fixed crash when loading Water Waves editor. Landscape and Water subsystems are generally disabled in preview worlds however the Waves editor is an exception. We override this behaviour with `FScopedAllowWaterSubsystemOnPreviewWorld` and we will end up with a water subsystem without a landscape subsystem.

diff --git a/Engine/Plugins/Experimental/Water/Source/Runtime/Private/WaterSubsystem.cpp b/Engine/Plugins/Experimental/Water/Source/Runtime/Private/WaterSubsystem.cpp
--- a/Engine/Plugins/Experimental/Water/Source/Runtime/Private/WaterSubsystem.cpp
+++ b/Engine/Plugins/Experimental/Water/Source/Runtime/Private/WaterSubsystem.cpp
@@ -294,11 +294,11 @@ void UWaterSubsystem::PostInitialize()
 
 #if WITH_EDITOR
 	// Register for heightmap streaming notifications
-	ULandscapeSubsystem* LandscapeSubsystem = World->GetSubsystem<ULandscapeSubsystem>();
-	check(LandscapeSubsystem);
-
-	check(!OnHeightmapStreamedHandle.IsValid());
-	OnHeightmapStreamedHandle = LandscapeSubsystem->GetOnHeightmapStreamedDelegate().AddUObject(this, &UWaterSubsystem::OnHeightmapStreamed);
+	if (ULandscapeSubsystem* LandscapeSubsystem = World->GetSubsystem<ULandscapeSubsystem>())
+	{
+		check(!OnHeightmapStreamedHandle.IsValid());
+		OnHeightmapStreamedHandle = LandscapeSubsystem->GetOnHeightmapStreamedDelegate().AddUObject(this, &UWaterSubsystem::OnHeightmapStreamed);
+	}
 #endif // WITH_EDITOR
 }
 
