Water: prevent checking for complete shader maps and submitting compile jobs in non-editor builds. If ever there is an incomplete shadermap, there is not much we can do in this case anyways so we will just render with the fallback materials.

diff --git a/Engine/Plugins/Experimental/Water/Source/Runtime/Private/WaterZoneActor.cpp b/Engine/Plugins/Experimental/Water/Source/Runtime/Private/WaterZoneActor.cpp
--- a/Engine/Plugins/Experimental/Water/Source/Runtime/Private/WaterZoneActor.cpp
+++ b/Engine/Plugins/Experimental/Water/Source/Runtime/Private/WaterZoneActor.cpp
@@ -357,15 +357,18 @@ bool AWaterZone::UpdateWaterInfoTexture()
 				return true;
 			}
 
+			WaterBodiesToRender.Add(WaterBodyComponent);
+			const FBox WaterBodyBounds = WaterBodyComponent->CalcBounds(WaterBodyComponent->GetComponentToWorld()).GetBox();
+			WaterZMax = FMath::Max(WaterZMax, WaterBodyBounds.Max.Z);
+			WaterZMin = FMath::Min(WaterZMin, WaterBodyBounds.Min.Z);
+
+#if WITH_EDITOR
 			if (UMaterialInterface* WaterInfoMaterial = WaterBodyComponent->GetWaterInfoMaterialInstance())
 			{
 				UsedMaterials.Add(WaterInfoMaterial);
 			}
+#endif // WITH_EDITOR
 
-			WaterBodiesToRender.Add(WaterBodyComponent);
-			const FBox WaterBodyBounds = WaterBodyComponent->CalcBounds(WaterBodyComponent->GetComponentToWorld()).GetBox();
-			WaterZMax = FMath::Max(WaterZMax, WaterBodyBounds.Max.Z);
-			WaterZMin = FMath::Min(WaterZMin, WaterBodyBounds.Min.Z);
 			return true;
 		});
 
@@ -395,6 +398,7 @@ bool AWaterZone::UpdateWaterInfoTexture()
 			}
 		}
 
+#if WITH_EDITOR
 		// Check all the ground actors have complete shader maps before we try to render them into the water info texture
 		for (TWeakObjectPtr<AActor> GroundActorPtr : GroundActors)
 		{
@@ -433,6 +437,7 @@ bool AWaterZone::UpdateWaterInfoTexture()
 		{
 			return false;
 		}
+#endif // WITH_EDITOR
 
 		const ETextureRenderTargetFormat Format = bHalfPrecisionTexture ? ETextureRenderTargetFormat::RTF_RGBA16f : RTF_RGBA32f;
 		WaterInfoTexture = FWaterUtils::GetOrCreateTransientRenderTarget2D(WaterInfoTexture, TEXT("WaterInfoTexture"), RenderTargetResolution, Format);
