Water: fixed a crash when water body actors are used in child actor components. An assumption is made that the water component is valid in PostEditMove, however, due to how child actors are created/destroyed every tick during a move operation, this assumption is invalid in this case.

diff --git a/Engine/Plugins/Experimental/Water/Source/Runtime/Private/WaterBodyActor.cpp b/Engine/Plugins/Experimental/Water/Source/Runtime/Private/WaterBodyActor.cpp
--- a/Engine/Plugins/Experimental/Water/Source/Runtime/Private/WaterBodyActor.cpp
+++ b/Engine/Plugins/Experimental/Water/Source/Runtime/Private/WaterBodyActor.cpp
@@ -118,6 +118,12 @@ void AWaterBody::PostEditMove(bool bFinished)
 {
 	Super::PostEditMove(bFinished);
 
+	// It's possible that the WaterBodyComponent is invalid here if, for example, we are part of a ChildActorComponent and are being destroyed/recreated every move.
+	if (!IsValid(WaterBodyComponent))
+	{
+		return;
+	}
+
 	if (bFinished)
 	{
 		WaterBodyComponent->UpdateWaterHeight();
