Water: don't render water far mesh without an ocean present.

diff --git a/Engine/Plugins/Experimental/Water/Source/Runtime/Private/WaterMeshComponent.cpp b/Engine/Plugins/Experimental/Water/Source/Runtime/Private/WaterMeshComponent.cpp
--- a/Engine/Plugins/Experimental/Water/Source/Runtime/Private/WaterMeshComponent.cpp
+++ b/Engine/Plugins/Experimental/Water/Source/Runtime/Private/WaterMeshComponent.cpp
@@ -241,6 +241,8 @@ void UWaterMeshComponent::RebuildWaterMesh(float InTileSize, const FIntPoint& In
 
 	// Will be updated with the ocean min bound, to be used to place the far mesh just under the ocean to avoid seams
 	float FarMeshHeight = GetComponentLocation().Z;
+	// Only use a far mesh when there is an ocean in the zone.
+	bool bHasOcean = false;
 
 	const UWaterSubsystem* WaterSubsystem = UWaterSubsystem::GetWaterSubsystem(GetWorld());
 
@@ -251,7 +253,7 @@ void UWaterMeshComponent::RebuildWaterMesh(float InTileSize, const FIntPoint& In
 	// Go through all water body actors to figure out bounds and water tiles
 	AWaterZone* OwningZone = GetOwner<AWaterZone>();
 	check(OwningZone);
-	OwningZone->ForEachWaterBodyComponent([this, WaterWorldBox, bIsFlooded, GlobalOceanHeight, OceanFlood, &FarMeshHeight](UWaterBodyComponent* WaterBodyComponent)
+	OwningZone->ForEachWaterBodyComponent([this, WaterWorldBox, bIsFlooded, GlobalOceanHeight, OceanFlood, &FarMeshHeight, &bHasOcean](UWaterBodyComponent* WaterBodyComponent)
 	{
 		check(WaterBodyComponent);
 		AActor* Actor = WaterBodyComponent->GetOwner();
@@ -545,6 +547,7 @@ void UWaterMeshComponent::RebuildWaterMesh(float InTileSize, const FIntPoint& In
 
 			// Place far mesh height just below the ocean level
 			FarMeshHeight = RenderData.SurfaceBaseHeight - WaterBodyComponent->GetMaxWaveHeight();
+			bHasOcean = true;
 
 			break;
 		}
@@ -559,7 +562,7 @@ void UWaterMeshComponent::RebuildWaterMesh(float InTileSize, const FIntPoint& In
 	});
 
 	// Build the far distance mesh instances if needed
-	if (IsMaterialUsedWithWater(FarDistanceMaterial) && FarDistanceMeshExtent > 0.0f)
+	if ((bHasOcean || bUseFarMeshWithoutOcean) && (IsMaterialUsedWithWater(FarDistanceMaterial) && FarDistanceMeshExtent > 0.0f))
 	{
 		UsedMaterials.Add(FarDistanceMaterial);
 
@@ -630,7 +633,8 @@ void UWaterMeshComponent::PostEditChangeProperty(FPropertyChangedEvent& Property
 			|| PropertyName == GET_MEMBER_NAME_CHECKED(UWaterMeshComponent, ExtentInTiles)
 			|| PropertyName == GET_MEMBER_NAME_CHECKED(UWaterMeshComponent, ForceCollapseDensityLevel)
 			|| PropertyName == GET_MEMBER_NAME_CHECKED(UWaterMeshComponent, FarDistanceMaterial)
-			|| PropertyName == GET_MEMBER_NAME_CHECKED(UWaterMeshComponent, FarDistanceMeshExtent))
+			|| PropertyName == GET_MEMBER_NAME_CHECKED(UWaterMeshComponent, FarDistanceMeshExtent)
+			|| PropertyName == GET_MEMBER_NAME_CHECKED(UWaterMeshComponent, bUseFarMeshWithoutOcean))
 		{
 			MarkWaterMeshGridDirty();
 			MarkRenderStateDirty();
diff --git a/Engine/Plugins/Experimental/Water/Source/Runtime/Public/WaterMeshComponent.h b/Engine/Plugins/Experimental/Water/Source/Runtime/Public/WaterMeshComponent.h
--- a/Engine/Plugins/Experimental/Water/Source/Runtime/Public/WaterMeshComponent.h
+++ b/Engine/Plugins/Experimental/Water/Source/Runtime/Public/WaterMeshComponent.h
@@ -106,6 +106,10 @@ private:
 	UPROPERTY(Transient, NonPIEDuplicateTransient, TextExportTransient)
 	TSet<TObjectPtr<UMaterialInterface>> UsedMaterials;
 
+	/** Forces the water mesh to always render the far mesh, regardless if there is an ocean or not.*/
+	UPROPERTY(Category = "Rendering|FarDistance", EditAnywhere)
+	bool bUseFarMeshWithoutOcean = false;
+
 	/** Dirty flag which will make sure the water mesh is updated properly */
 	bool bNeedsRebuild = true;
 
