Water: added a water fallback depth cvar to use when getting water body data when there are no ground actors.

diff --git a/Engine/Plugins/Experimental/Water/Source/Runtime/Private/WaterZoneActor.cpp b/Engine/Plugins/Experimental/Water/Source/Runtime/Private/WaterZoneActor.cpp
--- a/Engine/Plugins/Experimental/Water/Source/Runtime/Private/WaterZoneActor.cpp
+++ b/Engine/Plugins/Experimental/Water/Source/Runtime/Private/WaterZoneActor.cpp
@@ -31,6 +31,13 @@ static FAutoConsoleVariableRef CVarForceUpdateWaterInfoNextFrames(
 	ForceUpdateWaterInfoNextFrames,
 	TEXT("Force the water info texture to regenerate on the next N frames. A negative value will force update every frame."));
 
+TAutoConsoleVariable<float> CVarWaterFallbackDepth(
+	TEXT("r.Water.FallbackDepth"),
+	3000.0f,
+	TEXT("Depth to use for all water when there are no ground actors defined."),
+	ECVF_Default);
+
+
 
 AWaterZone::AWaterZone(const FObjectInitializer& Initializer)
 	: Super(Initializer)
@@ -398,6 +405,13 @@ bool AWaterZone::UpdateWaterInfoTexture()
 			}
 		}
 
+		// If we have no ground actors we need to set GroundZMin to be something sensible because it won't have been set above.
+		if (GroundActors.Num() == 0)
+		{
+			GroundZMax = WaterZMax;
+			GroundZMin = GroundZMax - CVarWaterFallbackDepth.GetValueOnGameThread();
+		}
+
 #if WITH_EDITOR
 		// Check all the ground actors have complete shader maps before we try to render them into the water info texture
 		for (TWeakObjectPtr<AActor> GroundActorPtr : GroundActors)
