#Water Add more information to an error

diff --git a/Engine/Plugins/Experimental/Water/Source/Runtime/Private/WaterBodyLakeComponent.cpp b/Engine/Plugins/Experimental/Water/Source/Runtime/Private/WaterBodyLakeComponent.cpp
--- a/Engine/Plugins/Experimental/Water/Source/Runtime/Private/WaterBodyLakeComponent.cpp
+++ b/Engine/Plugins/Experimental/Water/Source/Runtime/Private/WaterBodyLakeComponent.cpp
@@ -79,7 +79,7 @@ bool UWaterBodyLakeComponent::GenerateWaterBodyMesh(UE::Geometry::FDynamicMesh3&
 	Triangulation.Add(LakePoly);
 	if (!Triangulation.Triangulate())
 	{
-		UE_LOG(LogWater, Error, TEXT("Failed to triangulate Lake mesh for %s"), *GetOwner()->GetActorNameOrLabel());
+		UE_LOG(LogWater, Error, TEXT("Failed to triangulate Lake mesh for %s (%s)"), *GetOwner()->GetActorNameOrLabel(), *GetFullNameSafe(GetOwner()));
 	}
 
 	if (Triangulation.Triangles.Num() == 0)
