Landscape : Fixes for water usage from WorldPartition Commandlets.

* Guard against slate from being called when not initialized.
* Show message box if HLOD build commandlet process fails.

diff --git a/Engine/Plugins/Experimental/Water/Source/Editor/Private/WaterBrushManager.cpp b/Engine/Plugins/Experimental/Water/Source/Editor/Private/WaterBrushManager.cpp
--- a/Engine/Plugins/Experimental/Water/Source/Editor/Private/WaterBrushManager.cpp
+++ b/Engine/Plugins/Experimental/Water/Source/Editor/Private/WaterBrushManager.cpp
@@ -1109,70 +1109,72 @@ void AWaterBrushManager::SetMPCParams()
 		FVector RTWorldLocation, RTWorldSizeVector;
 		ComputeWaterLandscapeInfo(RTWorldLocation, RTWorldSizeVector);
 
-		UMaterialParameterCollection* LandscapeCollection = GEditor->GetEditorSubsystem<UWaterEditorSubsystem>()->GetLandscapeMaterialParameterCollection();
+		UWaterEditorSubsystem* WaterEditorSubsystem = GEditor->GetEditorSubsystem<UWaterEditorSubsystem>();
+		check(WaterEditorSubsystem);
+		UMaterialParameterCollection* LandscapeCollection =  WaterEditorSubsystem->GetLandscapeMaterialParameterCollection();
 		if (LandscapeCollection == nullptr)
+        {
+            UE_LOG(LogWaterEditor, Error, TEXT("No Landscape MaterialParameterCollection Assigned"));
+			return;
+        }
+        
+		UMaterialParameterCollectionInstance* LandscapeCollectionInstance = World->GetParameterCollectionInstance(CastChecked<UMaterialParameterCollection>(LandscapeCollection));
+		check(LandscapeCollectionInstance != nullptr);
+
+		if (!LandscapeCollectionInstance->SetScalarParameterValue(FName(TEXT("RTResX")), (float)LandscapeRTRes.X))
 		{
-			UE_LOG(LogWaterEditor, Error, TEXT("No Landscape MaterialParameterCollection Assigned"));
+			UE_LOG(LogWaterEditor, Error, TEXT("Failed to set \"RTResX\" on Landscape MaterialParameterCollection"));
 		}
-		else
-		{
-			UMaterialParameterCollectionInstance* LandscapeCollectionInstance = World->GetParameterCollectionInstance(CastChecked<UMaterialParameterCollection>(LandscapeCollection));
-			check(LandscapeCollectionInstance != nullptr);
-
-			if (!LandscapeCollectionInstance->SetScalarParameterValue(FName(TEXT("RTResX")), (float)LandscapeRTRes.X))
-			{
-				UE_LOG(LogWaterEditor, Error, TEXT("Failed to set \"RTResX\" on Landscape MaterialParameterCollection"));
-			}
 
-			if (!LandscapeCollectionInstance->SetScalarParameterValue(FName(TEXT("RTResY")), (float)LandscapeRTRes.Y))
-			{
-				UE_LOG(LogWaterEditor, Error, TEXT("Failed to set \"RTResY\" on Landscape MaterialParameterCollection"));
-			}
+		if (!LandscapeCollectionInstance->SetScalarParameterValue(FName(TEXT("RTResY")), (float)LandscapeRTRes.Y))
+		{
+			UE_LOG(LogWaterEditor, Error, TEXT("Failed to set \"RTResY\" on Landscape MaterialParameterCollection"));
+		}
 
-			if (!LandscapeCollectionInstance->SetScalarParameterValue(FName(TEXT("LSQuadsX")), (float)LandscapeQuads.X))
-			{
-				UE_LOG(LogWaterEditor, Error, TEXT("Failed to set \"LSQuadsX\" on Landscape MaterialParameterCollection"));
-			}
-			if (!LandscapeCollectionInstance->SetScalarParameterValue(FName(TEXT("LSQuadsY")), (float)LandscapeQuads.Y))
-			{
-				UE_LOG(LogWaterEditor, Error, TEXT("Failed to set \"LSQuadsY\" on Landscape MaterialParameterCollection"));
-			}
+		if (!LandscapeCollectionInstance->SetScalarParameterValue(FName(TEXT("LSQuadsX")), (float)LandscapeQuads.X))
+		{
+			UE_LOG(LogWaterEditor, Error, TEXT("Failed to set \"LSQuadsX\" on Landscape MaterialParameterCollection"));
+		}
+		if (!LandscapeCollectionInstance->SetScalarParameterValue(FName(TEXT("LSQuadsY")), (float)LandscapeQuads.Y))
+		{
+			UE_LOG(LogWaterEditor, Error, TEXT("Failed to set \"LSQuadsY\" on Landscape MaterialParameterCollection"));
+		}
 
-			if (!LandscapeCollectionInstance->SetScalarParameterValue(FName(TEXT("WorldSizeX")), WorldSize.X))
-			{
-				UE_LOG(LogWaterEditor, Error, TEXT("Failed to set \"WorldSizeX\" on Landscape MaterialParameterCollection"));
-			}
-			if (!LandscapeCollectionInstance->SetScalarParameterValue(FName(TEXT("WorldSizeY")), WorldSize.Y))
-			{
-				UE_LOG(LogWaterEditor, Error, TEXT("Failed to set \"WorldSizeY\" on Landscape MaterialParameterCollection"));
-			}
+		if (!LandscapeCollectionInstance->SetScalarParameterValue(FName(TEXT("WorldSizeX")), WorldSize.X))
+		{
+			UE_LOG(LogWaterEditor, Error, TEXT("Failed to set \"WorldSizeX\" on Landscape MaterialParameterCollection"));
+		}
+		if (!LandscapeCollectionInstance->SetScalarParameterValue(FName(TEXT("WorldSizeY")), WorldSize.Y))
+		{
+			UE_LOG(LogWaterEditor, Error, TEXT("Failed to set \"WorldSizeY\" on Landscape MaterialParameterCollection"));
+		}
 
-			if (!LandscapeCollectionInstance->SetVectorParameterValue(FName(TEXT("LandscapeLocation")), FLinearColor(LandscapeTransform.GetLocation())))
-			{
-				UE_LOG(LogWaterEditor, Error, TEXT("Failed to set \"LandscapeLocation\" on Landscape MaterialParameterCollection"));
-			}
-			if (!LandscapeCollectionInstance->SetScalarParameterValue(FName(TEXT("LandscapeZLocation")), LandscapeTransform.GetLocation().Z))
-			{
-				UE_LOG(LogWaterEditor, Error, TEXT("Failed to set \"LandscapeZLocation\" on Landscape MaterialParameterCollection"));
-			}
-			// TODO [jonathan.bard] : find out what this 128.0f corresponds to and put in a constant : ZSCALE in LandscapeLayersPS.usf maybe ??
-			if (!LandscapeCollectionInstance->SetScalarParameterValue(FName(TEXT("LandscapeZScale")), LandscapeTransform.GetScale3D().Z / 128.0f))
-			{
-				UE_LOG(LogWaterEditor, Error, TEXT("Failed to set \"LandscapeZScale\" on Landscape MaterialParameterCollection"));
-			}
-			if (!LandscapeCollectionInstance->SetVectorParameterValue(FName(TEXT("RTWorldSize")), FLinearColor(RTWorldSizeVector)))
-			{
-				UE_LOG(LogWaterEditor, Error, TEXT("Failed to set \"RTWorldSize\" on Landscape MaterialParameterCollection"));
-			}
-			if (!LandscapeCollectionInstance->SetVectorParameterValue(FName(TEXT("RTWorldLocation")), FLinearColor(RTWorldLocation)))
-			{
-				UE_LOG(LogWaterEditor, Error, TEXT("Failed to set \"RTWorldLocation\" on Landscape MaterialParameterCollection"));
-			}
-			if (!LandscapeCollectionInstance->SetScalarParameterValue(FName(TEXT("WaterClearHeight")), WaterClearHeight))
-			{
-				UE_LOG(LogWaterEditor, Error, TEXT("Failed to set \"WaterClearHeight\" on Landscape MaterialParameterCollection"));
-			}
+		if (!LandscapeCollectionInstance->SetVectorParameterValue(FName(TEXT("LandscapeLocation")), FLinearColor(LandscapeTransform.GetLocation())))
+		{
+			UE_LOG(LogWaterEditor, Error, TEXT("Failed to set \"LandscapeLocation\" on Landscape MaterialParameterCollection"));
+		}
+		if (!LandscapeCollectionInstance->SetScalarParameterValue(FName(TEXT("LandscapeZLocation")), LandscapeTransform.GetLocation().Z))
+		{
+			UE_LOG(LogWaterEditor, Error, TEXT("Failed to set \"LandscapeZLocation\" on Landscape MaterialParameterCollection"));
+		}
+		// TODO [jonathan.bard] : find out what this 128.0f corresponds to and put in a constant : ZSCALE in LandscapeLayersPS.usf maybe ??
+		if (!LandscapeCollectionInstance->SetScalarParameterValue(FName(TEXT("LandscapeZScale")), LandscapeTransform.GetScale3D().Z / 128.0f))
+		{
+			UE_LOG(LogWaterEditor, Error, TEXT("Failed to set \"LandscapeZScale\" on Landscape MaterialParameterCollection"));
+		}
+		if (!LandscapeCollectionInstance->SetVectorParameterValue(FName(TEXT("RTWorldSize")), FLinearColor(RTWorldSizeVector)))
+		{
+			UE_LOG(LogWaterEditor, Error, TEXT("Failed to set \"RTWorldSize\" on Landscape MaterialParameterCollection"));
+		}
+		if (!LandscapeCollectionInstance->SetVectorParameterValue(FName(TEXT("RTWorldLocation")), FLinearColor(RTWorldLocation)))
+		{
+			UE_LOG(LogWaterEditor, Error, TEXT("Failed to set \"RTWorldLocation\" on Landscape MaterialParameterCollection"));
 		}
+		if (!LandscapeCollectionInstance->SetScalarParameterValue(FName(TEXT("WaterClearHeight")), WaterClearHeight))
+		{
+			UE_LOG(LogWaterEditor, Error, TEXT("Failed to set \"WaterClearHeight\" on Landscape MaterialParameterCollection"));
+		}
+		
 	}
 }
 
diff --git a/Engine/Source/Editor/WorldPartitionEditor/Private/WorldPartitionEditorModule.cpp b/Engine/Source/Editor/WorldPartitionEditor/Private/WorldPartitionEditorModule.cpp
--- a/Engine/Source/Editor/WorldPartitionEditor/Private/WorldPartitionEditorModule.cpp
+++ b/Engine/Source/Editor/WorldPartitionEditor/Private/WorldPartitionEditorModule.cpp
@@ -410,10 +410,23 @@ void FWorldPartitionEditorModule::RunCommandletAsExternalProcess(const FString&
 		FPlatformProcess::Sleep(0.1);
 	}
 
+
+	FPlatformProcess::GetProcReturnCode(ProcessHandle, &OutResult);
+
 	UE_LOG(LogWorldPartitionEditor, Display, TEXT("#### Begin commandlet output ####\n%s"), *OutCommandletOutput);
-	UE_LOG(LogWorldPartitionEditor, Display, TEXT("#### End commandlet output ####"));
 
 	FPlatformProcess::ClosePipe(ReadPipe, WritePipe);
+
+	if (OutResult == 0)
+	{
+		UE_LOG(LogWorldPartitionEditor, Display, TEXT("#### Commandlet Succeeded ####"));
+	}
+	else
+	{
+		UE_LOG(LogWorldPartitionEditor, Error, TEXT("#### Commandlet Failed ####"));
+		UE_LOG(LogWorldPartitionEditor, Error, TEXT("%s %s"), *CurrentExecutableName, *Arguments);
+		UE_LOG(LogWorldPartitionEditor, Error, TEXT("Return Code: %i"), OutResult);
+	}
 }
 
 bool FWorldPartitionEditorModule::ConvertMap(const FString& InLongPackageName)
diff --git a/Engine/Source/Runtime/Landscape/Private/LandscapeEditLayers.cpp b/Engine/Source/Runtime/Landscape/Private/LandscapeEditLayers.cpp
--- a/Engine/Source/Runtime/Landscape/Private/LandscapeEditLayers.cpp
+++ b/Engine/Source/Runtime/Landscape/Private/LandscapeEditLayers.cpp
@@ -7883,6 +7883,11 @@ void ALandscape::UpdateLayersContent(bool bInWaitForStreaming, bool bInSkipMonit
 	{
 		return;
 	}
+
+	auto GetCurrentTime = []()
+	{
+		return FSlateApplicationBase::IsInitialized() ? FSlateApplicationBase::Get().GetCurrentTime() : 0.0f;
+	};
 	
 	// The Edit layers shaders only work on SM5 : cancel any update that might happen when SM5+ shading model is not active :
 	if (World->FeatureLevel < ERHIFeatureLevel::SM5)
@@ -7908,7 +7913,7 @@ void ALandscape::UpdateLayersContent(bool bInWaitForStreaming, bool bInSkipMonit
 	bResourcesReady &= PrepareLayersTextureResources(bInWaitForStreaming);
 	if (!bResourcesReady)
 	{
-		WaitingForLandscapeTextureResourcesStartTime = FSlateApplicationBase::Get().GetCurrentTime();
+		WaitingForLandscapeTextureResourcesStartTime = GetCurrentTime(); 
 		if (!WaitingForTexturesNotification.IsValid())
 		{
 			WaitingForTexturesNotification = MakeShared<FLandscapeNotification>(this, FLandscapeNotification::EType::LandscapeTextureResourcesNotReady);
@@ -7928,7 +7933,7 @@ void ALandscape::UpdateLayersContent(bool bInWaitForStreaming, bool bInSkipMonit
 	bResourcesReady &= PrepareLayersBrushResources(World->FeatureLevel, bInWaitForStreaming);
 	if (!bResourcesReady)
 	{
-		WaitingForLandscapeBrushResourcesStartTime = FSlateApplicationBase::Get().GetCurrentTime();
+		WaitingForLandscapeBrushResourcesStartTime = GetCurrentTime();
 		if (!WaitingForBrushesNotification.IsValid())
 		{
 			WaitingForBrushesNotification = MakeShared<FLandscapeNotification>(this, FLandscapeNotification::EType::LandscapeBrushResourcesNotReady);
diff --git a/Engine/Source/Runtime/Landscape/Private/LandscapeNotification.cpp b/Engine/Source/Runtime/Landscape/Private/LandscapeNotification.cpp
--- a/Engine/Source/Runtime/Landscape/Private/LandscapeNotification.cpp
+++ b/Engine/Source/Runtime/Landscape/Private/LandscapeNotification.cpp
@@ -37,7 +37,7 @@ bool FLandscapeNotification::operator < (const FLandscapeNotification& Other) co
 
 bool FLandscapeNotification::ShouldShowNotification() const 
 {
-	return ConditionCallback() && (FSlateApplicationBase::Get().GetCurrentTime() >= NotificationStartTime || NotificationStartTime < 0.0);
+	return ConditionCallback() && FSlateApplicationBase::IsInitialized() && (FSlateApplicationBase::Get().GetCurrentTime() >= NotificationStartTime || NotificationStartTime < 0.0);
 }
 
 void FLandscapeNotification::SetNotificationText()
