Add log to help track down a crash

diff --git a/Engine/Plugins/Experimental/Water/Source/Editor/Private/WaterBrushManager.cpp b/Engine/Plugins/Experimental/Water/Source/Editor/Private/WaterBrushManager.cpp
--- a/Engine/Plugins/Experimental/Water/Source/Editor/Private/WaterBrushManager.cpp
+++ b/Engine/Plugins/Experimental/Water/Source/Editor/Private/WaterBrushManager.cpp
@@ -1109,6 +1109,12 @@ void AWaterBrushManager::SetMPCParams()
 		FVector RTWorldLocation, RTWorldSizeVector;
 		ComputeWaterLandscapeInfo(RTWorldLocation, RTWorldSizeVector);
 
+		if (GEditor == nullptr)
+		{
+			UE_LOG(LogWaterEditor, Error, TEXT("GEditor is null for AWaterBrushManager::SetMPCParams(): %s"), *GetFullName());
+			return;
+		}
+
 		UWaterEditorSubsystem* WaterEditorSubsystem = GEditor->GetEditorSubsystem<UWaterEditorSubsystem>();
 		check(WaterEditorSubsystem);
 		UMaterialParameterCollection* LandscapeCollection =  WaterEditorSubsystem->GetLandscapeMaterialParameterCollection();
