WaterBodyCustomComponent - Make sure the custom mesh component is not HLOD relevant
The water body component is already responsible for generating the HLOD of the WaterBody actor

diff --git a/Engine/Plugins/Experimental/Water/Source/Runtime/Private/WaterBodyCustomComponent.cpp b/Engine/Plugins/Experimental/Water/Source/Runtime/Private/WaterBodyCustomComponent.cpp
--- a/Engine/Plugins/Experimental/Water/Source/Runtime/Private/WaterBodyCustomComponent.cpp
+++ b/Engine/Plugins/Experimental/Water/Source/Runtime/Private/WaterBodyCustomComponent.cpp
@@ -72,6 +72,7 @@ void UWaterBodyCustomComponent::OnUpdateBody(bool bWithExclusionVolumes)
 		MeshComp = NewObject<UStaticMeshComponent>(OwnerActor, TEXT("CustomMeshComponent"), RF_Transactional);
 		MeshComp->SetNetAddressable(); // it's deterministically named so it's addressable over network (needed for collision)
 		MeshComp->SetupAttachment(this);
+		MeshComp->bEnableAutoLODGeneration = false; // The water body component is already responsible for generating HLOD of the WaterBody actor.
 
 		if(IsRegistered())
 		{
@@ -162,6 +163,19 @@ bool UWaterBodyCustomComponent::IsIconVisible() const
 {
 	return (GetWaterMeshOverride() == nullptr);
 }
+
+void UWaterBodyCustomComponent::PostLoad()
+{
+	Super::PostLoad();
+
+	// Make sure the custom mesh component is not HLOD relevant
+	// The water body component is already responsible for generating 
+	// the HLOD of the WaterBody actor.
+	if (MeshComp && MeshComp->IsHLODRelevant())
+	{
+		MeshComp->bEnableAutoLODGeneration = false;
+	}
+}
 #endif // WITH_EDITOR
 
 #undef LOCTEXT_NAMESPACE
diff --git a/Engine/Plugins/Experimental/Water/Source/Runtime/Public/WaterBodyCustomComponent.h b/Engine/Plugins/Experimental/Water/Source/Runtime/Public/WaterBodyCustomComponent.h
--- a/Engine/Plugins/Experimental/Water/Source/Runtime/Public/WaterBodyCustomComponent.h
+++ b/Engine/Plugins/Experimental/Water/Source/Runtime/Public/WaterBodyCustomComponent.h
@@ -36,6 +36,8 @@ protected:
 	virtual const TCHAR* GetWaterSpriteTextureName() const override;
 
 	virtual bool IsIconVisible() const override;
+
+	virtual void PostLoad() override;
 #endif // WITH_EDITOR
 
 protected:
