Procedural Rivers: Add River prototype using procedural routes and water body

- Water Spline: Expose Synchronize / BroadcastDataChange funciton to BP

diff --git a/Engine/Plugins/Experimental/Water/Source/Runtime/Private/WaterSplineComponent.cpp b/Engine/Plugins/Experimental/Water/Source/Runtime/Private/WaterSplineComponent.cpp
--- a/Engine/Plugins/Experimental/Water/Source/Runtime/Private/WaterSplineComponent.cpp
+++ b/Engine/Plugins/Experimental/Water/Source/Runtime/Private/WaterSplineComponent.cpp
@@ -130,6 +130,20 @@ FBoxSphereBounds UWaterSplineComponent::CalcBounds(const FTransform& LocalToWorl
 	}
 }
 
+void UWaterSplineComponent::K2_SynchronizeAndBroadcastDataChange()
+{
+#if WITH_EDITOR
+	SynchronizeWaterProperties();
+
+	FPropertyChangedEvent PropertyChangedEvent(FindFProperty<FProperty>(UWaterSplineComponent::StaticClass(), TEXT("SplineCurves")), 1 << 6);
+	FOnWaterSplineDataChangedParams OnWaterSplineDataChangedParams(PropertyChangedEvent);
+	OnWaterSplineDataChangedParams.bUserTriggered = true;
+
+	WaterSplineDataChangedEvent.Broadcast(OnWaterSplineDataChangedParams);
+#endif
+}
+
+
 #if WITH_EDITOR
 
 bool UWaterSplineComponent::CanEditChange(const FProperty* InProperty) const
diff --git a/Engine/Plugins/Experimental/Water/Source/Runtime/Public/WaterSplineComponent.h b/Engine/Plugins/Experimental/Water/Source/Runtime/Public/WaterSplineComponent.h
--- a/Engine/Plugins/Experimental/Water/Source/Runtime/Public/WaterSplineComponent.h
+++ b/Engine/Plugins/Experimental/Water/Source/Runtime/Public/WaterSplineComponent.h
@@ -53,6 +53,13 @@ public:
 
 	virtual bool AllowsSplinePointScaleEditing() const override { return false; }
 
+	/*
+	 * Call to update water spline
+	 * Necessary if using USplineComponent::AddPoint(s) instead of editing the spline in editor
+	 */
+	UFUNCTION(BlueprintCallable, Category = Water, DisplayName="Synchronize And Broadcast Data Change")
+	void K2_SynchronizeAndBroadcastDataChange();
+
 #if WITH_EDITOR
 	DECLARE_EVENT(UWaterSplineComponent, UE_DEPRECATED(5.1, "Use FOnWaterSplineDataChanged") FOnSplineDataChanged);
 	UE_DEPRECATED(5.1, "Use OnWaterSplineDataChanged")
