Water: Ensure water spline properties are synchronized everytime the water body changes.

diff --git a/Engine/Plugins/Experimental/Water/Source/Runtime/Private/WaterBodyComponent.cpp b/Engine/Plugins/Experimental/Water/Source/Runtime/Private/WaterBodyComponent.cpp
--- a/Engine/Plugins/Experimental/Water/Source/Runtime/Private/WaterBodyComponent.cpp
+++ b/Engine/Plugins/Experimental/Water/Source/Runtime/Private/WaterBodyComponent.cpp
@@ -1368,8 +1368,9 @@ void UWaterBodyComponent::UpdateAll(const FOnWaterBodyChangedParams& InParams)
 
 void UWaterBodyComponent::UpdateSplineComponent()
 {
-	if (USplineComponent* WaterSpline = GetWaterSpline())
+	if (UWaterSplineComponent* WaterSpline = GetWaterSpline())
 	{
+		WaterSpline->SynchronizeWaterProperties();
 		WaterSpline->SetClosedLoop(IsWaterSplineClosedLoop());
 	}
 }
