Water: fixed non-editor compile.

diff --git a/Engine/Plugins/Experimental/Water/Source/Runtime/Private/WaterBodyComponent.cpp b/Engine/Plugins/Experimental/Water/Source/Runtime/Private/WaterBodyComponent.cpp
--- a/Engine/Plugins/Experimental/Water/Source/Runtime/Private/WaterBodyComponent.cpp
+++ b/Engine/Plugins/Experimental/Water/Source/Runtime/Private/WaterBodyComponent.cpp
@@ -1370,7 +1370,10 @@ void UWaterBodyComponent::UpdateSplineComponent()
 {
 	if (UWaterSplineComponent* WaterSpline = GetWaterSpline())
 	{
+#if WITH_EDITOR
+		// #todo_water: should we expose this at runtime? Might be necessary for dynamically changing water bodies.
 		WaterSpline->SynchronizeWaterProperties();
+#endif // WITH_EDITOR
 		WaterSpline->SetClosedLoop(IsWaterSplineClosedLoop());
 	}
 }
