Water: Fixed optimal landscape lod calculation.

diff --git a/Engine/Plugins/Experimental/Water/Source/Runtime/Private/WaterInfoRendering.cpp b/Engine/Plugins/Experimental/Water/Source/Runtime/Private/WaterInfoRendering.cpp
--- a/Engine/Plugins/Experimental/Water/Source/Runtime/Private/WaterInfoRendering.cpp
+++ b/Engine/Plugins/Experimental/Water/Source/Runtime/Private/WaterInfoRendering.cpp
@@ -27,6 +27,12 @@ static FAutoConsoleVariableRef CVarRenderCaptureNextWaterInfoDraws(
 	RenderCaptureNextWaterInfoDraws,
 	TEXT("Enable capturing of the water info texture for the next N draws"));
 
+static int32 WaterInfoRenderLandscapeMinimumMipLevel = 0;
+static FAutoConsoleVariableRef CVarWaterInfoRenderLandscapeMinimumMipLevel(
+	TEXT("r.Water.WaterInfo.LandscapeMinimumMipLevel"),
+	WaterInfoRenderLandscapeMinimumMipLevel,
+	TEXT("Clamps the minimum allowed mip level for the landscape when rendering the water info texture. Used on the lowest end platforms which cannot support rendering all the landscape vertices at the highest LOD."));
+
 namespace UE::WaterInfo
 {
 
@@ -360,8 +366,8 @@ private:
 			{
 				if (LandscapeSectionInfo != nullptr)
 				{
-					const double LandscapeComponentUnitsPerVertex = LandscapeSectionInfo->ComputeSectionResolution();
-					if (LandscapeComponentUnitsPerVertex <= 0.f)
+					const double LandscapeComponentUnitsPerQuad = LandscapeSectionInfo->ComputeSectionResolution();
+					if (LandscapeComponentUnitsPerQuad <= 0.f)
 					{
 						// No section resolution probably means the section is a mesh proxy, which might not have regular units per vertex.
 						// Avoid computing optimal LOD in this case.
@@ -369,8 +375,8 @@ private:
 					}
 
 					// Derived from:
-					// LandscapeComponentUnitsPerVertex * 2 ^ (LODLevel) <= WaterInfoUnitsPerPixel * 0.5
-					OptimalLODLevel = FMath::Max(0, int32(FMath::FloorLog2(WaterInfoUnitsPerPixel / LandscapeComponentUnitsPerVertex)) - 1);
+					// (ComponentWorldExtent / WaterInfoWorldspaceExtent) * WaterInfoTextureResolution = (NumComponentQuads / 2 ^ (LodLevel))
+					OptimalLODLevel = FMath::Max(WaterInfoRenderLandscapeMinimumMipLevel, FMath::FloorToInt(FMath::Log2(WaterInfoUnitsPerPixel / LandscapeComponentUnitsPerQuad)));
 
 					break;
 				}
