Fix: WaterSpline is checked for null before getting the local bounds

diff --git a/Engine/Plugins/Experimental/Water/Source/Runtime/Private/WaterBodyLakeComponent.cpp b/Engine/Plugins/Experimental/Water/Source/Runtime/Private/WaterBodyLakeComponent.cpp
--- a/Engine/Plugins/Experimental/Water/Source/Runtime/Private/WaterBodyLakeComponent.cpp
+++ b/Engine/Plugins/Experimental/Water/Source/Runtime/Private/WaterBodyLakeComponent.cpp
@@ -127,10 +127,15 @@ bool UWaterBodyLakeComponent::GenerateWaterBodyMesh(UE::Geometry::FDynamicMesh3&
 
 FBoxSphereBounds UWaterBodyLakeComponent::CalcBounds(const FTransform& LocalToWorld) const
 {
-	FBox BoxExtent = GetWaterSpline()->GetLocalBounds().GetBox();
-	BoxExtent.Max.Z += MaxWaveHeightOffset;
-	BoxExtent.Min.Z -= GetChannelDepth();
-	return FBoxSphereBounds(BoxExtent).TransformBy(LocalToWorld);
+	if (UWaterSplineComponent* WaterSpline = GetWaterSpline())
+	{
+		FBox BoxExtent = WaterSpline->GetLocalBounds().GetBox();
+		BoxExtent.Max.Z += MaxWaveHeightOffset;
+		BoxExtent.Min.Z -= GetChannelDepth();
+		return FBoxSphereBounds(BoxExtent).TransformBy(LocalToWorld);
+	}
+	
+	return FBoxSphereBounds(ForceInit);
 }
 
 void UWaterBodyLakeComponent::Reset()
