Water: landscape now renders at 2 quads per pixel in the water info texture.

diff --git a/Engine/Plugins/Experimental/Water/Source/Runtime/Private/WaterInfoRendering.cpp b/Engine/Plugins/Experimental/Water/Source/Runtime/Private/WaterInfoRendering.cpp
--- a/Engine/Plugins/Experimental/Water/Source/Runtime/Private/WaterInfoRendering.cpp
+++ b/Engine/Plugins/Experimental/Water/Source/Runtime/Private/WaterInfoRendering.cpp
@@ -366,7 +366,8 @@ private:
 			{
 				if (LandscapeSectionInfo != nullptr)
 				{
-					const double LandscapeComponentUnitsPerQuad = LandscapeSectionInfo->ComputeSectionResolution();
+					// Double the required landscape resolution to achieve 2 quads per pixel.
+					const double LandscapeComponentUnitsPerQuad = 2.0 * LandscapeSectionInfo->ComputeSectionResolution();
 					if (LandscapeComponentUnitsPerQuad <= 0.f)
 					{
 						// No section resolution probably means the section is a mesh proxy, which might not have regular units per vertex.
