Fixed various WaterZone actor property changes that would not trigger a proper water info/mesh rebuild

diff --git a/Engine/Plugins/Experimental/Water/Source/Runtime/Private/WaterZoneActor.cpp b/Engine/Plugins/Experimental/Water/Source/Runtime/Private/WaterZoneActor.cpp
--- a/Engine/Plugins/Experimental/Water/Source/Runtime/Private/WaterZoneActor.cpp
+++ b/Engine/Plugins/Experimental/Water/Source/Runtime/Private/WaterZoneActor.cpp
@@ -243,36 +243,36 @@ void AWaterZone::PostEditChangeProperty(FPropertyChangedEvent& PropertyChangedEv
 {
 	Super::PostEditChangeProperty(PropertyChangedEvent);
 
-	if (PropertyChangedEvent.MemberProperty != nullptr &&
-		PropertyChangedEvent.MemberProperty->GetFName() == GET_MEMBER_NAME_CHECKED(AWaterZone, ZoneExtent))
+	const FName PropertyName = PropertyChangedEvent.MemberProperty ? PropertyChangedEvent.MemberProperty->GetFName() : NAME_None;
+	if (PropertyName == GET_MEMBER_NAME_CHECKED(AWaterZone, ZoneExtent))
 	{
 		OnExtentChanged();
 	}
-	else if (PropertyChangedEvent.GetPropertyName() == GET_MEMBER_NAME_CHECKED(AWaterZone, BoundsComponent))
+	else if (PropertyName == GET_MEMBER_NAME_CHECKED(AWaterZone, BoundsComponent))
 	{
 		OnBoundsComponentModified();
 	}
-	else if (PropertyChangedEvent.GetPropertyName() == GET_MEMBER_NAME_CHECKED(AWaterZone, RenderTargetResolution))
+	else if (PropertyName == GET_MEMBER_NAME_CHECKED(AWaterZone, RenderTargetResolution))
 	{
 		MarkForRebuild(EWaterZoneRebuildFlags::UpdateWaterInfoTexture);
 	}
-	else if (PropertyChangedEvent.GetPropertyName() == GET_MEMBER_NAME_CHECKED(AWaterZone, bHalfPrecisionTexture))
+	else if (PropertyName == GET_MEMBER_NAME_CHECKED(AWaterZone, bHalfPrecisionTexture))
 	{
 		MarkForRebuild(EWaterZoneRebuildFlags::UpdateWaterInfoTexture);
 	}
-	else if (PropertyChangedEvent.GetPropertyName() == GET_MEMBER_NAME_CHECKED(AWaterZone, VelocityBlurRadius))
+	else if (PropertyName == GET_MEMBER_NAME_CHECKED(AWaterZone, VelocityBlurRadius))
 	{
 		MarkForRebuild(EWaterZoneRebuildFlags::UpdateWaterInfoTexture);
 	}
-	else if (PropertyChangedEvent.GetPropertyName() == GET_MEMBER_NAME_CHECKED(AWaterZone, TessellatedWaterMeshExtent))
+	else if (PropertyName == GET_MEMBER_NAME_CHECKED(AWaterZone, TessellatedWaterMeshExtent))
 	{
 		MarkForRebuild(EWaterZoneRebuildFlags::All);
 	}
-	else if (PropertyChangedEvent.GetPropertyName() == GET_MEMBER_NAME_CHECKED(AWaterZone, bEnableNonTessellatedLODMesh))
+	else if (PropertyName == GET_MEMBER_NAME_CHECKED(AWaterZone, bEnableNonTessellatedLODMesh))
 	{
 		MarkForRebuild(EWaterZoneRebuildFlags::All);
 	}
-	else if (PropertyChangedEvent.GetPropertyName() == GET_MEMBER_NAME_CHECKED(AWaterZone, NonTessellatedLODSectionScale))
+	else if (PropertyName == GET_MEMBER_NAME_CHECKED(AWaterZone, NonTessellatedLODSectionScale))
 	{
 		MarkForRebuild(EWaterZoneRebuildFlags::All);
 	}
