Fixed r.water.visualizeunderwaterpostprocess, which was showing a misleading material name (the topmost parent)

diff --git a/Engine/Plugins/Experimental/Water/Source/Runtime/Private/WaterSubsystem.cpp b/Engine/Plugins/Experimental/Water/Source/Runtime/Private/WaterSubsystem.cpp
--- a/Engine/Plugins/Experimental/Water/Source/Runtime/Private/WaterSubsystem.cpp
+++ b/Engine/Plugins/Experimental/Water/Source/Runtime/Private/WaterSubsystem.cpp
@@ -744,11 +744,15 @@ void UWaterSubsystem::ShowOnScreenDebugInfo(const FVector& InViewLocation, const
 
 	if (InDebugInfo.ActiveWaterBodyComponent.IsValid())
 	{
-		UMaterialInstanceDynamic* MID = InDebugInfo.ActiveWaterBodyComponent->GetUnderwaterPostProcessMaterialInstance();
-		FString MaterialName = MID ? MID->GetMaterial()->GetName() : TEXT("No material");
+		FString MaterialDescription(TEXT("No material"));
+		if (UMaterialInstanceDynamic* MID = InDebugInfo.ActiveWaterBodyComponent->GetUnderwaterPostProcessMaterialInstance())
+		{
+			check(MID->Parent != nullptr);
+			MaterialDescription = FString::Format(TEXT("{0} (parent: {1})"), { MID->Parent->GetName(), MID->GetMaterial()->GetName() });
+		}
 		OutputStrings.Add(FText::Format(LOCTEXT("VisualizeActiveUnderwaterPostProcess_ActivePostprocess", "Active underwater post process water body {0} (material: {1})"),
 			FText::FromString(InDebugInfo.ActiveWaterBodyComponent->GetOwner()->GetName()),
-			FText::FromString(MaterialName)));
+			FText::FromString(MaterialDescription)));
 	}
 	else
 	{
