Water: initialize Oceans with a more sensible collision extent matching the size of the water zone it was placed within. Users can still tweak the value manually afterwards but this should be a more intuitive default.

diff --git a/Engine/Plugins/Experimental/Water/Source/Editor/Private/WaterBodyActorFactory.cpp b/Engine/Plugins/Experimental/Water/Source/Editor/Private/WaterBodyActorFactory.cpp
--- a/Engine/Plugins/Experimental/Water/Source/Editor/Private/WaterBodyActorFactory.cpp
+++ b/Engine/Plugins/Experimental/Water/Source/Editor/Private/WaterBodyActorFactory.cpp
@@ -11,6 +11,7 @@
 #include "WaterBodyRiverComponent.h"
 #include "WaterSplineComponent.h"
 #include "WaterWaves.h"
+#include "WaterZoneActor.h"
 
 #include UE_INLINE_GENERATED_CPP_BY_NAME(WaterBodyActorFactory)
 
@@ -130,6 +131,15 @@ void UWaterBodyOceanActorFactory::PostSpawnActor(UObject* Asset, AActor* NewActo
 
 	UWaterSplineComponent* WaterSpline = WaterBodyComponent->GetWaterSpline();
 	WaterSpline->ResetSpline({ FVector(10000, -10000, 0), FVector(10000,  10000, 0), FVector(-10000,  10000, 0), FVector(-10000, -10000, 0) });
+
+	if (const AWaterZone* OwningWaterZone = WaterBodyComponent->GetWaterZone())
+	{
+		if (UWaterBodyOceanComponent* OceanComponent = Cast<UWaterBodyOceanComponent>(WaterBodyComponent))
+		{
+			const double ExistingCollisionHeight = OceanComponent->GetCollisionExtents().Z;
+			OceanComponent->SetCollisionExtents(FVector(OwningWaterZone->GetZoneExtent() / 2.0, ExistingCollisionHeight));
+		}
+	}
 }
 
 // --------------------------------------------------
diff --git a/Engine/Plugins/Experimental/Water/Source/Runtime/Private/WaterBodyOceanComponent.cpp b/Engine/Plugins/Experimental/Water/Source/Runtime/Private/WaterBodyOceanComponent.cpp
--- a/Engine/Plugins/Experimental/Water/Source/Runtime/Private/WaterBodyOceanComponent.cpp
+++ b/Engine/Plugins/Experimental/Water/Source/Runtime/Private/WaterBodyOceanComponent.cpp
@@ -80,6 +80,16 @@ void UWaterBodyOceanComponent::BeginUpdateWaterBody()
 }
 
 #if WITH_EDITOR
+void UWaterBodyOceanComponent::SetCollisionExtents(const FVector& NewExtents)
+{
+	CollisionExtents = NewExtents;
+
+	FOnWaterBodyChangedParams Params;
+	Params.bShapeOrPositionChanged = true;
+	Params.bWeightmapSettingsChanged = false;
+	UpdateAll(Params);
+}
+
 void UWaterBodyOceanComponent::OnPostEditChangeProperty(FOnWaterBodyChangedParams& InOutOnWaterBodyChangedParams)
 {
 	Super::OnPostEditChangeProperty(InOutOnWaterBodyChangedParams);
diff --git a/Engine/Plugins/Experimental/Water/Source/Runtime/Public/WaterBodyOceanComponent.h b/Engine/Plugins/Experimental/Water/Source/Runtime/Public/WaterBodyOceanComponent.h
--- a/Engine/Plugins/Experimental/Water/Source/Runtime/Public/WaterBodyOceanComponent.h
+++ b/Engine/Plugins/Experimental/Water/Source/Runtime/Public/WaterBodyOceanComponent.h
@@ -24,6 +24,10 @@ public:
 	virtual FVector GetCollisionExtents() const override { return CollisionExtents; }
 	virtual void SetHeightOffset(float InHeightOffset) override;
 	virtual float GetHeightOffset() const override { return HeightOffset; }
+	
+#if WITH_EDITOR
+	void SetCollisionExtents(const FVector& NewExtents);
+#endif // WITH_EDITOR
 
 	UE_DEPRECATED(5.1, "Oceans no longer rely on the visual extent parameter making this obsolete. Instead they will be guaranteed to fill the entire water zone to which they belong.")
 	void SetVisualExtents(FVector2D) {}
