Adds r.Water.WaterZoneActor.OverrideNewWaterZoneMinimumMargin CVar, allowing to enforce a minimum distance between landscape bounds and new water zone extent.

diff --git a/Engine/Plugins/Experimental/Water/Source/Editor/Private/WaterEditorModule.cpp b/Engine/Plugins/Experimental/Water/Source/Editor/Private/WaterEditorModule.cpp
--- a/Engine/Plugins/Experimental/Water/Source/Editor/Private/WaterEditorModule.cpp
+++ b/Engine/Plugins/Experimental/Water/Source/Editor/Private/WaterEditorModule.cpp
@@ -48,7 +48,8 @@ EAssetTypeCategories::Type FWaterEditorModule::WaterAssetCategory;
 
 namespace WaterEditorModule
 {
-	static TAutoConsoleVariable<float> CVarOverrideNewWaterZoneScale(TEXT("r.Water.WaterZoneActor.OverrideNewWaterZoneScale"), 0, TEXT("Multiply WaterZone actor extent beyond landscape by this amount. 0 means do override."));
+	static TAutoConsoleVariable<float> CVarOverrideNewWaterZoneScale(TEXT("r.Water.WaterZoneActor.OverrideNewWaterZoneScale"), 0, TEXT("Multiply WaterZone actor extent beyond landscape by this amount. 0 means no override."));
+	static TAutoConsoleVariable<float> CVarOverrideNewWaterZoneMinimumMargin(TEXT("r.Water.WaterZoneActor.OverrideNewWaterZoneMinimumMargin"), 0, TEXT("Enforce a minimum distance between landscape and the WaterZone actor extent. 0 means no override."));
 }
 
 void FWaterEditorModule::StartupModule()
@@ -216,6 +217,7 @@ void FWaterEditorModule::OnLevelActorAddedToWorld(AActor* Actor)
 			WaterZoneBounds += LandscapeBounds;
 		}
 
+		const FBox LandscapeBounds = WaterZoneBounds;
 		const UWaterEditorSettings* WaterEditorSettings = GetDefault<UWaterEditorSettings>();
 		check(WaterEditorSettings != nullptr);
 		// Automatically setup landscape-affecting features (water brush) if needed : 
@@ -325,6 +327,8 @@ void FWaterEditorModule::OnLevelActorAddedToWorld(AActor* Actor)
 						FVector2D NewExtent = 2 * FVector2D(WaterZoneBounds.GetExtent());
 
 						float ZoneExtentScale = WaterEditorModule::CVarOverrideNewWaterZoneScale.GetValueOnGameThread();
+						const float MinimumMargin = WaterEditorModule::CVarOverrideNewWaterZoneMinimumMargin.GetValueOnGameThread();
+
 						if (ZoneExtentScale == 0)
 						{
 							ZoneExtentScale = GetDefault<UWaterEditorSettings>()->WaterZoneActorDefaults.NewWaterZoneScale;
@@ -335,6 +339,15 @@ void FWaterEditorModule::OnLevelActorAddedToWorld(AActor* Actor)
 							NewExtent = FMath::Abs(ZoneExtentScale) * NewExtent;
 						}
 
+						if ((MinimumMargin > 0) && (LandscapeBounds.IsValid))
+						{
+							const FVector2D LandscapeBoundsDiameter = 2.f * FVector2D(LandscapeBounds.GetExtent());
+							const FVector2D MinimumWaterZoneExtent = LandscapeBoundsDiameter + 2.f * MinimumMargin;
+
+							NewExtent.X = (NewExtent.X < MinimumWaterZoneExtent.X) ? MinimumWaterZoneExtent.X : NewExtent.X;
+							NewExtent.Y = (NewExtent.Y < MinimumWaterZoneExtent.Y) ? MinimumWaterZoneExtent.Y : NewExtent.Y;
+						}
+
 						WaterZoneActor->SetZoneExtent(NewExtent);
 					}
 				}
