Water: water bodies will now search through all actor descs to determine their optimal owning water zone rather than only through the loaded actors.

diff --git a/Engine/Plugins/Experimental/Water/Source/Runtime/Private/WaterBodyComponent.cpp b/Engine/Plugins/Experimental/Water/Source/Runtime/Private/WaterBodyComponent.cpp
--- a/Engine/Plugins/Experimental/Water/Source/Runtime/Private/WaterBodyComponent.cpp
+++ b/Engine/Plugins/Experimental/Water/Source/Runtime/Private/WaterBodyComponent.cpp
@@ -216,11 +216,6 @@ FBox UWaterBodyComponent::GetCollisionComponentBounds() const
 	return Box;
 }
 
-FBoxSphereBounds UWaterBodyComponent::CalcBounds(const FTransform& LocalToWorld) const
-{
-	return Super::CalcBounds(LocalToWorld);
-}
-
 AWaterBody* UWaterBodyComponent::GetWaterBodyActor() const
 {
 	// If we have an Owner, it must be an AWaterBody
@@ -394,20 +389,30 @@ void UWaterBodyComponent::UpdateWaterZones()
 {
 	if (UWorld* World = GetWorld())
 	{
-		AWaterZone* OldWaterZone = GetWaterZone();
-		AWaterZone* FoundZone = FindWaterZone();
+		TSoftObjectPtr<AWaterZone> FoundZone = nullptr;
+		if (!WaterZoneOverride.IsNull())
+		{
+			FoundZone = WaterZoneOverride.Get();
+		}
+		else
+		{
+			const FBox Bounds3D = Bounds.GetBox();
+			if (const UWaterSubsystem* WaterSubsystem = UWaterSubsystem::GetWaterSubsystem(GetWorld()))
+			{
+				FoundZone = WaterSubsystem->FindWaterZone(FBox2D(FVector2D(Bounds3D.Min), FVector2D(Bounds3D.Max)));
+			}
+		}
 
-		if (OldWaterZone != FoundZone)
+		if (OwningWaterZone != FoundZone)
 		{
-			if (OldWaterZone)
+			if (AWaterZone* OldOwningZonePtr = OwningWaterZone.Get())
 			{
-				OldWaterZone->RemoveWaterBodyComponent(this);
-				OldWaterZone = nullptr;
+				OldOwningZonePtr->RemoveWaterBodyComponent(this);
 			}
 
-			if (FoundZone)
+			if (AWaterZone* FoundZonePtr = FoundZone.Get())
 			{
-				FoundZone->AddWaterBodyComponent(this);
+				FoundZonePtr->AddWaterBodyComponent(this);
 			}
 
 			OwningWaterZone = FoundZone;
@@ -1738,57 +1743,6 @@ UWaterWavesBase* UWaterBodyComponent::GetWaterWaves() const
 	return nullptr;
 }
 
-AWaterZone* UWaterBodyComponent::FindWaterZone() const
-{
-	TRACE_CPUPROFILER_EVENT_SCOPE(UWaterBodyComponent::FindWaterZone);
-
-	if (!WaterZoneOverride.IsNull())
-	{
-		return WaterZoneOverride.Get();
-	}
-
-	// Score each overlapping water zone and then pick the best.
-	TArray<TPair<int32, AWaterZone*>, TInlineAllocator<4>> ViableZones;
-
-	if (const UWorld* World = GetWorld())
-	{
-		const ULevel* PreferredLevel = GetTypedOuter<ULevel>();
-		if (PreferredLevel)
-		{
-			for (AWaterZone* WaterZone : TActorRange<AWaterZone>(World, AWaterZone::StaticClass(), EActorIteratorFlags::SkipPendingKill))
-			{
-				// WaterZone->GetZoneExtents returns the full extent of the zone but BoxSphereBounds expects a half-extent.
-				const FVector2D WaterZoneLocation = FVector2D(WaterZone->GetActorLocation());
-				const FVector2D WaterZoneHalfExtent = WaterZone->GetZoneExtent() / 2.0;
-				const FBox2D WaterZoneBounds(WaterZoneLocation - WaterZoneHalfExtent, WaterZoneLocation + WaterZoneHalfExtent);
-				const FBox ComponentBounds3d = CalcBounds(GetComponentTransform()).GetBox();
-				const FBox2D ComponentBounds = FBox2D(FVector2D(ComponentBounds3d.Min), FVector2D(ComponentBounds3d.Max));
-
-				// Only consider WaterZones which this component overlaps 
-				// but prefer choosing water zones which are part of the same outered level.
-				if (ComponentBounds.Intersect(WaterZoneBounds))
-				{
-					if (WaterZone->GetTypedOuter<ULevel>() == PreferredLevel)
-					{
-						ViableZones.Add({ WaterZone->GetOverlapPriority(), WaterZone });
-						continue;
-					}
-
-					// Fallback to zones not in the preferred level only as a final resort.
-					ViableZones.Add({ TNumericLimits<int32>::Min(), WaterZone});
-				}
-			}
-		}
-	}
-	
-	if (ViableZones.Num() == 0)
-	{
-		return nullptr;
-	}
-
-	return Algo::MaxElementBy(ViableZones, [](const TPair<int32, AWaterZone*>& A) { return A.Key; })->Value;
-}
-
 static inline FColor PackFlowData(float VelocityMagnitude, float DirectionAngle, float MaxVelocity)
 {
 	check((DirectionAngle >= 0.f) && (DirectionAngle <= TWO_PI));
diff --git a/Engine/Plugins/Experimental/Water/Source/Runtime/Private/WaterSubsystem.cpp b/Engine/Plugins/Experimental/Water/Source/Runtime/Private/WaterSubsystem.cpp
--- a/Engine/Plugins/Experimental/Water/Source/Runtime/Private/WaterSubsystem.cpp
+++ b/Engine/Plugins/Experimental/Water/Source/Runtime/Private/WaterSubsystem.cpp
@@ -26,6 +26,12 @@
 #include "WaterSplineComponent.h"
 #include "WaterUtils.h"
 #include "WaterViewExtension.h"
+#include "Algo/MaxElement.h"
+
+#if WITH_EDITOR
+#include "WaterZoneActorDesc.h"
+#include "WorldPartition/WorldPartitionHelpers.h"
+#endif // WITH_DITOR
 
 #include UE_INLINE_GENERATED_CPP_BY_NAME(WaterSubsystem)
 
@@ -525,9 +531,7 @@ void UWaterSubsystem::MarkWaterZonesInRegionForRebuild(const FBox2D& InUpdateReg
 	{
 		for (AWaterZone* WaterZone : TActorRange<AWaterZone>(World))
 		{
-			FVector2D ZoneExtent = WaterZone->GetZoneExtent();
-			FVector2D ZoneLocation = FVector2D(WaterZone->GetActorLocation());
-			const FBox2D WaterZoneBounds(ZoneLocation - ZoneExtent * 0.5, ZoneLocation + ZoneExtent * 0.5);
+			const FBox2D WaterZoneBounds = WaterZone->GetZoneBounds();
 
 			if (WaterZoneBounds.Intersect(InUpdateRegion))
 			{
@@ -537,6 +541,59 @@ void UWaterSubsystem::MarkWaterZonesInRegionForRebuild(const FBox2D& InUpdateReg
 	}
 }
 
+TSoftObjectPtr<AWaterZone> UWaterSubsystem::FindWaterZone(const FBox2D& Bounds) const
+{
+	if (!GetWorld())
+	{
+		return {};
+	}
+
+	const UWorld* World = GetWorld();
+
+	// Score each overlapping water zone and then pick the best.
+	TMap<TSoftObjectPtr<AWaterZone>, int32> ViableZones;
+
+#if WITH_EDITOR
+	// Within the editor, we also want to check unloaded actors to ensure that the water body has serialized the best possible water zone, rather than just looking through what might be loaded now.
+	if (GEditor && !World->IsGameWorld())
+	{
+		if (UWorldPartition* WorldPartition = GetWorld()->GetWorldPartition())
+		{
+			FWorldPartitionHelpers::ForEachActorDesc<AWaterZone>(WorldPartition, [Bounds, &ViableZones](const FWorldPartitionActorDesc* ActorDesc)
+			{
+				FWaterZoneActorDesc* WaterZoneActorDesc = (FWaterZoneActorDesc*)ActorDesc;
+				const FBox WaterZoneBounds = WaterZoneActorDesc->GetBounds();
+				const FBox2D WaterZoneBounds2D(FVector2D(WaterZoneBounds.Min), FVector2D(WaterZoneBounds.Max));
+
+				if (Bounds.Intersect(WaterZoneBounds2D))
+				{
+					ViableZones.Emplace(WaterZoneActorDesc->GetActorSoftPath(), WaterZoneActorDesc->GetOverlapPriority());
+				}
+
+				return true;
+			});
+		}
+	}
+#endif // WITH_EDITOR
+
+	for (AWaterZone* WaterZone : TActorRange<AWaterZone>(World, AWaterZone::StaticClass(), EActorIteratorFlags::SkipPendingKill))
+	{
+		const FBox2D WaterZoneBounds = WaterZone->GetZoneBounds();
+
+		if (Bounds.Intersect(WaterZoneBounds))
+		{
+			ViableZones.Emplace(WaterZone, WaterZone->GetOverlapPriority());
+		}
+	}
+
+	if (ViableZones.Num() == 0)
+	{
+		return {};
+	}
+
+	return Algo::MaxElementBy(ViableZones, [](const TPair<TSoftObjectPtr<AWaterZone>, int32>& A) { return A.Value; })->Key;
+}
+
 void UWaterSubsystem::NotifyWaterScalabilityChangedInternal(IConsoleVariable* CVar)
 {
 	OnWaterScalabilityChanged.Broadcast();
diff --git a/Engine/Plugins/Experimental/Water/Source/Runtime/Private/WaterZoneActor.cpp b/Engine/Plugins/Experimental/Water/Source/Runtime/Private/WaterZoneActor.cpp
--- a/Engine/Plugins/Experimental/Water/Source/Runtime/Private/WaterZoneActor.cpp
+++ b/Engine/Plugins/Experimental/Water/Source/Runtime/Private/WaterZoneActor.cpp
@@ -17,12 +17,13 @@
 
 #include UE_INLINE_GENERATED_CPP_BY_NAME(WaterZoneActor)
 
-#if	WITH_EDITOR
+#if WITH_EDITOR
 #include "Algo/Transform.h"
 #include "LevelEditor.h"
 #include "Modules/ModuleManager.h"
 #include "WaterIconHelper.h"
 #include "Components/BoxComponent.h"
+#include "WaterZoneActorDesc.h"
 #endif // WITH_EDITOR
 
 static int32 ForceUpdateWaterInfoNextFrames = 0;
@@ -81,6 +82,16 @@ void AWaterZone::SetZoneExtent(FVector2D NewExtent)
 	OnExtentChanged();
 }
 
+FBox2D AWaterZone::GetZoneBounds() const
+{
+	// GetZoneExtents returns the full extent of the zone but BoxSphereBounds expects a half-extent.
+	const FVector2D WaterZoneLocation = FVector2D(GetActorLocation());
+	const FVector2D WaterZoneHalfExtent = GetZoneExtent() / 2.0;
+	const FBox2D WaterZoneBounds(WaterZoneLocation - WaterZoneHalfExtent, WaterZoneLocation + WaterZoneHalfExtent);
+
+	return WaterZoneBounds;
+}
+
 void AWaterZone::SetRenderTargetResolution(FIntPoint NewResolution)
 {
 	RenderTargetResolution = NewResolution;
@@ -302,6 +313,21 @@ void AWaterZone::OnActorSelectionChanged(const TArray<UObject*>& NewSelection, b
 		MarkForRebuild(EWaterZoneRebuildFlags::UpdateWaterMesh | EWaterZoneRebuildFlags::UpdateWaterInfoTexture);
 	}
 }
+
+TUniquePtr<class FWorldPartitionActorDesc> AWaterZone::CreateClassActorDesc() const
+{
+	return TUniquePtr<FWorldPartitionActorDesc>(new FWaterZoneActorDesc());
+}
+
+FBox AWaterZone::GetStreamingBounds() const
+{
+	const FBox2D ZoneBounds2D = GetZoneBounds();
+	// #todo_water [roey]: Water zone doesn't have an explicit z bounds yet. For now just use the x or y.
+	const double StreamingBoundsZ = FMath::Max(ZoneExtent.X, ZoneExtent.Y);
+
+	return FBox(FVector3d(ZoneBounds2D.Min, -StreamingBoundsZ / 2.), FVector3d(ZoneBounds2D.Max, StreamingBoundsZ / 2.0));
+}
+
 #endif // WITH_EDITOR
 
 void AWaterZone::OnExtentChanged()
diff --git a/Engine/Plugins/Experimental/Water/Source/Runtime/Private/WaterZoneActorDesc.cpp b/Engine/Plugins/Experimental/Water/Source/Runtime/Private/WaterZoneActorDesc.cpp
new file mode 100644
--- /dev/null
+++ b/Engine/Plugins/Experimental/Water/Source/Runtime/Private/WaterZoneActorDesc.cpp
@@ -0,0 +1,40 @@
+// Copyright Epic Games, Inc. All Rights Reserved.
+
+#include "WaterZoneActorDesc.h"
+
+#if WITH_EDITOR
+#include "WaterZoneActor.h"
+#include "WorldPartition/ActorDescContainer.h"
+#include "WorldPartition/WorldPartitionHandle.h"
+#include "UObject/UE5MainStreamObjectVersion.h"
+#include "UObject/UE5ReleaseStreamObjectVersion.h"
+
+void FWaterZoneActorDesc::Init(const AActor* InActor)
+{
+	FWorldPartitionActorDesc::Init(InActor);
+
+	if (const AWaterZone* WaterZone = CastChecked<AWaterZone>(InActor))
+	{
+		OverlapPriority = WaterZone->GetOverlapPriority();
+	}
+}
+
+void FWaterZoneActorDesc::Serialize(FArchive& Ar)
+{
+	FWorldPartitionActorDesc::Serialize(Ar);
+
+	Ar << OverlapPriority;
+}
+
+bool FWaterZoneActorDesc::Equals(const FWorldPartitionActorDesc* Other) const
+{
+	if (FWorldPartitionActorDesc::Equals(Other))
+	{
+		const FWaterZoneActorDesc* WaterZoneActorDesc = (FWaterZoneActorDesc*)Other;
+		return OverlapPriority == WaterZoneActorDesc->OverlapPriority;
+	}
+
+	return false;
+}
+
+#endif
diff --git a/Engine/Plugins/Experimental/Water/Source/Runtime/Private/WaterZoneActorDesc.h b/Engine/Plugins/Experimental/Water/Source/Runtime/Private/WaterZoneActorDesc.h
new file mode 100644
--- /dev/null
+++ b/Engine/Plugins/Experimental/Water/Source/Runtime/Private/WaterZoneActorDesc.h
@@ -0,0 +1,25 @@
+// Copyright Epic Games, Inc. All Rights Reserved.
+
+#pragma once
+
+#if WITH_EDITOR
+#include "CoreMinimal.h"
+#include "UObject/ObjectMacros.h"
+#include "WorldPartition/WorldPartitionActorDesc.h"
+
+/**
+ * ActorDesc for WaterZone actors.
+ */
+class FWaterZoneActorDesc : public FWorldPartitionActorDesc
+{
+public:
+	virtual void Init(const AActor* InActor) override;
+	virtual bool Equals(const FWorldPartitionActorDesc* Other) const override;
+
+	int32 GetOverlapPriority() const { return OverlapPriority; }
+protected:
+	virtual void Serialize(FArchive& Ar) override;
+
+	int32 OverlapPriority;
+};
+#endif
diff --git a/Engine/Plugins/Experimental/Water/Source/Runtime/Public/WaterBodyComponent.h b/Engine/Plugins/Experimental/Water/Source/Runtime/Public/WaterBodyComponent.h
--- a/Engine/Plugins/Experimental/Water/Source/Runtime/Public/WaterBodyComponent.h
+++ b/Engine/Plugins/Experimental/Water/Source/Runtime/Public/WaterBodyComponent.h
@@ -188,8 +188,6 @@ public:
 	/** Returns collision half-extents */
 	virtual FVector GetCollisionExtents() const { return FVector::ZeroVector; }
 
-	virtual FBoxSphereBounds CalcBounds(const FTransform& LocalToWorld) const override;
-
 	/** Sets an additional water height (For internal use. Please use AWaterBodyOcean instead.) */
 	virtual void SetHeightOffset(float InHeightOffset) { check(false); }
 
@@ -482,7 +480,6 @@ protected:
 	EWaterBodyQueryFlags CheckAndAjustQueryFlags(EWaterBodyQueryFlags InQueryFlags) const;
 	void UpdateSplineComponent();
 	void UpdateExclusionVolumes();
-	AWaterZone* FindWaterZone() const;
 	bool UpdateWaterHeight();
 	virtual void CreateOrUpdateWaterMID();
 	void CreateOrUpdateWaterLODMID();
@@ -524,6 +521,8 @@ protected:
 	virtual TSubclassOf<class UHLODBuilder> GetCustomHLODBuilderClass() const override;
 #endif // WITH_EDITOR
 
+	UE_DEPRECATED(5.2, "This function is no longer called. Instead uses UWaterSubsystem::FindWaterZone to determine the right water zone when no override is provided.")
+	AWaterZone* FindWaterZone() const { return nullptr; }
 public:
 	// INavRelevantInterface start
 	virtual void GetNavigationData(struct FNavigationRelevantData& Data) const override;
@@ -644,7 +643,7 @@ protected:
 	UPROPERTY(Transient)
 	mutable TWeakObjectPtr<ALandscapeProxy> Landscape;
 
-	UPROPERTY()
+	UPROPERTY(Category = Water, VisibleAnywhere, AdvancedDisplay)
 	TSoftObjectPtr<AWaterZone> OwningWaterZone;
 
 	UPROPERTY(Category = Water, EditAnywhere, BlueprintReadOnly, AdvancedDisplay)
diff --git a/Engine/Plugins/Experimental/Water/Source/Runtime/Public/WaterSubsystem.h b/Engine/Plugins/Experimental/Water/Source/Runtime/Public/WaterSubsystem.h
--- a/Engine/Plugins/Experimental/Water/Source/Runtime/Public/WaterSubsystem.h
+++ b/Engine/Plugins/Experimental/Water/Source/Runtime/Public/WaterSubsystem.h
@@ -177,6 +177,9 @@ public:
 	void MarkAllWaterZonesForRebuild(EWaterZoneRebuildFlags RebuildFlags = EWaterZoneRebuildFlags::All);
 	void MarkWaterZonesInRegionForRebuild(const FBox2D& InUpdateRegion, EWaterZoneRebuildFlags InRebuildFlags);
 
+	/** Returns the water with the highest priority within the bounds provided. */
+	TSoftObjectPtr<AWaterZone> FindWaterZone(const FBox2D& Bounds) const;
+
 #if WITH_EDITOR
 	/** Little scope object to temporarily change the value of bAllowWaterSubsystemOnPreviewWorld */
 	struct WATER_API FScopedAllowWaterSubsystemOnPreviewWorld
diff --git a/Engine/Plugins/Experimental/Water/Source/Runtime/Public/WaterZoneActor.h b/Engine/Plugins/Experimental/Water/Source/Runtime/Public/WaterZoneActor.h
--- a/Engine/Plugins/Experimental/Water/Source/Runtime/Public/WaterZoneActor.h
+++ b/Engine/Plugins/Experimental/Water/Source/Runtime/Public/WaterZoneActor.h
@@ -7,6 +7,7 @@
 #include "GameFramework/Actor.h"
 #include "Engine/TextureRenderTarget2D.h"
 #include "MaterialShared.h"
+#include "WorldPartition/WorldPartitionActorDesc.h"
 #include "WaterZoneActor.generated.h"
 
 class UWaterMeshComponent;
@@ -44,6 +45,8 @@ public:
 	FVector2D GetZoneExtent() const { return ZoneExtent; }
 	void SetZoneExtent(FVector2D NewExtents);
 
+	FBox2D GetZoneBounds() const;
+
 	void SetRenderTargetResolution(FIntPoint NewResolution);
 	FIntPoint GetRenderTargetResolution() const { return RenderTargetResolution; }
 
@@ -75,6 +78,11 @@ public:
 
 	bool IsNonTessellatedLODMeshEnabled() const { return bEnableNonTessellatedLODMesh; }
 
+#if WITH_EDITOR
+	virtual TUniquePtr<class FWorldPartitionActorDesc> CreateClassActorDesc() const override;
+	virtual FBox GetStreamingBounds() const override;
+#endif //WITH_EDITOR
+
 private:
 
 	/**
@@ -172,4 +180,5 @@ private:
 	UPROPERTY(Transient)
 	TObjectPtr<UTexture2D> WaterVelocityTexture_DEPRECATED;
 #endif // WITH_EDITORONLY_DATA
-};
\ No newline at end of file
+};
+DEFINE_ACTORDESC_TYPE(AWaterZone, FWaterZoneActorDesc);
