Water: removed use of deprecated FWorldPartitionActorDesc::GetBounds call.

diff --git a/Engine/Plugins/Experimental/Water/Source/Runtime/Private/WaterSubsystem.cpp b/Engine/Plugins/Experimental/Water/Source/Runtime/Private/WaterSubsystem.cpp
--- a/Engine/Plugins/Experimental/Water/Source/Runtime/Private/WaterSubsystem.cpp
+++ b/Engine/Plugins/Experimental/Water/Source/Runtime/Private/WaterSubsystem.cpp
@@ -562,7 +562,7 @@ TSoftObjectPtr<AWaterZone> UWaterSubsystem::FindWaterZone(const FBox2D& Bounds)
 			FWorldPartitionHelpers::ForEachActorDesc<AWaterZone>(WorldPartition, [Bounds, &ViableZones](const FWorldPartitionActorDesc* ActorDesc)
 			{
 				FWaterZoneActorDesc* WaterZoneActorDesc = (FWaterZoneActorDesc*)ActorDesc;
-				const FBox WaterZoneBounds = WaterZoneActorDesc->GetBounds();
+				const FBox WaterZoneBounds = WaterZoneActorDesc->GetEditorBounds();
 				const FBox2D WaterZoneBounds2D(FVector2D(WaterZoneBounds.Min), FVector2D(WaterZoneBounds.Max));
 
 				if (Bounds.Intersect(WaterZoneBounds2D))
