Nonunity fix
#rnx

diff --git a/Engine/Plugins/Experimental/Water/Source/Runtime/Private/WaterSubsystem.cpp b/Engine/Plugins/Experimental/Water/Source/Runtime/Private/WaterSubsystem.cpp
--- a/Engine/Plugins/Experimental/Water/Source/Runtime/Private/WaterSubsystem.cpp
+++ b/Engine/Plugins/Experimental/Water/Source/Runtime/Private/WaterSubsystem.cpp
@@ -31,6 +31,8 @@
 #if WITH_EDITOR
 #include "WaterZoneActorDesc.h"
 #include "WorldPartition/WorldPartitionHelpers.h"
+
+extern UNREALED_API UEditorEngine* GEditor;
 #endif // WITH_DITOR
 
 #include UE_INLINE_GENERATED_CPP_BY_NAME(WaterSubsystem)
