Fixes invalid WaterInfoTexture when landscape actors are being streamed out.

diff --git a/Engine/Plugins/Experimental/Water/Source/Runtime/Private/WaterSubsystem.cpp b/Engine/Plugins/Experimental/Water/Source/Runtime/Private/WaterSubsystem.cpp
--- a/Engine/Plugins/Experimental/Water/Source/Runtime/Private/WaterSubsystem.cpp
+++ b/Engine/Plugins/Experimental/Water/Source/Runtime/Private/WaterSubsystem.cpp
@@ -533,7 +533,7 @@ void UWaterSubsystem::MarkWaterZonesInRegionForRebuild(const FBox2D& InUpdateReg
 	{
 		for (AWaterZone* WaterZone : TActorRange<AWaterZone>(World))
 		{
-			const FBox2D WaterZoneBounds = WaterZone->GetZoneBounds();
+			const FBox2D WaterZoneBounds = WaterZone->GetZoneBounds2D();
 
 			if (WaterZoneBounds.Intersect(InUpdateRegion))
 			{
@@ -580,7 +580,7 @@ TSoftObjectPtr<AWaterZone> UWaterSubsystem::FindWaterZone(const FBox2D& Bounds)
 
 	for (AWaterZone* WaterZone : TActorRange<AWaterZone>(World, AWaterZone::StaticClass(), EActorIteratorFlags::SkipPendingKill))
 	{
-		const FBox2D WaterZoneBounds = WaterZone->GetZoneBounds();
+		const FBox2D WaterZoneBounds = WaterZone->GetZoneBounds2D();
 
 		if (Bounds.Intersect(WaterZoneBounds))
 		{
diff --git a/Engine/Plugins/Experimental/Water/Source/Runtime/Private/WaterZoneActor.cpp b/Engine/Plugins/Experimental/Water/Source/Runtime/Private/WaterZoneActor.cpp
--- a/Engine/Plugins/Experimental/Water/Source/Runtime/Private/WaterZoneActor.cpp
+++ b/Engine/Plugins/Experimental/Water/Source/Runtime/Private/WaterZoneActor.cpp
@@ -7,6 +7,7 @@
 #include "WaterBodyActor.h"
 #include "WaterInfoRendering.h"
 #include "EngineUtils.h"
+#include "LandscapeComponent.h"
 #include "LandscapeProxy.h"
 #include "WaterUtils.h"
 #include "Materials/MaterialInstanceDynamic.h"
@@ -14,6 +15,7 @@
 #include "WaterBodyOceanComponent.h"
 #include "RenderCaptureInterface.h"
 #include "WaterViewExtension.h"
+#include "Algo/AnyOf.h"
 
 #include UE_INLINE_GENERATED_CPP_BY_NAME(WaterZoneActor)
 
@@ -82,7 +84,7 @@ void AWaterZone::SetZoneExtent(FVector2D NewExtent)
 	OnExtentChanged();
 }
 
-FBox2D AWaterZone::GetZoneBounds() const
+FBox2D AWaterZone::GetZoneBounds2D() const
 {
 	// GetZoneExtents returns the full extent of the zone but BoxSphereBounds expects a half-extent.
 	const FVector2D WaterZoneLocation = FVector2D(GetActorLocation());
@@ -92,6 +94,15 @@ FBox2D AWaterZone::GetZoneBounds() const
 	return WaterZoneBounds;
 }
 
+FBox AWaterZone::GetZoneBounds() const
+{
+	const FBox2D ZoneBounds2D = GetZoneBounds2D();
+	// #todo_water [roey]: Water zone doesn't have an explicit z bounds yet. For now just use the x or y.
+	const double StreamingBoundsZ = FMath::Max(ZoneExtent.X, ZoneExtent.Y);
+
+	return FBox(FVector3d(ZoneBounds2D.Min, -StreamingBoundsZ / 2.), FVector3d(ZoneBounds2D.Max, StreamingBoundsZ / 2.0));
+}
+
 void AWaterZone::SetRenderTargetResolution(FIntPoint NewResolution)
 {
 	RenderTargetResolution = NewResolution;
@@ -103,6 +114,15 @@ void AWaterZone::BeginPlay()
 	Super::BeginPlay();
 
 	MarkForRebuild(EWaterZoneRebuildFlags::All);
+
+	FWorldDelegates::LevelAddedToWorld.AddUObject(this, &AWaterZone::OnLevelAddedToWorld);
+	FWorldDelegates::LevelRemovedFromWorld.AddUObject(this, &AWaterZone::OnLevelRemovedFromWorld);
+}
+
+void AWaterZone::EndPlay(const EEndPlayReason::Type EndPlayReason)
+{
+	FWorldDelegates::LevelAddedToWorld.RemoveAll(this);
+	FWorldDelegates::LevelRemovedFromWorld.RemoveAll(this);
 }
 
 void AWaterZone::PostLoadSubobjects(FObjectInstancingGraph* OuterInstanceGraph)
@@ -321,11 +341,7 @@ TUniquePtr<class FWorldPartitionActorDesc> AWaterZone::CreateClassActorDesc() co
 
 FBox AWaterZone::GetStreamingBounds() const
 {
-	const FBox2D ZoneBounds2D = GetZoneBounds();
-	// #todo_water [roey]: Water zone doesn't have an explicit z bounds yet. For now just use the x or y.
-	const double StreamingBoundsZ = FMath::Max(ZoneExtent.X, ZoneExtent.Y);
-
-	return FBox(FVector3d(ZoneBounds2D.Min, -StreamingBoundsZ / 2.), FVector3d(ZoneBounds2D.Max, StreamingBoundsZ / 2.0));
+	return GetZoneBounds();
 }
 
 #endif // WITH_EDITOR
@@ -426,7 +442,7 @@ bool AWaterZone::UpdateWaterInfoTexture()
 		float GroundZMax = TNumericLimits<float>::Lowest();
 
 		TArray<TWeakObjectPtr<AActor>> GroundActors;
-		const FBox WaterZoneBounds(GetActorLocation() - FVector(GetZoneExtent() / 2., 0.), GetActorLocation() + FVector(GetZoneExtent() / 2., 0));
+		const FBox WaterZoneBounds = GetZoneBounds();
 		for (ALandscapeProxy* LandscapeProxy : TActorRange<ALandscapeProxy>(World))
 		{
 			const FBox LandscapeBox = LandscapeProxy->GetComponentsBoundingBox();
@@ -546,3 +562,58 @@ float AWaterZone::GetNonTessellatedLODSectionSize() const
 	return WaterMesh->GetTileSize() * NonTessellatedLODSectionScale;
 }
 
+void AWaterZone::OnLevelAddedToWorld(ULevel* InLevel, UWorld* InWorld)
+{
+	OnLevelChanged(InLevel, InWorld);
+}
+
+void AWaterZone::OnLevelRemovedFromWorld(ULevel* InLevel, UWorld* InWorld)
+{
+	OnLevelChanged(InLevel, InWorld);
+}
+
+void AWaterZone::OnLevelChanged(ULevel* InLevel, UWorld* InWorld)
+{
+	if ((InLevel == nullptr) || (InWorld != GetWorld()))
+	{
+		return;
+	}
+
+	const FBox WaterZoneBounds = GetZoneBounds();
+	const bool bContainsGroundActors = Algo::AnyOf(InLevel->Actors, [this, &WaterZoneBounds](const AActor* Actor)
+	{
+		return IsAffectingWaterZone(WaterZoneBounds, Actor);
+	});
+
+	if (bContainsGroundActors)
+	{
+		MarkForRebuild(EWaterZoneRebuildFlags::UpdateWaterInfoTexture);
+	}
+}
+
+bool AWaterZone::IsAffectingWaterZone(const FBox& InWaterZoneBounds, const AActor* InActor) const
+{
+	if (InActor == nullptr)
+	{
+		return false;
+	}
+
+	if (const ALandscapeProxy* LandscapeProxy = Cast<const ALandscapeProxy>(InActor))
+	{
+		bool bIntersectsWaterZone = false;
+
+		LandscapeProxy->ForEachComponent<ULandscapeComponent>(false, [&bIntersectsWaterZone, &InWaterZoneBounds](const ULandscapeComponent* LandscapeComponent)
+		{
+			if (!bIntersectsWaterZone && (LandscapeComponent != nullptr))
+			{
+				const FBox LandscapeBox = LandscapeComponent->Bounds.GetBox();
+
+				bIntersectsWaterZone = InWaterZoneBounds.IntersectXY(LandscapeBox) && (LandscapeBox.GetVolume() > 0.0);
+			}
+		});
+
+		return bIntersectsWaterZone;
+	}
+
+	return false;
+}
diff --git a/Engine/Plugins/Experimental/Water/Source/Runtime/Public/WaterZoneActor.h b/Engine/Plugins/Experimental/Water/Source/Runtime/Public/WaterZoneActor.h
--- a/Engine/Plugins/Experimental/Water/Source/Runtime/Public/WaterZoneActor.h
+++ b/Engine/Plugins/Experimental/Water/Source/Runtime/Public/WaterZoneActor.h
@@ -45,7 +45,8 @@ public:
 	FVector2D GetZoneExtent() const { return ZoneExtent; }
 	void SetZoneExtent(FVector2D NewExtents);
 
-	FBox2D GetZoneBounds() const;
+	FBox2D GetZoneBounds2D() const;
+	FBox GetZoneBounds() const;
 
 	void SetRenderTargetResolution(FIntPoint NewResolution);
 	FIntPoint GetRenderTargetResolution() const { return RenderTargetResolution; }
@@ -53,6 +54,7 @@ public:
 	uint32 GetVelocityBlurRadius() const { return VelocityBlurRadius; }
 
 	virtual void BeginPlay() override;
+	virtual void EndPlay(const EEndPlayReason::Type EndPlayReason) override;
 	virtual void PostLoadSubobjects(FObjectInstancingGraph* OuterInstanceGraph) override;
 	virtual void PostLoad() override;
 #if WITH_EDITORONLY_DATA
@@ -99,6 +101,19 @@ private:
 
 	void OnExtentChanged();
 
+	/** Delegates called when levels are added or removed from world. */
+	void OnLevelAddedToWorld(ULevel* InLevel, UWorld* InWorld);
+	void OnLevelRemovedFromWorld(ULevel* InLevel, UWorld* InWorld);
+	void OnLevelChanged(ULevel* InLevel, UWorld* InWorld);
+
+	/** Returns true if the provided actor can affect waterzone resources.
+	 *
+	 * @param	InWaterZoneBounds	The bounds of this waterzone.
+	 * @param	InActor		The Actor to check
+	 * @return true if the provided actor can affect waterzone resources.
+	 */
+	bool IsAffectingWaterZone(const FBox& InWaterZoneBounds, const AActor* InActor) const;
+
 #if WITH_EDITOR
 	void OnActorSelectionChanged(const TArray<UObject*>& NewSelection, bool bForceRefresh);
 	
