Water: Fixed issue where we can't retrieve water subsystem pointer or world partition pointer after actor is moved into the cell world during cook.

diff --git a/Engine/Plugins/Experimental/Water/Source/Runtime/Private/WaterBodyComponent.cpp b/Engine/Plugins/Experimental/Water/Source/Runtime/Private/WaterBodyComponent.cpp
--- a/Engine/Plugins/Experimental/Water/Source/Runtime/Private/WaterBodyComponent.cpp
+++ b/Engine/Plugins/Experimental/Water/Source/Runtime/Private/WaterBodyComponent.cpp
@@ -396,6 +396,12 @@ void UWaterBodyComponent::UpdateWaterZones()
 		}
 		else
 		{
+			// Don't attempt to find a water zone while cooking and just rely on the serialized pointer from the editor.
+			if (IsRunningCookCommandlet())
+			{
+				return;
+			}
+
 			const FBox Bounds3D = Bounds.GetBox();
 			if (const UWaterSubsystem* WaterSubsystem = UWaterSubsystem::GetWaterSubsystem(GetWorld()))
 			{
diff --git a/Engine/Plugins/Experimental/Water/Source/Runtime/Private/WaterSubsystem.cpp b/Engine/Plugins/Experimental/Water/Source/Runtime/Private/WaterSubsystem.cpp
--- a/Engine/Plugins/Experimental/Water/Source/Runtime/Private/WaterSubsystem.cpp
+++ b/Engine/Plugins/Experimental/Water/Source/Runtime/Private/WaterSubsystem.cpp
@@ -545,13 +545,13 @@ void UWaterSubsystem::MarkWaterZonesInRegionForRebuild(const FBox2D& InUpdateReg
 
 TSoftObjectPtr<AWaterZone> UWaterSubsystem::FindWaterZone(const FBox2D& Bounds) const
 {
-	if (!GetWorld())
+
+	const UWorld* World = GetWorld();
+	if (!World)
 	{
 		return {};
 	}
 
-	const UWorld* World = GetWorld();
-
 	// Score each overlapping water zone and then pick the best.
 	TMap<TSoftObjectPtr<AWaterZone>, int32> ViableZones;
 
@@ -559,7 +559,7 @@ TSoftObjectPtr<AWaterZone> UWaterSubsystem::FindWaterZone(const FBox2D& Bounds)
 	// Within the editor, we also want to check unloaded actors to ensure that the water body has serialized the best possible water zone, rather than just looking through what might be loaded now.
 	if (GEditor && !World->IsGameWorld())
 	{
-		if (UWorldPartition* WorldPartition = GetWorld()->GetWorldPartition())
+		if (UWorldPartition* WorldPartition = World->GetWorldPartition())
 		{
 			FWorldPartitionHelpers::ForEachActorDesc<AWaterZone>(WorldPartition, [Bounds, &ViableZones](const FWorldPartitionActorDesc* ActorDesc)
 			{
