Water: Water bodies were not registering themselves to their owning water zone when they get streamed in.

diff --git a/Engine/Plugins/Experimental/Water/Source/Runtime/Private/WaterBodyComponent.cpp b/Engine/Plugins/Experimental/Water/Source/Runtime/Private/WaterBodyComponent.cpp
--- a/Engine/Plugins/Experimental/Water/Source/Runtime/Private/WaterBodyComponent.cpp
+++ b/Engine/Plugins/Experimental/Water/Source/Runtime/Private/WaterBodyComponent.cpp
@@ -1369,6 +1369,7 @@ void UWaterBodyComponent::UpdateAll(const FOnWaterBodyChangedParams& InParams)
 	// #todo_water: this should be reconsidered along with the entire concept of non-dynamic water bodies now that evening is runtime-generated.
 	if (bShapeOrPositionChanged)
 	{
+		UpdateWaterZones();
 		UpdateWaterBodyRenderData();
 	}
 }
