Water: Avoid updating the water mesh and water info texture when water bodies which should not affect it are streamed in/out.

diff --git a/Engine/Plugins/Experimental/Water/Source/Runtime/Private/WaterZoneActor.cpp b/Engine/Plugins/Experimental/Water/Source/Runtime/Private/WaterZoneActor.cpp
--- a/Engine/Plugins/Experimental/Water/Source/Runtime/Private/WaterZoneActor.cpp
+++ b/Engine/Plugins/Experimental/Water/Source/Runtime/Private/WaterZoneActor.cpp
@@ -199,7 +199,18 @@ void AWaterZone::AddWaterBodyComponent(UWaterBodyComponent* WaterBodyComponent)
 	if (OwnedWaterBodies.Find(WaterBodyComponent) == INDEX_NONE)
 	{
 		OwnedWaterBodies.Add(WaterBodyComponent);
-		MarkForRebuild(EWaterZoneRebuildFlags::UpdateWaterInfoTexture | EWaterZoneRebuildFlags::UpdateWaterMesh);
+
+		EWaterZoneRebuildFlags RebuildFlags = EWaterZoneRebuildFlags::None;
+		if (WaterBodyComponent->AffectsWaterInfo())
+		{
+			RebuildFlags |= EWaterZoneRebuildFlags::UpdateWaterInfoTexture;
+		}
+		if (WaterBodyComponent->AffectsWaterMesh())
+		{
+			RebuildFlags |= EWaterZoneRebuildFlags::UpdateWaterMesh;
+		}
+
+		MarkForRebuild(RebuildFlags);
 	}
 }
 
@@ -209,7 +220,18 @@ void AWaterZone::RemoveWaterBodyComponent(UWaterBodyComponent* WaterBodyComponen
 	if (OwnedWaterBodies.Find(WaterBodyComponent, Index))
 	{
 		OwnedWaterBodies.RemoveAtSwap(Index);
-		MarkForRebuild(EWaterZoneRebuildFlags::UpdateWaterInfoTexture | EWaterZoneRebuildFlags::UpdateWaterMesh);
+
+		EWaterZoneRebuildFlags RebuildFlags = EWaterZoneRebuildFlags::None;
+		if (WaterBodyComponent->AffectsWaterInfo())
+		{
+			RebuildFlags |= EWaterZoneRebuildFlags::UpdateWaterInfoTexture;
+		}
+		if (WaterBodyComponent->AffectsWaterMesh())
+		{
+			RebuildFlags |= EWaterZoneRebuildFlags::UpdateWaterMesh;
+		}
+		
+		MarkForRebuild(RebuildFlags);
 	}
 }
 
