Fix thread-safety issue during async loading in water brush by moving code out of PostInitProperties to more appropriate places

#rnx

diff --git a/Engine/Plugins/Experimental/Water/Source/Editor/Private/WaterLandscapeBrush.cpp b/Engine/Plugins/Experimental/Water/Source/Editor/Private/WaterLandscapeBrush.cpp
--- a/Engine/Plugins/Experimental/Water/Source/Editor/Private/WaterLandscapeBrush.cpp
+++ b/Engine/Plugins/Experimental/Water/Source/Editor/Private/WaterLandscapeBrush.cpp
@@ -229,10 +229,8 @@ bool AWaterLandscapeBrush::IsActorAffectingLandscape(AActor* Actor) const
 	return ((WaterBrushActor != nullptr) && WaterBrushActor->AffectsLandscape());
 }
 
-void AWaterLandscapeBrush::PostInitProperties()
+void AWaterLandscapeBrush::RegisterDelegates()
 {
-	Super::PostInitProperties();
-
 	if (!IsTemplate())
 	{
 		OnWorldPostInitHandle = FWorldDelegates::OnPostWorldInitialization.AddLambda([this](UWorld* World, const UWorld::InitializationValues IVS)
@@ -270,12 +268,30 @@ void AWaterLandscapeBrush::PostInitProperties()
 		OnLevelActorDeletedHandle = GEngine->OnLevelActorDeleted().AddUObject(this, &AWaterLandscapeBrush::OnLevelActorRemoved);
 
 		IWaterBrushActorInterface::GetOnWaterBrushActorChangedEvent().AddUObject(this, &AWaterLandscapeBrush::OnWaterBrushActorChanged);
-
-		// If we are loading do not trigger events
-		UpdateActors(!GIsEditorLoadingPackage);
 	}
 }
 
+void AWaterLandscapeBrush::PostActorCreated()
+{
+	Super::PostActorCreated();
+
+	RegisterDelegates();
+
+	constexpr bool bInTriggerEvents = true;
+	UpdateActors(bInTriggerEvents);
+}
+
+void AWaterLandscapeBrush::PostLoad()
+{
+	Super::PostLoad();
+
+	RegisterDelegates();
+
+	// If we are loading do not trigger events
+	constexpr bool bInTriggerEvents = false;
+	UpdateActors(bInTriggerEvents);
+}
+
 void AWaterLandscapeBrush::OnLevelActorAdded(AActor* InActor)
 {
 	if (InActor->GetWorld() == GetWorld())
diff --git a/Engine/Plugins/Experimental/Water/Source/Editor/Private/WaterLandscapeBrush.h b/Engine/Plugins/Experimental/Water/Source/Editor/Private/WaterLandscapeBrush.h
--- a/Engine/Plugins/Experimental/Water/Source/Editor/Private/WaterLandscapeBrush.h
+++ b/Engine/Plugins/Experimental/Water/Source/Editor/Private/WaterLandscapeBrush.h
@@ -23,7 +23,6 @@ class AWaterLandscapeBrush : public ALandscapeBlueprintBrush
 public:
 	AWaterLandscapeBrush(const FObjectInitializer& ObjectInitializer = FObjectInitializer::Get());
 
-	virtual void PostInitProperties() override;
 	virtual void OnConstruction(const FTransform& Transform) override;
 	virtual void BeginDestroy() override;
 	virtual void PostRegisterAllComponents() override;
@@ -90,6 +89,9 @@ public:
 
 	virtual void GetRenderDependencies(TSet<UObject *>& OutDependencies) override;
 
+	virtual void PostLoad() override;
+	virtual void PostActorCreated() override;
+
 	void ForceUpdate();
 
 #if WITH_EDITOR
@@ -119,7 +121,7 @@ private:
 	void UpdateAffectedWeightmaps();
 	void ClearActors();
 	bool IsActorAffectingLandscape(AActor* Actor) const;
-
+	void RegisterDelegates();
 	void OnFullHeightmapRenderDone(UTextureRenderTarget2D* HeightmapRenderTarget);
 	void OnWaterBrushActorChanged(const IWaterBrushActorInterface::FWaterBrushActorChangedEventParams& InParams);
 	void OnActorsAffectingLandscapeChanged();
