Water: converted wave asset type actions to asset definitions.

diff --git a/Engine/Plugins/Experimental/Water/Source/Editor/Private/AssetDefinition_WaterWaves.cpp b/Engine/Plugins/Experimental/Water/Source/Editor/Private/AssetDefinition_WaterWaves.cpp
new file mode 100644
--- /dev/null
+++ b/Engine/Plugins/Experimental/Water/Source/Editor/Private/AssetDefinition_WaterWaves.cpp
@@ -0,0 +1,38 @@
+// Copyright Epic Games, Inc. All Rights Reserved.
+
+#include "AssetDefinition_WaterWaves.h"
+#include "Modules/ModuleManager.h"
+#include "WaterWaves.h"
+#include "WaterEditorModule.h"
+
+#define LOCTEXT_NAMESPACE "WaterWaves"
+
+FText UAssetDefinition_WaterWaves::GetAssetDisplayName() const
+{
+	return NSLOCTEXT("AssetTypeActions", "FAssetTypeActions_WaterWaves", "Water Waves");
+}
+
+TSoftClassPtr<UObject> UAssetDefinition_WaterWaves::GetAssetClass() const 
+{
+	return UWaterWavesAsset::StaticClass();
+}
+
+TConstArrayView<FAssetCategoryPath> UAssetDefinition_WaterWaves::GetAssetCategories() const
+{
+	static TArray<FAssetCategoryPath> AssetCategories = { FAssetCategoryPath(LOCTEXT("Water", "Water"))};
+	return AssetCategories;
+}
+
+EAssetCommandResult UAssetDefinition_WaterWaves::OpenAssets(const FAssetOpenArgs& OpenArgs) const
+{
+	FWaterEditorModule* WaterEditorModule = &FModuleManager::LoadModuleChecked<FWaterEditorModule>("WaterEditor");
+
+	for (UWaterWavesAsset* WavesAsset : OpenArgs.LoadObjects<UWaterWavesAsset>())
+	{
+		WaterEditorModule->CreateWaterWaveAssetEditor(OpenArgs.GetToolkitMode(), OpenArgs.ToolkitHost, WavesAsset);
+	}
+
+	return EAssetCommandResult::Handled;
+}
+
+#undef LOCTEXT_NAMESPACE
\ No newline at end of file
diff --git a/Engine/Plugins/Experimental/Water/Source/Editor/Private/AssetDefinition_WaterWaves.h b/Engine/Plugins/Experimental/Water/Source/Editor/Private/AssetDefinition_WaterWaves.h
new file mode 100644
--- /dev/null
+++ b/Engine/Plugins/Experimental/Water/Source/Editor/Private/AssetDefinition_WaterWaves.h
@@ -0,0 +1,22 @@
+// Copyright Epic Games, Inc. All Rights Reserved.
+
+#pragma once
+
+#include "AssetDefinitionDefault.h"
+
+#include "AssetDefinition_WaterWaves.generated.h"
+
+UCLASS()
+class UAssetDefinition_WaterWaves : public UAssetDefinitionDefault
+{
+	GENERATED_BODY()
+public:
+
+	// Begin IAssetTypeActions Interface
+	virtual FText GetAssetDisplayName() const override;
+	virtual FLinearColor GetAssetColor() const override { return FLinearColor(FColor(0, 85, 200)); }
+	virtual TSoftClassPtr<UObject> GetAssetClass() const override;
+	virtual TConstArrayView<FAssetCategoryPath> GetAssetCategories() const;
+	virtual EAssetCommandResult OpenAssets(const FAssetOpenArgs& OpenArgs) const override;
+	// End IAssetTypeActions Interface
+};
\ No newline at end of file
diff --git a/Engine/Plugins/Experimental/Water/Source/Editor/Private/AssetTypeActions_WaterWaves.h b/Engine/Plugins/Experimental/Water/Source/Editor/Private/AssetTypeActions_WaterWaves.h
deleted file mode 100644
diff --git a/Engine/Plugins/Experimental/Water/Source/Editor/Private/WaterEditorModule.cpp b/Engine/Plugins/Experimental/Water/Source/Editor/Private/WaterEditorModule.cpp
--- a/Engine/Plugins/Experimental/Water/Source/Editor/Private/WaterEditorModule.cpp
+++ b/Engine/Plugins/Experimental/Water/Source/Editor/Private/WaterEditorModule.cpp
@@ -6,6 +6,8 @@
 #include "PropertyEditorModule.h"
 #include "WaterLandscapeBrush.h"
 #include "EngineUtils.h"
+#include "ThumbnailRendering/ThumbnailManager.h"
+#include "Modules/ModuleManager.h"
 #include "Landscape.h"
 #include "Logging/MessageLog.h"
 #include "WaterZoneActor.h"
@@ -19,7 +21,7 @@
 #include "WaterZoneActorFactory.h"
 #include "WaterBodyActorDetailCustomization.h"
 #include "WaterBrushManagerFactory.h"
-#include "AssetTypeActions_WaterWaves.h"
+#include "AssetDefinition_WaterWaves.h"
 #include "WaterBrushCacheContainer.h"
 #include "WaterBodyBrushCacheContainerThumbnailRenderer.h"
 #include "WaterWavesEditorToolkit.h"
@@ -44,19 +46,6 @@ void FWaterEditorModule::StartupModule()
 	FPropertyEditorModule& PropertyModule = FModuleManager::LoadModuleChecked<FPropertyEditorModule>(TEXT("PropertyEditor"));
 	PropertyModule.RegisterCustomClassLayout(TEXT("WaterBody"), FOnGetDetailCustomizationInstance::CreateStatic(&FWaterBodyActorDetailCustomization::MakeInstance));
 
-	IAssetTools& AssetTools = FModuleManager::LoadModuleChecked<FAssetToolsModule>("AssetTools").Get();
-	WaterAssetCategory = AssetTools.RegisterAdvancedAssetCategory(FName(TEXT("Water")), LOCTEXT("WaterAssetCategory", "Water"));
-
-	// Helper lambda for registering asset type actions for automatic cleanup on shutdown
-	auto RegisterAssetTypeAction = [&](TSharedRef<IAssetTypeActions> Action)
-	{
-		AssetTools.RegisterAssetTypeActions(Action);
-		CreatedAssetTypeActions.Add(Action);
-	};
-
-	// Register type actions
-	RegisterAssetTypeAction(MakeShareable(new FAssetTypeActions_WaterWaves));
-
 	GEngine->OnLevelActorAdded().AddRaw(this, &FWaterEditorModule::OnLevelActorAddedToWorld);
 
 	RegisterComponentVisualizer(UWaterSplineComponent::StaticClass()->GetFName(), MakeShareable(new FWaterSplineComponentVisualizer));
@@ -109,16 +98,6 @@ void FWaterEditorModule::ShutdownModule()
 		}
 	}
 
-	if (FModuleManager::Get().IsModuleLoaded("AssetTools"))
-	{
-		IAssetTools& AssetTools = FModuleManager::GetModuleChecked<FAssetToolsModule>("AssetTools").Get();
-		for (auto CreatedAssetTypeAction : CreatedAssetTypeActions)
-		{
-			AssetTools.UnregisterAssetTypeActions(CreatedAssetTypeAction.ToSharedRef());
-		}
-	}
-	CreatedAssetTypeActions.Empty();
-
 	FEditorDelegates::OnMapOpened.RemoveAll(this);
 
 	FWaterUIStyle::Shutdown();
diff --git a/Engine/Plugins/Experimental/Water/Source/Editor/WaterEditor.Build.cs b/Engine/Plugins/Experimental/Water/Source/Editor/WaterEditor.Build.cs
--- a/Engine/Plugins/Experimental/Water/Source/Editor/WaterEditor.Build.cs
+++ b/Engine/Plugins/Experimental/Water/Source/Editor/WaterEditor.Build.cs
@@ -21,6 +21,7 @@ public class WaterEditor : ModuleRules
 				"SlateCore",
 				"Slate",
 				"UnrealEd",
+				"AssetDefinition",
 				"Water",
                 "Projects",
 				"PropertyEditor",
