Water: Speculative crash fix

diff --git a/Engine/Plugins/Experimental/Water/Source/Editor/Private/WaterEditorSettings.cpp b/Engine/Plugins/Experimental/Water/Source/Editor/Private/WaterEditorSettings.cpp
--- a/Engine/Plugins/Experimental/Water/Source/Editor/Private/WaterEditorSettings.cpp
+++ b/Engine/Plugins/Experimental/Water/Source/Editor/Private/WaterEditorSettings.cpp
@@ -151,16 +151,16 @@ UWaterEditorSettings::UWaterEditorSettings()
 {
 
 	// Ideally, this should be done in FWaterBodyLakeDefaults but we need CreateEditorOnlyDefaultSubobject, which is only available in a UObject : 
+	if (UWaterWavesAssetReference* WaterWavesRef = CreateEditorOnlyDefaultSubobject<UWaterWavesAssetReference>(TEXT("DefaultLakeWaterWaves"), /* bTransient = */false))
 	{
-		UWaterWavesAssetReference* WaterWavesRef = CreateEditorOnlyDefaultSubobject<UWaterWavesAssetReference>(TEXT("DefaultLakeWaterWaves"), /* bTransient = */false);
 		static ConstructorHelpers::FObjectFinder<UWaterWavesAsset> WaterWavesAssetRef(TEXT("/Water/Waves/GerstnerWaves_Lake.GerstnerWaves_Lake"));
 		WaterWavesRef->SetWaterWavesAsset(WaterWavesAssetRef.Object);
 		WaterBodyLakeDefaults.WaterWaves = WaterWavesRef;
 	}
 
 	// Ideally, this should be done in FWaterBodyLakeDefaults but we need CreateEditorOnlyDefaultSubobject, which is only available in a UObject : 
+	if (UWaterWavesAssetReference* WaterWavesRef = CreateEditorOnlyDefaultSubobject<UWaterWavesAssetReference>(TEXT("DefaultOceanWaterWaves"), /* bTransient = */false))
 	{
-		UWaterWavesAssetReference* WaterWavesRef = CreateEditorOnlyDefaultSubobject<UWaterWavesAssetReference>(TEXT("DefaultOceanWaterWaves"), /* bTransient = */false);
 		static ConstructorHelpers::FObjectFinder<UWaterWavesAsset> WaterWavesAssetRef(TEXT("/Water/Waves/GerstnerWaves_Ocean.GerstnerWaves_Ocean"));
 		WaterWavesRef->SetWaterWavesAsset(WaterWavesAssetRef.Object);
 		WaterBodyOceanDefaults.WaterWaves = WaterWavesRef;
