Water: Allow dilation data to write over water data in cases where the water info is far enough under the landscape to be clipped.

diff --git a/Engine/Plugins/Experimental/Water/Shaders/Private/WaterInfoMerge.usf b/Engine/Plugins/Experimental/Water/Shaders/Private/WaterInfoMerge.usf
--- a/Engine/Plugins/Experimental/Water/Shaders/Private/WaterInfoMerge.usf
+++ b/Engine/Plugins/Experimental/Water/Shaders/Private/WaterInfoMerge.usf
@@ -14,6 +14,8 @@ SamplerState DilationTextureSampler;
 float2 WaterHeightExtents;
 float GroundZMin;
 float CaptureZ;
+float UndergroundDilationDepthOffset;
+float DilationOverwriteMinimumDistance;
 
 struct FWaterColorPassResult
 {
@@ -59,11 +61,17 @@ void Main(
 	const FWaterColorPassResult WaterInfo = MakeWaterColorPassResult(ColorSample);
 	const float GroundDeviceDepth = DecodeFloatRGB(DepthSample.rgb);
 	const float GroundSceneDepth = ConvertFromDeviceZ(GroundDeviceDepth);
+	float GroundZ = CaptureZ - GroundSceneDepth;
 	
 	const float WaterSceneDepth = WaterInfo.SceneDepth;
 	const float WaterZ = CaptureZ - WaterSceneDepth;
 
-	float GroundZ = CaptureZ - GroundSceneDepth;
+	const float DilationDeviceDepth = DecodeFloatRGB(DilationSample.rgb);
+	const float DilationSceneDepth = ConvertFromDeviceZ(DilationDeviceDepth);
+	const bool bPixelHasDilationData = DilationSceneDepth != 0.0;
+		
+	float DilatedZ = CaptureZ - DilationSceneDepth;
+
 
 	// Always initialize output to avoid bugs in old shader compilers 
 	OutColor = 0.f;
@@ -79,7 +87,13 @@ void Main(
 	
 	// if the measured WaterZ is below WaterZMin it means that there is no water at this pixel because it is not possible
 	// for water to be further than ZMin.
-	const bool bPixelHasValidWaterData = WaterZ >= WaterZMin;
+	bool bPixelHasValidWaterData = WaterZ >= WaterZMin;
+
+	// Additionally, if we have dilation data here and the water data would be placed below the landscape, instead draw the dilation data
+	if (bPixelHasDilationData && (WaterZ + 128.f < GroundZ))
+	{
+		bPixelHasValidWaterData = false;
+	}
 
 	if (bPixelHasValidWaterData)
 	{
@@ -87,35 +101,26 @@ void Main(
 		const float NormalizedGroundZ = NormalizeHeight(GroundZ, GroundZMin, WaterZMax);
 
 		OutColor = float4(WaterInfo.Velocity, NormalizedWaterZ, NormalizedGroundZ);
-		return;
 	}
-	else
+	else if (bPixelHasDilationData)
 	{
-		const float DilationDeviceDepth = DecodeFloatRGB(DilationSample.rgb);
-		const float DilationSceneDepth = ConvertFromDeviceZ(DilationDeviceDepth);
-		const bool bPixelHasDilationData = DilationSceneDepth != 0.0;
-		
-		if (bPixelHasDilationData)
-		{
-			float DilatedZ = CaptureZ - DilationSceneDepth;
 
-			// If the dilated water would render on top of terrain, squash it down below
-			if (DilatedZ > GroundZ)
-			{
-				DilatedZ = GroundZ - 64.f;
-			}
+		// If the dilated water would render on top of terrain, squash it down below
+		if (DilatedZ > GroundZ)
+		{
+			DilatedZ = GroundZ - UndergroundDilationDepthOffset;
+		}
 
-			const float NormalizedWaterZ = NormalizeHeight(DilatedZ, WaterZMin, WaterZMax);
-			const float NormalizedGroundZ = NormalizeHeight(GroundZ, GroundZMin, WaterZMax);
+		const float NormalizedWaterZ = NormalizeHeight(DilatedZ, WaterZMin, WaterZMax);
+		const float NormalizedGroundZ = NormalizeHeight(GroundZ, GroundZMin, WaterZMax);
 
-			OutColor = float4(0.0, 0.0, NormalizedWaterZ, NormalizedGroundZ);
-		}
-		else
-		{
-			// Pixel may not have any dilation data nor water data. In this case just write 0s to the texture
+		OutColor = float4(0.0, 0.0, NormalizedWaterZ, NormalizedGroundZ);
+	}
+	else
+	{
+		// Pixel may not have any dilation data nor water data. In this case just write 0s to the texture
 
-			const float NormalizedGroundZ = NormalizeHeight(GroundZ, GroundZMin, WaterZMax);
-			OutColor = float4(0.0, 0.0, 0.0, NormalizedGroundZ);
-		}
+		const float NormalizedGroundZ = NormalizeHeight(GroundZ, GroundZMin, WaterZMax);
+		OutColor = float4(0.0, 0.0, 0.0, NormalizedGroundZ);
 	}
 }
diff --git a/Engine/Plugins/Experimental/Water/Source/Runtime/Private/WaterInfoRendering.cpp b/Engine/Plugins/Experimental/Water/Source/Runtime/Private/WaterInfoRendering.cpp
--- a/Engine/Plugins/Experimental/Water/Source/Runtime/Private/WaterInfoRendering.cpp
+++ b/Engine/Plugins/Experimental/Water/Source/Runtime/Private/WaterInfoRendering.cpp
@@ -34,6 +34,16 @@ static FAutoConsoleVariableRef CVarWaterInfoRenderLandscapeMinimumMipLevel(
 	WaterInfoRenderLandscapeMinimumMipLevel,
 	TEXT("Clamps the minimum allowed mip level for the landscape when rendering the water info texture. Used on the lowest end platforms which cannot support rendering all the landscape vertices at the highest LOD."));
 
+static TAutoConsoleVariable<float> CVarWaterInfoUndergroundDilationDepthOffset(
+		TEXT("r.Water.WaterInfo.UndergroundDilationDepthOffset"),
+		64.f,
+		TEXT("The minimum distance below the ground when we allow dilation to write on top of water"));
+
+static TAutoConsoleVariable<float> CVarWaterInfoDilationOverwriteMinimumDistance(
+		TEXT("r.Water.WaterInfo.DilationOverwriteMinimumDistance"),
+		128.f,
+		TEXT("The minimum distance below the ground when we allow dilation to write on top of water"));
+
 namespace UE::WaterInfo
 {
 
@@ -74,6 +84,8 @@ public:
 		SHADER_PARAMETER(FVector2f, WaterHeightExtents)
 		SHADER_PARAMETER(float, GroundZMin)
 		SHADER_PARAMETER(float, CaptureZ)
+		SHADER_PARAMETER(float, UndergroundDilationDepthOffset)
+		SHADER_PARAMETER(float, DilationOverwriteMinimumDistance)
 		RENDER_TARGET_BINDING_SLOTS()
 	END_SHADER_PARAMETER_STRUCT()
 
@@ -122,6 +134,8 @@ static void MergeWaterInfoAndDepth(
 		PassParameters->CaptureZ = Params.CaptureZ;
 		PassParameters->WaterHeightExtents = Params.WaterHeightExtents;
 		PassParameters->GroundZMin = Params.GroundZMin;
+		PassParameters->DilationOverwriteMinimumDistance = CVarWaterInfoDilationOverwriteMinimumDistance.GetValueOnRenderThread();
+		PassParameters->UndergroundDilationDepthOffset = CVarWaterInfoUndergroundDilationDepthOffset.GetValueOnRenderThread();
 
 		TShaderMapRef<FScreenVS> VertexShader(View.ShaderMap);
 		TShaderMapRef<FWaterInfoMergePS> PixelShader(View.ShaderMap);
