World Partition - HLOD: Name static meshes created by the WaterBody HLOD builder

diff --git a/Engine/Plugins/Experimental/Water/Source/Runtime/Private/WaterBodyHLODBuilder.cpp b/Engine/Plugins/Experimental/Water/Source/Runtime/Private/WaterBodyHLODBuilder.cpp
--- a/Engine/Plugins/Experimental/Water/Source/Runtime/Private/WaterBodyHLODBuilder.cpp
+++ b/Engine/Plugins/Experimental/Water/Source/Runtime/Private/WaterBodyHLODBuilder.cpp
@@ -77,7 +77,8 @@ TArray<UActorComponent*> UWaterBodyHLODBuilder::Build(const FHLODBuildContext& I
 
 	for (UWaterBodyComponent* WaterBodyComponent : SourceWaterBodyComponents)
 	{
-		UStaticMesh* StaticMesh = NewObject<UStaticMesh>(InHLODBuildContext.AssetsOuter);
+		FName StaticMeshName = MakeUniqueObjectName(nullptr, UStaticMesh::StaticClass(), TEXT("WaterBodyHLODMesh"));
+		UStaticMesh* StaticMesh = NewObject<UStaticMesh>(InHLODBuildContext.AssetsOuter, StaticMeshName);
 
 		// Mesh
 		{
