Water: removed the hard assert that water body components must be attached to water body actors. Instead, graceful fail without crashing.

diff --git a/Engine/Plugins/Experimental/Water/Source/Runtime/Private/WaterBodyComponent.cpp b/Engine/Plugins/Experimental/Water/Source/Runtime/Private/WaterBodyComponent.cpp
--- a/Engine/Plugins/Experimental/Water/Source/Runtime/Private/WaterBodyComponent.cpp
+++ b/Engine/Plugins/Experimental/Water/Source/Runtime/Private/WaterBodyComponent.cpp
@@ -205,7 +205,7 @@ FBox UWaterBodyComponent::GetCollisionComponentBounds() const
 AWaterBody* UWaterBodyComponent::GetWaterBodyActor() const
 {
 	// If we have an Owner, it must be an AWaterBody
-	return GetOwner() ? CastChecked<AWaterBody>(GetOwner()) : nullptr;
+	return GetOwner() ? Cast<AWaterBody>(GetOwner()) : nullptr;
 }
 
 UWaterSplineComponent* UWaterBodyComponent::GetWaterSpline() const
@@ -763,6 +763,11 @@ void UWaterBodyComponent::OnRegister()
 	Super::OnRegister();
 
 	AWaterBody* OwningWaterBodyActor = GetWaterBodyActor();
+	if (OwningWaterBodyActor == nullptr)
+	{
+		return;
+	}
+
 	WaterSplineMetadata = OwningWaterBodyActor->GetWaterSplineMetadata();
 
 	check(WaterSplineMetadata);
@@ -1072,22 +1077,33 @@ TArray<TSharedRef<FTokenizedMessage>> UWaterBodyComponent::CheckWaterBodyStatus(
 	const UWorld* World = GetWorld();
 	if (!IsTemplate() && World && World->WorldType != EWorldType::EditorPreview)
 	{
-		if (AffectsWaterMesh() && (GetWaterZone() == nullptr))
+		if (GetWaterBodyActor() == nullptr)
 		{
 			Result.Add(FTokenizedMessage::Create(EMessageSeverity::Error)
 				->AddToken(FUObjectToken::Create(this))
 				->AddToken(FTextToken::Create(FText::Format(
-					LOCTEXT("MapCheck_Message_MissingWaterZone", "Water body {0} requires a WaterZone actor to be rendered. Please add one to the map. "),
-					FText::FromString(GetWaterBodyActor()->GetActorLabel())))));
+					LOCTEXT("MapCheck_Message_MissingWaterBodyActor", "WaterBodyComponent is attached to an actor which is not an AWaterBody ({0})! WaterBodyComponents required parent water body actors to function correctly"),
+					FText::FromString(GetOwner() ? GetOwner()->GetActorLabel() : TEXT(""))))));
 		}
-
-		if (AffectsLandscape() && (FindLandscape() == nullptr))
+		else
 		{
-			Result.Add(FTokenizedMessage::Create(EMessageSeverity::Error)
-				->AddToken(FUObjectToken::Create(this))
-				->AddToken(FTextToken::Create(FText::Format(
-					LOCTEXT("MapCheck_Message_MissingLandscape", "Water body {0} requires a Landscape to be rendered. Please add one to the map. "),
-					FText::FromString(GetWaterBodyActor()->GetActorLabel())))));
+			if (AffectsWaterMesh() && (GetWaterZone() == nullptr))
+			{
+				Result.Add(FTokenizedMessage::Create(EMessageSeverity::Error)
+					->AddToken(FUObjectToken::Create(this))
+					->AddToken(FTextToken::Create(FText::Format(
+						LOCTEXT("MapCheck_Message_MissingWaterZone", "Water body {0} requires a WaterZone actor to be rendered. Please add one to the map. "),
+						FText::FromString(GetWaterBodyActor()->GetActorLabel())))));
+			}
+
+			if (AffectsLandscape() && (FindLandscape() == nullptr))
+			{
+				Result.Add(FTokenizedMessage::Create(EMessageSeverity::Error)
+					->AddToken(FUObjectToken::Create(this))
+					->AddToken(FTextToken::Create(FText::Format(
+						LOCTEXT("MapCheck_Message_MissingLandscape", "Water body {0} requires a Landscape to be rendered. Please add one to the map. "),
+						FText::FromString(GetWaterBodyActor()->GetActorLabel())))));
+			}
 		}
 	}
 	return Result;
@@ -1302,7 +1318,10 @@ void UWaterBodyComponent::UpdateAll(const FOnWaterBodyChangedParams& InParams)
 	BeginUpdateWaterBody();
 
 	AWaterBody* WaterBodyOwner = GetWaterBodyActor();
-	check(WaterBodyOwner);
+	if (WaterBodyOwner == nullptr)
+	{
+		return;
+	}
 
 	bool bShapeOrPositionChanged = InParams.bShapeOrPositionChanged;
 	
@@ -1399,6 +1418,11 @@ void UWaterBodyComponent::OnWaterBodyChanged(const FOnWaterBodyChangedParams& In
 
 #if WITH_EDITOR
 	AWaterBody* const WaterBodyActor = GetWaterBodyActor();
+	if (WaterBodyActor == nullptr)
+	{
+		return;
+	}
+
 	// Transfer the FOnWaterBodyChangedParams parameters to FWaterBrushActorChangedEventParams :
 	IWaterBrushActorInterface::FWaterBrushActorChangedEventParams Params(WaterBodyActor, InParams.PropertyChangedEvent);
 	Params.bUserTriggered = InParams.bUserTriggered;
diff --git a/Engine/Plugins/Experimental/Water/Source/Runtime/Private/WaterBodyCustomComponent.cpp b/Engine/Plugins/Experimental/Water/Source/Runtime/Private/WaterBodyCustomComponent.cpp
--- a/Engine/Plugins/Experimental/Water/Source/Runtime/Private/WaterBodyCustomComponent.cpp
+++ b/Engine/Plugins/Experimental/Water/Source/Runtime/Private/WaterBodyCustomComponent.cpp
@@ -140,7 +140,7 @@ TArray<TSharedRef<FTokenizedMessage>> UWaterBodyCustomComponent::CheckWaterBodyS
 
 	if (!IsTemplate())
 	{
-		if (WaterMeshOverride == nullptr)
+		if (WaterMeshOverride == nullptr && GetWaterBodyActor() != nullptr)
 		{
 			StatusMessages.Add(FTokenizedMessage::Create(EMessageSeverity::Error)
 				->AddToken(FUObjectToken::Create(this))
