Water: default water zone render target resolution for new water zones changed to 1024x1024

diff --git a/Engine/Plugins/Experimental/Water/Source/Editor/Private/WaterZoneActorFactory.cpp b/Engine/Plugins/Experimental/Water/Source/Editor/Private/WaterZoneActorFactory.cpp
--- a/Engine/Plugins/Experimental/Water/Source/Editor/Private/WaterZoneActorFactory.cpp
+++ b/Engine/Plugins/Experimental/Water/Source/Editor/Private/WaterZoneActorFactory.cpp
@@ -23,9 +23,11 @@ void UWaterZoneActorFactory::PostSpawnActor(UObject* Asset, AActor* NewActor)
 	
 	AWaterZone* WaterZoneActor = CastChecked<AWaterZone>(NewActor);
 
-	const FWaterZoneActorDefaults& WaterMeshActorDefaults = GetDefault<UWaterEditorSettings>()->WaterZoneActorDefaults;
-	WaterZoneActor->GetWaterMeshComponent()->FarDistanceMaterial = WaterMeshActorDefaults.GetFarDistanceMaterial();
-	WaterZoneActor->GetWaterMeshComponent()->FarDistanceMeshExtent = WaterMeshActorDefaults.FarDistanceMeshExtent;
+	const FWaterZoneActorDefaults& WaterZoneActorDefaults = GetDefault<UWaterEditorSettings>()->WaterZoneActorDefaults;
+	WaterZoneActor->GetWaterMeshComponent()->FarDistanceMaterial = WaterZoneActorDefaults.GetFarDistanceMaterial();
+	WaterZoneActor->GetWaterMeshComponent()->FarDistanceMeshExtent = WaterZoneActorDefaults.FarDistanceMeshExtent;
+
+	WaterZoneActor->SetRenderTargetResolution(WaterZoneActorDefaults.RenderTargetResolution);
 }
 
 #undef LOCTEXT_NAMESPACE
diff --git a/Engine/Plugins/Experimental/Water/Source/Editor/Public/WaterEditorSettings.h b/Engine/Plugins/Experimental/Water/Source/Editor/Public/WaterEditorSettings.h
--- a/Engine/Plugins/Experimental/Water/Source/Editor/Public/WaterEditorSettings.h
+++ b/Engine/Plugins/Experimental/Water/Source/Editor/Public/WaterEditorSettings.h
@@ -51,6 +51,9 @@ struct WATEREDITOR_API FWaterZoneActorDefaults
 	UPROPERTY(EditAnywhere, config, Category = Mesh)
 	float NewWaterZoneScale = 1.0f;
 
+	UPROPERTY(EditAnywhere, config, Category = Texture)
+	FIntPoint RenderTargetResolution = FIntPoint(1024, 1024);
+
 protected:
 	UPROPERTY(EditAnywhere, Config, Category = Mesh)
 	TSoftObjectPtr<UMaterialInterface> FarDistanceMaterial;
