Changing WaterInfoRendering to use more Renderer APIs instead of private data.

diff --git a/Engine/Plugins/Experimental/Water/Source/Runtime/Private/WaterInfoRendering.cpp b/Engine/Plugins/Experimental/Water/Source/Runtime/Private/WaterInfoRendering.cpp
--- a/Engine/Plugins/Experimental/Water/Source/Runtime/Private/WaterInfoRendering.cpp
+++ b/Engine/Plugins/Experimental/Water/Source/Runtime/Private/WaterInfoRendering.cpp
@@ -11,7 +11,7 @@
 #include "WaterBodySceneProxy.h"
 #include "Engine/TextureRenderTarget2D.h"
 #include "SceneCaptureRendering.h"
-#include "PostProcess/SceneFilterRendering.h"
+#include "PostProcess/DrawRectangle.h"
 #include "Math/OrthoMatrix.h"
 #include "GameFramework/WorldSettings.h"
 #include "ScreenRendering.h"
@@ -22,6 +22,9 @@
 #include "WaterBodyComponent.h"
 #include "Containers/StridedView.h"
 
+#include "RenderGraphBuilder.h"
+#include "SceneRenderTargetParameters.h"
+
 #include "ScenePrivate.h"
 #include "SceneRendering.h"
 
@@ -52,6 +55,7 @@ namespace UE::WaterInfo
 
 struct FUpdateWaterInfoParams
 {
+	FSceneInterface* Scene = nullptr;
 	FSceneRenderer* DepthRenderer = nullptr;
 	FSceneRenderer* ColorRenderer = nullptr;
 	FSceneRenderer* DilationRenderer = nullptr;
@@ -107,9 +111,8 @@ IMPLEMENT_GLOBAL_SHADER(FWaterInfoMergePS, "/Plugin/Water/Private/WaterInfoMerge
 
 static void MergeWaterInfoAndDepth(
 	FRDGBuilder& GraphBuilder,
-	const FMinimalSceneTextures& SceneTextures,
 	const FSceneViewFamily& ViewFamily,
-	FViewInfo& View,
+	const FSceneView& View,
 	FRDGTextureRef OutputTexture,
 	FRDGTextureRef DepthTexture,
 	FRDGTextureRef ColorTexture,
@@ -118,6 +121,8 @@ static void MergeWaterInfoAndDepth(
 {
 	RDG_EVENT_SCOPE(GraphBuilder, "WaterInfoDepthMerge");
 
+	FGlobalShaderMap* ShaderMap = GetGlobalShaderMap(View.GetFeatureLevel());
+
 	FGraphicsPipelineStateInitializer GraphicsPSOInit;
 	GraphicsPSOInit.RasterizerState = TStaticRasterizerState<FM_Solid, CM_None>::GetRHI();
 	GraphicsPSOInit.DepthStencilState = TStaticDepthStencilState<false, CF_Always>::GetRHI();
@@ -127,7 +132,7 @@ static void MergeWaterInfoAndDepth(
 		FWaterInfoMergePS::FParameters* PassParameters = GraphBuilder.AllocParameters<FWaterInfoMergePS::FParameters>();
 		PassParameters->View = View.ViewUniformBuffer;
 		PassParameters->RenderTargets[0] = FRenderTargetBinding(OutputTexture, ERenderTargetLoadAction::ENoAction);
-		PassParameters->SceneTextures = SceneTextures.GetSceneTextureShaderParameters(ViewFamily.GetFeatureLevel());
+		PassParameters->SceneTextures = GetSceneTextureShaderParameters(View);
 		PassParameters->DepthTexture = DepthTexture;
 		PassParameters->DepthTextureSampler = TStaticSamplerState<SF_Point, AM_Clamp, AM_Clamp, AM_Clamp>::GetRHI();
 		PassParameters->ColorTexture = ColorTexture;
@@ -140,8 +145,8 @@ static void MergeWaterInfoAndDepth(
 		PassParameters->DilationOverwriteMinimumDistance = CVarWaterInfoDilationOverwriteMinimumDistance.GetValueOnRenderThread();
 		PassParameters->UndergroundDilationDepthOffset = CVarWaterInfoUndergroundDilationDepthOffset.GetValueOnRenderThread();
 
-		TShaderMapRef<FScreenVS> VertexShader(View.ShaderMap);
-		TShaderMapRef<FWaterInfoMergePS> PixelShader(View.ShaderMap);
+		TShaderMapRef<FScreenVS> VertexShader(ShaderMap);
+		TShaderMapRef<FWaterInfoMergePS> PixelShader(ShaderMap);
 
 		GraphicsPSOInit.BoundShaderState.VertexDeclarationRHI = GFilterVertexDeclaration.VertexDeclarationRHI;
 		GraphicsPSOInit.BoundShaderState.VertexShaderRHI = VertexShader.GetVertexShader();
@@ -159,16 +164,7 @@ static void MergeWaterInfoAndDepth(
 				SetGraphicsPipelineState(RHICmdList, LocalGraphicsPSOInit, 0);
 				SetShaderParameters(RHICmdList, PixelShader, PixelShader.GetPixelShader(), *PassParameters);
 
-				DrawRectangle(
-					RHICmdList,
-					View.ViewRect.Min.X, View.ViewRect.Min.Y,
-					View.ViewRect.Width(), View.ViewRect.Height(),
-					View.ViewRect.Min.X, View.ViewRect.Min.Y,
-					View.ViewRect.Width(), View.ViewRect.Height(),
-					View.UnconstrainedViewRect.Size(),
-					View.UnconstrainedViewRect.Size(),
-					VertexShader,
-					EDRF_UseTriangleOptimization);
+				UE::Renderer::PostProcess::DrawRectangle(RHICmdList, VertexShader, View, EDRF_UseTriangleOptimization);
 			});
 	}
 }
@@ -226,15 +222,16 @@ IMPLEMENT_GLOBAL_SHADER(FWaterInfoFinalizePS, "/Plugin/Water/Private/WaterInfoFi
 
 static void FinalizeWaterInfo(
 	FRDGBuilder& GraphBuilder,
-	const FMinimalSceneTextures& SceneTextures,
 	const FSceneViewFamily& ViewFamily,
-	FViewInfo& View,
+	const FSceneView& View,
 	FRDGTextureRef WaterInfoTexture,
 	FRDGTextureRef OutputTexture,
 	const FUpdateWaterInfoParams& Params)
 {
 	RDG_EVENT_SCOPE(GraphBuilder, "WaterInfoFinalize");
 
+	FGlobalShaderMap* ShaderMap = GetGlobalShaderMap(View.GetFeatureLevel());
+
 	FGraphicsPipelineStateInitializer GraphicsPSOInit;
 	GraphicsPSOInit.RasterizerState = TStaticRasterizerState<FM_Solid, CM_None>::GetRHI();
 	GraphicsPSOInit.DepthStencilState = TStaticDepthStencilState<false, CF_Always>::GetRHI();
@@ -247,15 +244,15 @@ static void FinalizeWaterInfo(
 		FWaterInfoFinalizePS::FParameters* PassParameters = GraphBuilder.AllocParameters<FWaterInfoFinalizePS::FParameters>();
 		PassParameters->View = View.ViewUniformBuffer;
 		PassParameters->RenderTargets[0] = FRenderTargetBinding(OutputTexture, ERenderTargetLoadAction::ENoAction);
-		PassParameters->SceneTextures = SceneTextures.GetSceneTextureShaderParameters(ViewFamily.GetFeatureLevel());
+		PassParameters->SceneTextures = GetSceneTextureShaderParameters(View);
 		PassParameters->WaterInfoTexture = WaterInfoTexture;
 		PassParameters->BlurRadius = Params.VelocityBlurRadius;
 		PassParameters->CaptureZ = Params.CaptureZ;
 		PassParameters->WaterHeightExtents = Params.WaterHeightExtents;
 		PassParameters->GroundZMin = Params.GroundZMin;
 
-		TShaderMapRef<FScreenVS> VertexShader(View.ShaderMap);
-		TShaderMapRef<FWaterInfoFinalizePS> PixelShader(View.ShaderMap, PixelPermutationVector);
+		TShaderMapRef<FScreenVS> VertexShader(ShaderMap);
+		TShaderMapRef<FWaterInfoFinalizePS> PixelShader(ShaderMap, PixelPermutationVector);
 
 		GraphicsPSOInit.BoundShaderState.VertexDeclarationRHI = GFilterVertexDeclaration.VertexDeclarationRHI;
 		GraphicsPSOInit.BoundShaderState.VertexShaderRHI = VertexShader.GetVertexShader();
@@ -272,17 +269,8 @@ static void FinalizeWaterInfo(
 				RHICmdList.ApplyCachedRenderTargets(LocalGraphicsPSOInit);
 				SetGraphicsPipelineState(RHICmdList, LocalGraphicsPSOInit, 0);
 				SetShaderParameters(RHICmdList, PixelShader, PixelShader.GetPixelShader(), *PassParameters);
-				
-				DrawRectangle(
-					RHICmdList,
-					View.ViewRect.Min.X, View.ViewRect.Min.Y,
-					View.ViewRect.Width(), View.ViewRect.Height(),
-					View.ViewRect.Min.X, View.ViewRect.Min.Y,
-					View.ViewRect.Width(), View.ViewRect.Height(),
-					View.UnconstrainedViewRect.Size(),
-					View.UnconstrainedViewRect.Size(),
-					VertexShader,
-					EDRF_UseTriangleOptimization);
+
+				UE::Renderer::PostProcess::DrawRectangle(RHICmdList, VertexShader, View, EDRF_UseTriangleOptimization);
 			});
 	}
 }
@@ -311,14 +299,14 @@ static FMatrix BuildOrthoMatrix(float InOrthoWidth, float InOrthoHeight)
 
 // ---------------------------------------------------------------------------------------------------------------------
 
-static void SetWaterBodiesWithinWaterInfoPass(FSceneRenderer* SceneRenderer, EWaterInfoPass InPass)
+static void SetWaterBodiesWithinWaterInfoPass(FSceneInterface* Scene, const FSceneView& View, EWaterInfoPass InPass)
 {
 	TRACE_CPUPROFILER_EVENT_SCOPE(WaterInfo::SetWaterBodiesWithinWaterInfoPass);
-	if (SceneRenderer->Views[0].ShowOnlyPrimitives.IsSet())
+	if (View.ShowOnlyPrimitives.IsSet())
 	{
-		for (FPrimitiveSceneProxy* PrimProxy : SceneRenderer->Scene->PrimitiveSceneProxies)
+		for (FPrimitiveSceneProxy* PrimProxy : Scene->GetPrimitiveSceneProxies())
 		{
-			if (PrimProxy && SceneRenderer->Views[0].ShowOnlyPrimitives->Contains(PrimProxy->GetPrimitiveComponentId()))
+			if (PrimProxy && View.ShowOnlyPrimitives->Contains(PrimProxy->GetPrimitiveComponentId()))
 			{
 				FWaterBodySceneProxy* WaterProxy = (FWaterBodySceneProxy*)PrimProxy;
 				WaterProxy->SetWithinWaterInfoPass(InPass);
@@ -437,6 +425,8 @@ static void UpdateWaterInfoRendering_RenderThread(
 	TRefCountPtr<IPooledRenderTarget> ExtractedDepthTexture;
 	TRefCountPtr<IPooledRenderTarget> ExtractedColorTexture;
 
+	FSceneInterface* Scene = Params.Scene;
+
 	// Depth-only pass for actors which are considered the ground for water rendering
 	{
 		FSceneRenderer* DepthRenderer = Params.DepthRenderer;
@@ -467,7 +457,7 @@ static void UpdateWaterInfoRendering_RenderThread(
 			DepthRenderer->Render(GraphBuilder);
 		}
 		
-		if (DepthRenderer->Scene->GetShadingPath() == EShadingPath::Mobile)
+		if (Scene->GetShadingPath() == EShadingPath::Mobile)
 		{
 			RDG_EVENT_SCOPE(GraphBuilder, "CaptureSceneColor");
 			CopySceneCaptureComponentToTarget(
@@ -497,6 +487,8 @@ static void UpdateWaterInfoRendering_RenderThread(
 
 		ColorRenderer->RenderThreadBegin(RHICmdList);
 
+		FViewInfo& View = ColorRenderer->Views[0];
+
 		SCOPED_DRAW_EVENT(RHICmdList, ColorRendering_RT);
 
 		FRDGBuilder GraphBuilder(RHICmdList, RDG_EVENT_NAME("WaterInfoColorRendering"), ERDGBuilderFlags::AllowParallelExecute);
@@ -504,15 +496,13 @@ static void UpdateWaterInfoRendering_RenderThread(
 		// We need to execute the pre-render view extensions before we do any view dependent work.
 		FSceneRenderer::ViewExtensionPreRender_RenderThread(GraphBuilder, ColorRenderer);
 
-		SetWaterBodiesWithinWaterInfoPass(ColorRenderer, EWaterInfoPass::Color);
+		SetWaterBodiesWithinWaterInfoPass(Scene, View, EWaterInfoPass::Color);
 		
 		FRDGTextureRef TargetTexture = RegisterExternalTexture(GraphBuilder, RenderTarget->GetRenderTargetTexture(), TEXT("WaterColorTarget"));
 
 		FRDGTextureDesc ColorTextureDesc(TargetTexture->Desc);
 		ColorTextureDesc.Format = PF_A32B32G32R32F;
 		FRDGTextureRef ColorTexture = GraphBuilder.CreateTexture(ColorTextureDesc, TEXT("WaterColorTexture"));
-		
-		FViewInfo& View = ColorRenderer->Views[0];
 
 		AddClearRenderTargetPass(GraphBuilder, TargetTexture, FLinearColor::Black, View.UnscaledViewRect);
 
@@ -524,8 +514,6 @@ static void UpdateWaterInfoRendering_RenderThread(
 			ColorRenderer->Render(GraphBuilder);
 		}
 
-		const FMinimalSceneTextures& SceneTextures = View.GetSceneTextures();
-
 		// This CopySceneCaptureComponentToTarget is required on all platforms as it extracts a higher pixel depth texture than the scene render target so we can't just copy out of it
 		{
 			RDG_EVENT_SCOPE(GraphBuilder, "CaptureSceneColor");
@@ -551,6 +539,8 @@ static void UpdateWaterInfoRendering_RenderThread(
 
 		DilationRenderer->RenderThreadBegin(RHICmdList);
 
+		FViewInfo& View = DilationRenderer->Views[0];
+
 		SCOPED_DRAW_EVENT(RHICmdList, DilationRendering_RT);
 
 		FRDGBuilder GraphBuilder(RHICmdList, RDG_EVENT_NAME("WaterInfoDilationRendering"), ERDGBuilderFlags::AllowParallelExecute);
@@ -559,15 +549,13 @@ static void UpdateWaterInfoRendering_RenderThread(
 		// We need to execute the pre-render view extensions before we do any view dependent work.
 		FSceneRenderer::ViewExtensionPreRender_RenderThread(GraphBuilder, DilationRenderer);
 
-		SetWaterBodiesWithinWaterInfoPass(DilationRenderer, EWaterInfoPass::Dilation);
+		SetWaterBodiesWithinWaterInfoPass(Scene, View, EWaterInfoPass::Dilation);
 
 		FRDGTextureRef TargetTexture = RegisterExternalTexture(GraphBuilder, RenderTarget->GetRenderTargetTexture(), TEXT("WaterDilationTarget"));
 
 		FRDGTextureRef DepthTexture = GraphBuilder.RegisterExternalTexture(ExtractedDepthTexture);
 		FRDGTextureRef ColorTexture = GraphBuilder.RegisterExternalTexture(ExtractedColorTexture);
 
-		FViewInfo& View = DilationRenderer->Views[0];
-
 		AddClearRenderTargetPass(GraphBuilder, TargetTexture, FLinearColor::Black, View.UnscaledViewRect);
 
 		View.bDisableQuerySubmissions = true;
@@ -578,11 +566,10 @@ static void UpdateWaterInfoRendering_RenderThread(
 			DilationRenderer->Render(GraphBuilder);
 		}
 
-		const FMinimalSceneTextures& SceneTextures = View.GetSceneTextures();
 		FRDGTextureDesc TextureDesc(TargetTexture->Desc);
 		FRDGTextureRef DilationTexture = GraphBuilder.CreateTexture(TextureDesc, TEXT("WaterDilationTexture"));
 
-		if (DilationRenderer->Scene->GetShadingPath() == EShadingPath::Mobile)
+		if (Scene->GetShadingPath() == EShadingPath::Mobile)
 		{
 			RDG_EVENT_SCOPE(GraphBuilder, "CaptureSceneColor");
 			CopySceneCaptureComponentToTarget(
@@ -598,16 +585,16 @@ static void UpdateWaterInfoRendering_RenderThread(
 		}
 
 		FRDGTextureRef MergeTargetTexture = GraphBuilder.CreateTexture(TextureDesc, TEXT("WaterInfoMerged"));
-		MergeWaterInfoAndDepth(GraphBuilder, SceneTextures, DilationRenderer->ViewFamily, DilationRenderer->Views[0], MergeTargetTexture, DepthTexture, ColorTexture, DilationTexture, Params);
+		MergeWaterInfoAndDepth(GraphBuilder, DilationRenderer->ViewFamily, DilationRenderer->Views[0], MergeTargetTexture, DepthTexture, ColorTexture, DilationTexture, Params);
 
 		FRDGTextureRef FinalizedTexture = GraphBuilder.CreateTexture(TargetTexture->Desc, TEXT("WaterInfoFinalized"));
-		FinalizeWaterInfo(GraphBuilder, SceneTextures, DilationRenderer->ViewFamily, DilationRenderer->Views[0], MergeTargetTexture, FinalizedTexture, Params);
+		FinalizeWaterInfo(GraphBuilder, DilationRenderer->ViewFamily, DilationRenderer->Views[0], MergeTargetTexture, FinalizedTexture, Params);
 		
 		FRDGTextureRef ShaderResourceTexture = RegisterExternalTexture(GraphBuilder, OutputTexture->TextureRHI, TEXT("WaterInfoResolve"));
 		AddCopyTexturePass(GraphBuilder, FinalizedTexture, ShaderResourceTexture);
 		GraphBuilder.Execute();
 
-		SetWaterBodiesWithinWaterInfoPass(DilationRenderer, EWaterInfoPass::None);
+		SetWaterBodiesWithinWaterInfoPass(Scene, View, EWaterInfoPass::None);
 		DilationRenderer->RenderThreadEnd(RHICmdList);
 	}
 
@@ -781,6 +768,7 @@ void UpdateWaterInfoRendering(
 	FTextureRenderTargetResource* TextureRenderTargetResource = Context.TextureRenderTarget->GameThread_GetRenderTargetResource();
 
 	FUpdateWaterInfoParams Params;
+	Params.Scene = Scene;
 	Params.DepthRenderer = DepthRenderer;
 	Params.ColorRenderer = ColorRenderer;
 	Params.DilationRenderer = DilationRenderer;
diff --git a/Engine/Source/Runtime/Renderer/Private/PostProcess/SceneFilterRendering.cpp b/Engine/Source/Runtime/Renderer/Private/PostProcess/SceneFilterRendering.cpp
--- a/Engine/Source/Runtime/Renderer/Private/PostProcess/SceneFilterRendering.cpp
+++ b/Engine/Source/Runtime/Renderer/Private/PostProcess/SceneFilterRendering.cpp
@@ -343,6 +343,30 @@ namespace UE::Renderer::PostProcess
 		::DrawRectangle(RHICmdList, X, Y, SizeX, SizeY, U, V, SizeU, SizeV, TargetSize, TextureSize, VertexShader, Flags, InstanceCount);
 	}
 
+	void DrawRectangle(
+		FRHICommandList& RHICmdList,
+		const TShaderRef<FShader>& VertexShader,
+		const FSceneView& InView,
+		EDrawRectangleFlags Flags,
+		uint32 InstanceCount
+	)
+	{
+		if (ensure(InView.bIsViewInfo))
+		{
+			const FViewInfo& View = static_cast<const FViewInfo&>(InView);
+			DrawRectangle(
+				RHICmdList,
+				VertexShader,
+				View.ViewRect.Min.X, View.ViewRect.Min.Y,
+				View.ViewRect.Width(), View.ViewRect.Height(),
+				View.ViewRect.Min.X, View.ViewRect.Min.Y,
+				View.ViewRect.Width(), View.ViewRect.Height(),
+				View.UnconstrainedViewRect.Size(),
+				View.UnconstrainedViewRect.Size(),
+				Flags);
+		}
+	}
+
 	void DrawPostProcessPass(
 		FRHICommandList& RHICmdList,
 		const TShaderRef<FShader>& VertexShader,
diff --git a/Engine/Source/Runtime/Renderer/Private/SceneTextures.cpp b/Engine/Source/Runtime/Renderer/Private/SceneTextures.cpp
--- a/Engine/Source/Runtime/Renderer/Private/SceneTextures.cpp
+++ b/Engine/Source/Runtime/Renderer/Private/SceneTextures.cpp
@@ -1044,6 +1044,13 @@ FSceneTextureShaderParameters CreateSceneTextureShaderParameters(FRDGBuilder& Gr
 	return Parameters;
 }
 
+FSceneTextureShaderParameters GetSceneTextureShaderParameters(const FSceneView& View)
+{
+	check(View.bIsViewInfo);
+	const FMinimalSceneTextures& SceneTextures = static_cast<const FViewInfo&>(View).GetSceneTextures();
+	return SceneTextures.GetSceneTextureShaderParameters(View.GetFeatureLevel());
+}
+
 TRDGUniformBufferRef<FSceneTextureUniformParameters> GetSceneTextureUnformBuffer(const FSceneView& View)
 {
 	if (const FSceneTextures* SceneTextures = static_cast<const FViewFamilyInfo*>(View.Family)->GetSceneTexturesChecked())
diff --git a/Engine/Source/Runtime/Renderer/Public/PostProcess/DrawRectangle.h b/Engine/Source/Runtime/Renderer/Public/PostProcess/DrawRectangle.h
--- a/Engine/Source/Runtime/Renderer/Public/PostProcess/DrawRectangle.h
+++ b/Engine/Source/Runtime/Renderer/Public/PostProcess/DrawRectangle.h
@@ -65,6 +65,14 @@ namespace UE::Renderer::PostProcess
 		uint32 InstanceCount = 1
 	);
 
+	RENDERER_API void DrawRectangle(
+		FRHICommandList& RHICmdList,
+		const TShaderRef<FShader>& VertexShader,
+		const FSceneView& View,
+		EDrawRectangleFlags Flags = EDRF_Default,
+		uint32 InstanceCount = 1
+	);
+
 	// NOTE: Assumes previously set PSO has PrimitiveType = PT_TriangleList
 	RENDERER_API void DrawPostProcessPass(
 		FRHICommandList& RHICmdList,
diff --git a/Engine/Source/Runtime/Renderer/Public/SceneRenderTargetParameters.h b/Engine/Source/Runtime/Renderer/Public/SceneRenderTargetParameters.h
--- a/Engine/Source/Runtime/Renderer/Public/SceneRenderTargetParameters.h
+++ b/Engine/Source/Runtime/Renderer/Public/SceneRenderTargetParameters.h
@@ -100,6 +100,8 @@ extern RENDERER_API FSceneTextureShaderParameters CreateSceneTextureShaderParame
 	const FSceneView& View,
 	ESceneTextureSetupMode SetupMode = ESceneTextureSetupMode::All);
 
+extern RENDERER_API FSceneTextureShaderParameters GetSceneTextureShaderParameters(const FSceneView& View);
+
 /** Struct containing references to extracted RHI resources after RDG execution. All textures are
  *  left in an SRV read state, so they can safely be used for read without being re-imported into
  *  RDG. Likewise, the uniform buffer is non-RDG and can be used as is.
