Water: Fixed a typo in ground extent setup when there are no ground actors. Ground extent should fully encapsulate the water bodies.

diff --git a/Engine/Plugins/Experimental/Water/Source/Runtime/Private/WaterZoneActor.cpp b/Engine/Plugins/Experimental/Water/Source/Runtime/Private/WaterZoneActor.cpp
--- a/Engine/Plugins/Experimental/Water/Source/Runtime/Private/WaterZoneActor.cpp
+++ b/Engine/Plugins/Experimental/Water/Source/Runtime/Private/WaterZoneActor.cpp
@@ -477,7 +477,7 @@ bool AWaterZone::UpdateWaterInfoTexture()
 		if (GroundActors.Num() == 0)
 		{
 			GroundZMax = WaterZMax;
-			GroundZMin = GroundZMax - CVarWaterFallbackDepth.GetValueOnGameThread();
+			GroundZMin = WaterZMin - CVarWaterFallbackDepth.GetValueOnGameThread();
 		}
 
 #if WITH_EDITOR
