Water: Changed "failed to triangulate" error to a warning since it is not a severe issue.

diff --git a/Engine/Plugins/Experimental/Water/Source/Runtime/Private/WaterBodyLakeComponent.cpp b/Engine/Plugins/Experimental/Water/Source/Runtime/Private/WaterBodyLakeComponent.cpp
--- a/Engine/Plugins/Experimental/Water/Source/Runtime/Private/WaterBodyLakeComponent.cpp
+++ b/Engine/Plugins/Experimental/Water/Source/Runtime/Private/WaterBodyLakeComponent.cpp
@@ -77,7 +77,7 @@ bool UWaterBodyLakeComponent::GenerateWaterBodyMesh(UE::Geometry::FDynamicMesh3&
 	Triangulation.Add(LakePoly);
 	if (!Triangulation.Triangulate())
 	{
-		UE_LOG(LogWater, Error, TEXT("Failed to triangulate Lake mesh for %s (%s)"), *GetOwner()->GetActorNameOrLabel(), *GetFullNameSafe(GetOwner()));
+		UE_LOG(LogWater, Warning, TEXT("Failed to triangulate Lake mesh for %s.  Ensure that the Lake's spline does not form any loops (%s)"), *GetOwner()->GetActorNameOrLabel(), *GetFullNameSafe(GetOwner()));
 	}
 
 	if (Triangulation.Triangles.Num() == 0)
diff --git a/Engine/Plugins/Experimental/Water/Source/Runtime/Private/WaterBodyOceanComponent.cpp b/Engine/Plugins/Experimental/Water/Source/Runtime/Private/WaterBodyOceanComponent.cpp
--- a/Engine/Plugins/Experimental/Water/Source/Runtime/Private/WaterBodyOceanComponent.cpp
+++ b/Engine/Plugins/Experimental/Water/Source/Runtime/Private/WaterBodyOceanComponent.cpp
@@ -222,7 +222,7 @@ bool UWaterBodyOceanComponent::GenerateWaterBodyMesh(UE::Geometry::FDynamicMesh3
 	Triangulation.Add(Island);
 	if (!Triangulation.Triangulate())
 	{
-		UE_LOG(LogWater, Error, TEXT("Failed to triangulate Ocean mesh for %s. Ensure that the Ocean's spline does not form any loops."), *GetOwner()->GetActorNameOrLabel());
+		UE_LOG(LogWater, Warning, TEXT("Failed to triangulate Ocean mesh for %s. Ensure that the Ocean's spline does not form any loops."), *GetOwner()->GetActorNameOrLabel());
 	}
 
 	if (Triangulation.Triangles.Num() == 0)
