Water: Allow water bodies which do not affect the landscape to still search for a landscape. With the current implementation of QueryWaterinfoCLosestToWorldLocation, gameplay depth queries will never be correct when a landscape is not available. This can eventually be changed back when we support proper raycasted depth calculations.

diff --git a/Engine/Plugins/Experimental/Water/Source/Runtime/Private/WaterBodyComponent.cpp b/Engine/Plugins/Experimental/Water/Source/Runtime/Private/WaterBodyComponent.cpp
--- a/Engine/Plugins/Experimental/Water/Source/Runtime/Private/WaterBodyComponent.cpp
+++ b/Engine/Plugins/Experimental/Water/Source/Runtime/Private/WaterBodyComponent.cpp
@@ -941,7 +941,7 @@ void UWaterBodyComponent::PrepareCurrentPostProcessSettings()
 
 ALandscapeProxy* UWaterBodyComponent::FindLandscape() const
 {
-	if (bAffectsLandscape && !Landscape.IsValid())
+	if (!Landscape.IsValid())
 	{
 		const FVector Location = GetComponentLocation();
 		for (TObjectIterator<ALandscapeProxy> It; It; ++It)
