Buoyancy Types:

Add missing tooltips for some properties.

diff --git a/Engine/Plugins/Experimental/Water/Source/Runtime/Public/BuoyancyTypes.h b/Engine/Plugins/Experimental/Water/Source/Runtime/Public/BuoyancyTypes.h
--- a/Engine/Plugins/Experimental/Water/Source/Runtime/Public/BuoyancyTypes.h
+++ b/Engine/Plugins/Experimental/Water/Source/Runtime/Public/BuoyancyTypes.h
@@ -218,18 +218,23 @@ struct FBuoyancyData
 	UPROPERTY(EditDefaultsOnly, Category = Buoyancy)
 	bool bApplyDragForcesInWater = false;
 
+	/** Coefficient for applying linear drag based on speed */
 	UPROPERTY(EditDefaultsOnly, Category = Buoyancy, Meta = (EditCondition = "bApplyDragForcesInWater"))
 	float DragCoefficient = 20.f;
 
+	/** Coefficient for applying linear drag based on the square of the speed */
 	UPROPERTY(EditDefaultsOnly, Category = Buoyancy, Meta = (EditCondition = "bApplyDragForcesInWater"))
 	float DragCoefficient2 = 0.01f;
 
+	/** Coefficient for applying angular drag that resists the rotation of the object */
 	UPROPERTY(EditDefaultsOnly, Category = Buoyancy, Meta = (EditCondition = "bApplyDragForcesInWater"))
 	float AngularDragCoefficient = 1.f;
 
+	/** Max speed for which drag force is applied */
 	UPROPERTY(EditDefaultsOnly, Category = Buoyancy, Meta = (EditCondition = "bApplyDragForcesInWater"))
 	float MaxDragSpeed = 15.f;
 
+	/** Whether we should apply river forces such as downstream push and shore push */
 	UPROPERTY(EditDefaultsOnly, Category = "Buoyancy | River Behavior")
 	bool bApplyRiverForces = true;
 
@@ -265,6 +270,7 @@ struct FBuoyancyData
 	UPROPERTY(EditDefaultsOnly, Category = "Buoyancy | River Behavior", Meta = (EditCondition = "bApplyRiverForces"))
 	bool bAllowCurrentWhenMovingFastUpstream = false;
 
+	/** Apply torque to align the object along the downstream direction of the river */
 	UPROPERTY(EditDefaultsOnly, Category = "Buoyancy | River Behavior", Meta = (EditCondition = "bApplyRiverForces"))
 	bool bApplyDownstreamAngularRotation = false;
 
@@ -272,16 +278,19 @@ struct FBuoyancyData
 	UPROPERTY(EditDefaultsOnly, Category = "Buoyancy | River Behavior", Meta = (EditCondition = "bApplyRiverForces && bApplyDownstreamAngularRotation"))
 	FVector DownstreamAxisOfRotation;
 
-	/** Strength of the angular rotation application */
+	/** Strength of the downstream angular rotation application */
 	UPROPERTY(EditDefaultsOnly, Category = "Buoyancy | River Behavior", Meta = (EditCondition = "bApplyRiverForces && bApplyDownstreamAngularRotation", ClampMin = 0.0f, ClampMax = 1.0f))
 	float DownstreamRotationStrength = 0.05f;
 
+	/** Stiffness of the spring used to align the object along the downstream direction */
 	UPROPERTY(EditDefaultsOnly, Category = "Buoyancy | River Behavior", Meta = (EditCondition = "bApplyRiverForces && bApplyDownstreamAngularRotation"))
 	float DownstreamRotationStiffness = 20.f;
 
+	/** Damping of the spring used to align the object along the downstream direction */
 	UPROPERTY(EditDefaultsOnly, Category = "Buoyancy | River Behavior", Meta = (EditCondition = "bApplyRiverForces && bApplyDownstreamAngularRotation"))
 	float DownstreamRotationAngularDamping = 5.0f;
 
+	/** Maximum torque to apply per update for downstream rotation */
 	UPROPERTY(EditDefaultsOnly, Category = "Buoyancy | River Behavior", Meta = (EditCondition = "bApplyRiverForces && bApplyDownstreamAngularRotation"))
 	float DownstreamMaxAcceleration = 10.0f;
 
