Init bounds when no WaterZone is present to avoid uninitialized data hitting NaN ensures

diff --git a/Engine/Plugins/Experimental/Water/Source/Runtime/Private/WaterBodyOceanComponent.cpp b/Engine/Plugins/Experimental/Water/Source/Runtime/Private/WaterBodyOceanComponent.cpp
--- a/Engine/Plugins/Experimental/Water/Source/Runtime/Private/WaterBodyOceanComponent.cpp
+++ b/Engine/Plugins/Experimental/Water/Source/Runtime/Private/WaterBodyOceanComponent.cpp
@@ -389,7 +389,7 @@ FBoxSphereBounds UWaterBodyOceanComponent::CalcBounds(const FTransform& LocalToW
 		const FVector Max(FVector(OceanExtents / 2.0, 0.0));
 		return FBoxSphereBounds(FBox(Min, Max)).TransformBy(LocalToWorld);
 	}
-	return FBoxSphereBounds().TransformBy(LocalToWorld);
+	return FBoxSphereBounds(EForceInit::ForceInit).TransformBy(LocalToWorld);
 }
 
 void UWaterBodyOceanComponent::OnUpdateBody(bool bWithExclusionVolumes)
