Water: fixed an issue where newly streamed in water bodies would not render into water info texture. Newly streamed in actors are forcibly hidden until all streaming is complete. This could lead to an issue where if actors were streamed across frames and the water body was not in the final frame, it would not render into the water info texture.

Resolved by finally updated the water info texture after the level is fully streamed in via LevelAddedToWorld callback.

diff --git a/Engine/Plugins/Experimental/Water/Source/Runtime/Private/WaterZoneActor.cpp b/Engine/Plugins/Experimental/Water/Source/Runtime/Private/WaterZoneActor.cpp
--- a/Engine/Plugins/Experimental/Water/Source/Runtime/Private/WaterZoneActor.cpp
+++ b/Engine/Plugins/Experimental/Water/Source/Runtime/Private/WaterZoneActor.cpp
@@ -598,12 +598,12 @@ void AWaterZone::OnLevelChanged(ULevel* InLevel, UWorld* InWorld)
 	}
 
 	const FBox WaterZoneBounds = GetZoneBounds();
-	const bool bContainsGroundActors = Algo::AnyOf(InLevel->Actors, [this, &WaterZoneBounds](const AActor* Actor)
+	const bool bContainsActorsAffectingWaterZone = Algo::AnyOf(InLevel->Actors, [this, &WaterZoneBounds](const AActor* Actor)
 	{
 		return IsAffectingWaterZone(WaterZoneBounds, Actor);
 	});
 
-	if (bContainsGroundActors)
+	if (bContainsActorsAffectingWaterZone)
 	{
 		MarkForRebuild(EWaterZoneRebuildFlags::UpdateWaterInfoTexture);
 	}
@@ -616,7 +616,17 @@ bool AWaterZone::IsAffectingWaterZone(const FBox& InWaterZoneBounds, const AActo
 		return false;
 	}
 
-	if (const ALandscapeProxy* LandscapeProxy = Cast<const ALandscapeProxy>(InActor))
+	if (const AWaterBody* WaterBodyActor = Cast<const AWaterBody>(InActor))
+	{
+		if (const UWaterBodyComponent* WaterBodyComponent = WaterBodyActor->GetWaterBodyComponent())
+		{
+			if (WaterBodyComponent->GetWaterZone() == this)
+			{
+				return WaterBodyComponent->AffectsWaterInfo();
+			}
+		}
+	}
+	else if (const ALandscapeProxy* LandscapeProxy = Cast<const ALandscapeProxy>(InActor))
 	{
 		bool bIntersectsWaterZone = false;
 
