World Partition - HLOD: Fix building an HLOD mesh for water components that are using WaterMeshOverride static meshes, when working with cooked content
* Source model isn't available in this case, so the mesh description needs to be constructed from render data.

diff --git a/Engine/Plugins/Experimental/Water/Source/Runtime/Private/WaterBodyComponent.cpp b/Engine/Plugins/Experimental/Water/Source/Runtime/Private/WaterBodyComponent.cpp
--- a/Engine/Plugins/Experimental/Water/Source/Runtime/Private/WaterBodyComponent.cpp
+++ b/Engine/Plugins/Experimental/Water/Source/Runtime/Private/WaterBodyComponent.cpp
@@ -37,6 +37,7 @@
 #include "WaterModule.h"
 #include "StaticMeshAttributes.h"
 #include "WaterBodyHLODBuilder.h"
+#include "MeshMergeModule.h"
 #endif
 
 #define LOCTEXT_NAMESPACE "Water"
@@ -2196,13 +2197,26 @@ TSubclassOf<UHLODBuilder> UWaterBodyComponent::GetCustomHLODBuilderClass() const
 
 FMeshDescription UWaterBodyComponent::GetHLODMeshDescription() const
 {
+	FMeshDescription MeshDescription;
+
 	if (WaterMeshOverride)
 	{
-		return *WaterMeshOverride->GetMeshDescription(WaterMeshOverride->GetNumLODs() - 1);
+		const int32 LastLOD = WaterMeshOverride->GetNumLODs() - 1;
+
+		// If source model is valid, return the mesh description directly, otherwise recreate it from the render data
+		if (WaterMeshOverride->IsSourceModelValid(LastLOD))
+		{
+			MeshDescription = *WaterMeshOverride->GetMeshDescription(LastLOD);
+		}
+		else
+		{
+			const IMeshMergeUtilities& MeshMergeUtilities = FModuleManager::Get().LoadModuleChecked<IMeshMergeModule>("MeshMergeUtilities").GetUtilities();
+			MeshMergeUtilities.RetrieveMeshDescription(WaterMeshOverride, LastLOD, MeshDescription);
+		}
+
+		return MeshDescription;
 	}
 	
-	FMeshDescription MeshDescription;
-
 	FStaticMeshAttributes StaticMeshAttributes(MeshDescription);
 	StaticMeshAttributes.Register();
 
diff --git a/Engine/Plugins/Experimental/Water/Source/Runtime/Water.Build.cs b/Engine/Plugins/Experimental/Water/Source/Runtime/Water.Build.cs
--- a/Engine/Plugins/Experimental/Water/Source/Runtime/Water.Build.cs
+++ b/Engine/Plugins/Experimental/Water/Source/Runtime/Water.Build.cs
@@ -61,7 +61,8 @@ public class Water : ModuleRules
 					"SourceControl",
 					"UnrealEd",
 					"StaticMeshDescription",
-					"MeshConversion"
+					"MeshConversion",
+					"MeshMergeUtilities"
 				}
 			);
 		}
