PR #8623: HLOD prevent transient meshcomponents/LandscapeGrass from getting merged (Contributed by [REDACTED_PERSON])

diff --git a/Engine/Plugins/Experimental/Water/Source/Runtime/Private/WaterBodyComponent.cpp b/Engine/Plugins/Experimental/Water/Source/Runtime/Private/WaterBodyComponent.cpp
--- a/Engine/Plugins/Experimental/Water/Source/Runtime/Private/WaterBodyComponent.cpp
+++ b/Engine/Plugins/Experimental/Water/Source/Runtime/Private/WaterBodyComponent.cpp
@@ -115,6 +115,11 @@ UWaterBodyComponent::UWaterBodyComponent(const FObjectInitializer& ObjectInitial
 
 bool UWaterBodyComponent::IsHLODRelevant() const
 {
+	if (HasAnyFlags(RF_Transient))
+	{
+		return false;
+	}
+
 	return bEnableAutoLODGeneration;
 }
 
diff --git a/Engine/Source/Editor/UnrealEd/Private/HierarchicalLOD.cpp b/Engine/Source/Editor/UnrealEd/Private/HierarchicalLOD.cpp
--- a/Engine/Source/Editor/UnrealEd/Private/HierarchicalLOD.cpp
+++ b/Engine/Source/Editor/UnrealEd/Private/HierarchicalLOD.cpp
@@ -704,11 +704,6 @@ bool FHierarchicalLODBuilder::ShouldGenerateCluster(AActor* Actor, const int32 H
 		return false;
 	}
 
-	if( Actor->HasAnyFlags( RF_Transient ) )
-	{
-		return false;
-	}
-
 	if (ALODActor* LODActor = Cast<ALODActor>(Actor))
 	{
 		// Ignore previous LOD actors
diff --git a/Engine/Source/Runtime/Engine/Private/Actor.cpp b/Engine/Source/Runtime/Engine/Private/Actor.cpp
--- a/Engine/Source/Runtime/Engine/Private/Actor.cpp
+++ b/Engine/Source/Runtime/Engine/Private/Actor.cpp
@@ -5883,6 +5883,11 @@ void AActor::PostRename(UObject* OldOuter, const FName OldName)
 
 bool AActor::IsHLODRelevant() const
 {
+	if (HasAnyFlags(RF_Transient))
+	{
+		return false;
+	}
+
 	if (IsHidden())
 	{
 		return false;
diff --git a/Engine/Source/Runtime/Engine/Private/Components/StaticMeshComponent.cpp b/Engine/Source/Runtime/Engine/Private/Components/StaticMeshComponent.cpp
--- a/Engine/Source/Runtime/Engine/Private/Components/StaticMeshComponent.cpp
+++ b/Engine/Source/Runtime/Engine/Private/Components/StaticMeshComponent.cpp
@@ -2855,6 +2855,11 @@ void UStaticMeshComponent::ApplyComponentInstanceData(FStaticMeshComponentInstan
 
 bool UStaticMeshComponent::IsHLODRelevant() const
 {
+	if (HasAnyFlags(RF_Transient))
+	{
+		return false;
+	}
+
 	if (!GetStaticMesh())
 	{
 		return false;
diff --git a/Engine/Source/Runtime/Landscape/Private/Landscape.cpp b/Engine/Source/Runtime/Landscape/Private/Landscape.cpp
--- a/Engine/Source/Runtime/Landscape/Private/Landscape.cpp
+++ b/Engine/Source/Runtime/Landscape/Private/Landscape.cpp
@@ -1620,6 +1620,11 @@ void ULandscapeComponent::PropagateLightingScenarioChange()
 
 bool ULandscapeComponent::IsHLODRelevant() const
 {
+	if (HasAnyFlags(RF_Transient))
+	{
+		return false;
+	}
+
 #if WITH_EDITOR
 	return bEnableAutoLODGeneration;
 #else
