Water: Added blueprint getter and setter to WaterMeshOverride.

diff --git a/Engine/Plugins/Experimental/Water/Source/Runtime/Private/WaterBodyComponent.cpp b/Engine/Plugins/Experimental/Water/Source/Runtime/Private/WaterBodyComponent.cpp
--- a/Engine/Plugins/Experimental/Water/Source/Runtime/Private/WaterBodyComponent.cpp
+++ b/Engine/Plugins/Experimental/Water/Source/Runtime/Private/WaterBodyComponent.cpp
@@ -208,6 +208,17 @@ FBox UWaterBodyComponent::GetCollisionComponentBounds() const
 	return Box;
 }
 
+void UWaterBodyComponent::SetWaterMeshOverride(UStaticMesh* InMesh)
+{
+	WaterMeshOverride = InMesh;
+
+	FOnWaterBodyChangedParams Params;
+	Params.bShapeOrPositionChanged = true;
+	Params.bUserTriggered = true;
+	Params.bWeightmapSettingsChanged = false;
+	OnWaterBodyChanged(Params);
+}
+
 AWaterBody* UWaterBodyComponent::GetWaterBodyActor() const
 {
 	// If we have an Owner, it must be an AWaterBody
diff --git a/Engine/Plugins/Experimental/Water/Source/Runtime/Public/WaterBodyComponent.h b/Engine/Plugins/Experimental/Water/Source/Runtime/Public/WaterBodyComponent.h
--- a/Engine/Plugins/Experimental/Water/Source/Runtime/Public/WaterBodyComponent.h
+++ b/Engine/Plugins/Experimental/Water/Source/Runtime/Public/WaterBodyComponent.h
@@ -199,7 +199,7 @@ public:
 	virtual float GetHeightOffset() const { return 0.f; }
 
 	/** Sets a static mesh to use as a replacement for the water mesh (for water bodies that are being rendered by the water mesh) */
-	void SetWaterMeshOverride(UStaticMesh* InMesh) { WaterMeshOverride = InMesh; }
+	void SetWaterMeshOverride(UStaticMesh* InMesh);
 
 	/** Returns River to lake transition material instance (For internal use. Please use AWaterBodyRiver instead.) */
 	UFUNCTION(BlueprintCallable, Category = Rendering)
@@ -608,7 +608,7 @@ protected:
 	UPROPERTY(Transient, DuplicateTransient, NonTransactional, VisibleAnywhere, BlueprintReadOnly, Category = Water)
 	int32 WaterBodyIndex = INDEX_NONE;
 
-	UPROPERTY(Category = Rendering, EditAnywhere, BlueprintReadOnly)
+	UPROPERTY(Category = Rendering, EditAnywhere, BlueprintReadWrite, Getter, Setter)
 	TObjectPtr<UStaticMesh> WaterMeshOverride;
 
 	/** 
