Water: Fixed a crash when copying a water body component in isolation. Spline component is not valid and crashes in PostDuplicate.

diff --git a/Engine/Plugins/Experimental/Water/Source/Runtime/Private/WaterBodyComponent.cpp b/Engine/Plugins/Experimental/Water/Source/Runtime/Private/WaterBodyComponent.cpp
--- a/Engine/Plugins/Experimental/Water/Source/Runtime/Private/WaterBodyComponent.cpp
+++ b/Engine/Plugins/Experimental/Water/Source/Runtime/Private/WaterBodyComponent.cpp
@@ -453,8 +453,12 @@ void UWaterBodyComponent::PostDuplicate(bool bDuplicateForPie)
 #if WITH_EDITOR
 	if (!bDuplicateForPie && GIsEditor)
 	{
-		// After duplication due to copy-pasting, UWaterSplineMetadata might have been edited without the spline component being made aware of that (for some reason, USplineComponent::PostDuplicate isn't called)::
-		GetWaterSpline()->SynchronizeWaterProperties();
+		// Sometimes during PostDuplicate the water body component can be created in isolation without an owner. In this case there won't be a water spline.
+		if (UWaterSplineComponent* WaterSpline = GetWaterSpline())
+		{
+			// After duplication due to copy-pasting, UWaterSplineMetadata might have been edited without the spline component being made aware of that (for some reason, USplineComponent::PostDuplicate isn't called)::
+			GetWaterSpline()->SynchronizeWaterProperties();
+		}
 
 		FOnWaterBodyChangedParams Params;
 		Params.bShapeOrPositionChanged = true;
