Water: prevent updating water info while landscacpe is painting via unnecessary visibility update on water bodies.

diff --git a/Engine/Plugins/Experimental/Water/Source/Editor/Private/WaterBrushManager.cpp b/Engine/Plugins/Experimental/Water/Source/Editor/Private/WaterBrushManager.cpp
--- a/Engine/Plugins/Experimental/Water/Source/Editor/Private/WaterBrushManager.cpp
+++ b/Engine/Plugins/Experimental/Water/Source/Editor/Private/WaterBrushManager.cpp
@@ -1271,12 +1271,6 @@ void AWaterBrushManager::RenderBrushActorContext(FBrushRenderContext& BrushRende
 	}
 	++BrushRenderContext.RTIndex;
 
-	if (WaterBody != nullptr)
-	{
-		// rebuilding the water mesh is expensive and not necessary : 
-		WaterBody->GetWaterBodyComponent()->UpdateComponentVisibility(/* bAllowWaterMeshRebuild = */false);
-	}
-
 	ApplyToCompositeWaterBodyTexture(BrushRenderContext, BrushActorRenderContext);
 }
 
