Water: added IsValid check to the texture passed to the water info rendering. Speculative fix for infrequent crash

diff --git a/Engine/Plugins/Experimental/Water/Source/Runtime/Private/WaterInfoRendering.cpp b/Engine/Plugins/Experimental/Water/Source/Runtime/Private/WaterInfoRendering.cpp
--- a/Engine/Plugins/Experimental/Water/Source/Runtime/Private/WaterInfoRendering.cpp
+++ b/Engine/Plugins/Experimental/Water/Source/Runtime/Private/WaterInfoRendering.cpp
@@ -701,7 +701,7 @@ void UpdateWaterInfoRendering(
 	RenderCaptureInterface::FScopedCapture RenderCapture((RenderCaptureNextWaterInfoDraws != 0), TEXT("RenderWaterInfo"));
 	RenderCaptureNextWaterInfoDraws = FMath::Max(0, RenderCaptureNextWaterInfoDraws - 1);
 
-	if (Context.TextureRenderTarget == nullptr || Scene == nullptr)
+	if (!IsValid(Context.TextureRenderTarget) || Scene == nullptr)
 	{
 		return;
 	}
