[Water] Update UWaterBodyComponent::UpdateWaterZones() to allow a preferred level.
Useful when there are multiple water zones in multiple levels in the same location.
A waterzone actor in a level the water actor belongs to will be given priority over other levels.

diff --git a/Engine/Plugins/Experimental/Water/Source/Runtime/Private/WaterBodyComponent.cpp b/Engine/Plugins/Experimental/Water/Source/Runtime/Private/WaterBodyComponent.cpp
--- a/Engine/Plugins/Experimental/Water/Source/Runtime/Private/WaterBodyComponent.cpp
+++ b/Engine/Plugins/Experimental/Water/Source/Runtime/Private/WaterBodyComponent.cpp
@@ -408,7 +408,9 @@ void UWaterBodyComponent::UpdateWaterZones()
 			const FBox Bounds3D = Bounds.GetBox();
 			if (const UWaterSubsystem* WaterSubsystem = UWaterSubsystem::GetWaterSubsystem(GetWorld()))
 			{
-				FoundZone = WaterSubsystem->FindWaterZone(FBox2D(FVector2D(Bounds3D.Min), FVector2D(Bounds3D.Max)));
+				const AActor* ActorOwner = GetTypedOuter<AActor>();
+				const ULevel* PreferredLevel = ActorOwner ? ActorOwner->GetLevel() : nullptr;
+				FoundZone = WaterSubsystem->FindWaterZone(FBox2D(FVector2D(Bounds3D.Min), FVector2D(Bounds3D.Max)), PreferredLevel);
 			}
 		}
 
diff --git a/Engine/Plugins/Experimental/Water/Source/Runtime/Private/WaterSubsystem.cpp b/Engine/Plugins/Experimental/Water/Source/Runtime/Private/WaterSubsystem.cpp
--- a/Engine/Plugins/Experimental/Water/Source/Runtime/Private/WaterSubsystem.cpp
+++ b/Engine/Plugins/Experimental/Water/Source/Runtime/Private/WaterSubsystem.cpp
@@ -536,7 +536,7 @@ void UWaterSubsystem::MarkWaterZonesInRegionForRebuild(const FBox2D& InUpdateReg
 	}
 }
 
-TSoftObjectPtr<AWaterZone> UWaterSubsystem::FindWaterZone(const FBox2D& Bounds) const
+TSoftObjectPtr<AWaterZone> UWaterSubsystem::FindWaterZone(const FBox2D& Bounds, const TSoftObjectPtr<const ULevel> PreferredLevel) const
 {
 
 	const UWorld* World = GetWorld();
@@ -586,6 +586,37 @@ TSoftObjectPtr<AWaterZone> UWaterSubsystem::FindWaterZone(const FBox2D& Bounds)
 		return {};
 	}
 
+	// Return best match in PreferredLevel if there is a match
+	if (!PreferredLevel.IsNull() && ViableZones.Num() > 1)
+	{
+		TSoftObjectPtr<AWaterZone> PreferredZone;
+		int32 PreferredZoneMax = INT32_MIN;
+
+		FNameBuilder ParentPath;
+		PreferredLevel.ToSoftObjectPath().ToString(ParentPath);
+		FStringView ParentPathView(ParentPath);
+
+		for (const TPair<TSoftObjectPtr<AWaterZone>, int32>& It : ViableZones)
+		{
+			if (PreferredZone.IsNull() || (It.Value > PreferredZoneMax))
+			{
+				const TSoftObjectPtr<AWaterZone>& WaterZoneSoftPath = It.Key;
+				FNameBuilder ActorPath;
+				WaterZoneSoftPath.ToSoftObjectPath().ToString(ActorPath);
+				if (ActorPath.ToView().StartsWith(ParentPathView))
+				{
+					PreferredZone = WaterZoneSoftPath;
+					PreferredZoneMax = It.Value;
+				}
+			}
+		}
+
+		if (!PreferredZone.IsNull())
+		{
+			return PreferredZone;
+		}
+	}
+
 	return Algo::MaxElementBy(ViableZones, [](const TPair<TSoftObjectPtr<AWaterZone>, int32>& A) { return A.Value; })->Key;
 }
 
diff --git a/Engine/Plugins/Experimental/Water/Source/Runtime/Public/WaterSubsystem.h b/Engine/Plugins/Experimental/Water/Source/Runtime/Public/WaterSubsystem.h
--- a/Engine/Plugins/Experimental/Water/Source/Runtime/Public/WaterSubsystem.h
+++ b/Engine/Plugins/Experimental/Water/Source/Runtime/Public/WaterSubsystem.h
@@ -177,7 +177,7 @@ public:
 	void MarkWaterZonesInRegionForRebuild(const FBox2D& InUpdateRegion, EWaterZoneRebuildFlags InRebuildFlags);
 
 	/** Returns the water with the highest priority within the bounds provided. */
-	TSoftObjectPtr<AWaterZone> FindWaterZone(const FBox2D& Bounds) const;
+	TSoftObjectPtr<AWaterZone> FindWaterZone(const FBox2D& Bounds, const TSoftObjectPtr<const ULevel> PreferredLevel = {}) const;
 
 #if WITH_EDITOR
 	/** Little scope object to temporarily change the value of bAllowWaterSubsystemOnPreviewWorld */
