Water: fixed issue where brush renderable components were not having their visibility set to false after rendering the water brush.  This was originally removed to prevent updating the water zone when the landscape was being painted so the scope of the bAllowWaterMeshRebuild parameter was expanded to include blocking the update to the water info texture.

diff --git a/Engine/Plugins/Experimental/Water/Source/Editor/Private/WaterBrushManager.cpp b/Engine/Plugins/Experimental/Water/Source/Editor/Private/WaterBrushManager.cpp
--- a/Engine/Plugins/Experimental/Water/Source/Editor/Private/WaterBrushManager.cpp
+++ b/Engine/Plugins/Experimental/Water/Source/Editor/Private/WaterBrushManager.cpp
@@ -1271,6 +1271,12 @@ void AWaterBrushManager::RenderBrushActorContext(FBrushRenderContext& BrushRende
 	}
 	++BrushRenderContext.RTIndex;
 
+	if (WaterBody != nullptr)
+	{
+		// rebuilding the water mesh is expensive and not necessary : 
+		WaterBody->GetWaterBodyComponent()->UpdateComponentVisibility(/* bAllowWaterZoneRebuild = */false);
+	}
+
 	ApplyToCompositeWaterBodyTexture(BrushRenderContext, BrushActorRenderContext);
 }
 
diff --git a/Engine/Plugins/Experimental/Water/Source/Runtime/Private/WaterBodyComponent.cpp b/Engine/Plugins/Experimental/Water/Source/Runtime/Private/WaterBodyComponent.cpp
--- a/Engine/Plugins/Experimental/Water/Source/Runtime/Private/WaterBodyComponent.cpp
+++ b/Engine/Plugins/Experimental/Water/Source/Runtime/Private/WaterBodyComponent.cpp
@@ -127,14 +127,14 @@ void UWaterBodyComponent::OnVisibilityChanged()
 {
 	Super::OnVisibilityChanged();
 
-	UpdateComponentVisibility(/* bAllowWaterMeshRebuild = */true);
+	UpdateComponentVisibility(/* bAllowWaterZoneRebuild = */true);
 }
 
 void UWaterBodyComponent::OnHiddenInGameChanged()
 {
 	Super::OnHiddenInGameChanged();
 
-	UpdateComponentVisibility(/* bAllowWaterMeshRebuild = */true);
+	UpdateComponentVisibility(/* bAllowWaterZoneRebuild = */true);
 }
 
 FPrimitiveSceneProxy* UWaterBodyComponent::CreateSceneProxy()
@@ -983,7 +983,7 @@ ALandscapeProxy* UWaterBodyComponent::FindLandscape() const
 	return Landscape.Get();
 }
 
-void UWaterBodyComponent::UpdateComponentVisibility(bool bAllowWaterMeshRebuild)
+void UWaterBodyComponent::UpdateComponentVisibility(bool bAllowWaterZoneRebuild)
 {
 	if (UWorld* World = GetWorld())
 	{
@@ -999,17 +999,24 @@ void UWaterBodyComponent::UpdateComponentVisibility(bool bAllowWaterMeshRebuild)
 	 		Component->SetHiddenInGame(bLocalHiddenInGame);
 	 	}
 
-		if (AWaterZone* WaterZone = GetWaterZone())
+		if (bAllowWaterZoneRebuild)
 		{
-			// If the component is being or can be rendered by the water mesh or renders into the water info texture, rebuild it in case its visibility has changed : 
-			if (bAllowWaterMeshRebuild && AffectsWaterMesh())
+			if (AWaterZone* WaterZone = GetWaterZone())
 			{
-				WaterZone->MarkForRebuild(EWaterZoneRebuildFlags::UpdateWaterMesh);
-			}
+				// If the component is being or can be rendered by the water mesh or renders into the water info texture, rebuild it in case its visibility has changed : 
 
-			if (AffectsWaterInfo())
-			{
-				WaterZone->MarkForRebuild(EWaterZoneRebuildFlags::UpdateWaterInfoTexture);
+				EWaterZoneRebuildFlags RebuildFlags = EWaterZoneRebuildFlags::None;
+				if (AffectsWaterMesh())
+				{
+					RebuildFlags |= EWaterZoneRebuildFlags::UpdateWaterMesh;
+				}
+
+				if (AffectsWaterInfo())
+				{
+					RebuildFlags |= EWaterZoneRebuildFlags::UpdateWaterInfoTexture;
+				}
+
+				WaterZone->MarkForRebuild(RebuildFlags);
 			}
 		}
 	}
@@ -1387,7 +1394,7 @@ void UWaterBodyComponent::UpdateAll(const FOnWaterBodyChangedParams& InParams)
 			UpdateWaterBodyRenderData();
 		}
 
-		UpdateComponentVisibility(/* bAllowWaterMeshRebuild = */true);
+		UpdateComponentVisibility(/* bAllowWaterZoneRebuild = */true);
 
 #if WITH_EDITOR
 		UpdateWaterSpriteComponent();
diff --git a/Engine/Plugins/Experimental/Water/Source/Runtime/Private/WaterSubsystem.cpp b/Engine/Plugins/Experimental/Water/Source/Runtime/Private/WaterSubsystem.cpp
--- a/Engine/Plugins/Experimental/Water/Source/Runtime/Private/WaterSubsystem.cpp
+++ b/Engine/Plugins/Experimental/Water/Source/Runtime/Private/WaterSubsystem.cpp
@@ -630,7 +630,7 @@ void UWaterSubsystem::NotifyWaterVisibilityChangedInternal(IConsoleVariable* CVa
 	// Water body visibility depends on various CVars. All need to update the visibility in water body components : 
 	GetWaterBodyManagerInternal().ForEachWaterBodyComponent([](UWaterBodyComponent* WaterBodyComponent)
 	{
-		WaterBodyComponent->UpdateComponentVisibility(/* bAllowWaterMeshRebuild = */true);
+		WaterBodyComponent->UpdateComponentVisibility(/* bAllowWaterZoneRebuild = */true);
 		return true;
 	});
 }
diff --git a/Engine/Plugins/Experimental/Water/Source/Runtime/Public/WaterBodyComponent.h b/Engine/Plugins/Experimental/Water/Source/Runtime/Public/WaterBodyComponent.h
--- a/Engine/Plugins/Experimental/Water/Source/Runtime/Public/WaterBodyComponent.h
+++ b/Engine/Plugins/Experimental/Water/Source/Runtime/Public/WaterBodyComponent.h
@@ -362,9 +362,9 @@ public:
 	bool IsWorldLocationInExclusionVolume(const FVector& InWorldLocation) const;
 
 	/** Updates the bVisible/bHiddenInGame flags on the component and eventually the child renderable components (e.g. custom water body)
-	 - bAllowWaterMeshRebuild : if true, the function will request a rebuild of the water mesh (expensive), which is necessary to take into account visibility changes
+	 - bAllowWaterZoneRebuild : if true, the function will request a rebuild of the water zone (expensive), which is necessary to take into account visibility changes
 	*/
-	virtual void UpdateComponentVisibility(bool bAllowWaterMeshRebuild);
+	virtual void UpdateComponentVisibility(bool bAllowWaterZoneRebuild);
 
 	/** Creates/Destroys/Updates necessary MIDS */
 	virtual void UpdateMaterialInstances();
