Water: fixed an issue where water body component cached bounds were not updated immediately when the water body was modified.

diff --git a/Engine/Plugins/Experimental/Water/Source/Runtime/Private/WaterBodyComponent.cpp b/Engine/Plugins/Experimental/Water/Source/Runtime/Private/WaterBodyComponent.cpp
--- a/Engine/Plugins/Experimental/Water/Source/Runtime/Private/WaterBodyComponent.cpp
+++ b/Engine/Plugins/Experimental/Water/Source/Runtime/Private/WaterBodyComponent.cpp
@@ -1341,6 +1341,9 @@ void UWaterBodyComponent::UpdateWaterBody(bool bWithExclusionVolumes)
 	{
 		OnUpdateBody(bWithExclusionVolumes);
 	}
+
+	// Update our cached bounds now that our shape may have changed:
+	UpdateBounds();
 }
 
 void UWaterBodyComponent::UpdateAll(const FOnWaterBodyChangedParams& InParams)
