Water Body Static Mesh Sections.

Repurposed existing non-tessellated water LOD system to a new static mesh system to allow water bodies to render outside the local dynamic tessellated region provided by the WaterMesh.

diff --git a/Engine/Plugins/Experimental/Water/Shaders/Private/WaterInfoFinalize.usf b/Engine/Plugins/Experimental/Water/Shaders/Private/WaterInfoFinalize.usf
--- a/Engine/Plugins/Experimental/Water/Shaders/Private/WaterInfoFinalize.usf
+++ b/Engine/Plugins/Experimental/Water/Shaders/Private/WaterInfoFinalize.usf
@@ -5,7 +5,8 @@
 
 Texture2D<float4> WaterInfoTexture;
 
-float2 WaterHeightExtents;
+float WaterZMin;
+float WaterZMax;
 float GroundZMin;
 float CaptureZ;
 
@@ -43,9 +44,6 @@ float2 VelocityBlur(float Radius, uint2 TextureCoordinates)
 
 			const FWaterInfo WaterInfo = MakeWaterInfo(WaterInfoTexture.Load(int3(CurrentTextureCoordinates, 0)));
 			
-			const float WaterZMin = WaterHeightExtents.x;
-			const float WaterZMax = WaterHeightExtents.y;
-
 			const float WaterZ = WaterInfo.NormalizedWaterZ * (WaterZMax - WaterZMin) + WaterZMin;
 			const float GroundZ = WaterInfo.NormalizedGroundZ * (WaterZMax - GroundZMin) + GroundZMin;
 
@@ -75,5 +73,5 @@ void Main(
 	const FWaterInfo WaterInfo = MakeWaterInfo(WaterInfoTexture.Load(int3(TextureCoordinates, 0)));
 	const float2 BlurredVelocity = VelocityBlur(BlurRadius, TextureCoordinates);
 
-	OutColor = float4(BlurredVelocity.rg, WaterInfo.NormalizedWaterZ, WaterInfo.NormalizedGroundZ);
+	OutColor = float4(BlurredVelocity.xy, WaterInfo.NormalizedWaterZ, WaterInfo.NormalizedGroundZ);
 }
diff --git a/Engine/Plugins/Experimental/Water/Source/Editor/Private/WaterBodyActorFactory.cpp b/Engine/Plugins/Experimental/Water/Source/Editor/Private/WaterBodyActorFactory.cpp
--- a/Engine/Plugins/Experimental/Water/Source/Editor/Private/WaterBodyActorFactory.cpp
+++ b/Engine/Plugins/Experimental/Water/Source/Editor/Private/WaterBodyActorFactory.cpp
@@ -41,6 +41,7 @@ void UWaterBodyActorFactory::PostSpawnActor(UObject* Asset, AActor* NewActor)
 	if (const FWaterBodyDefaults* WaterBodyDefaults = GetWaterBodyDefaults())
 	{
 		WaterBodyComponent->SetWaterMaterial(WaterBodyDefaults->GetWaterMaterial());
+		WaterBodyComponent->SetWaterStaticMeshMaterial(WaterBodyDefaults->GetWaterStaticMeshMaterial());
 		WaterBodyComponent->SetHLODMaterial(WaterBodyDefaults->GetWaterHLODMaterial());
 		WaterBodyComponent->SetUnderwaterPostProcessMaterial(WaterBodyDefaults->GetUnderwaterPostProcessMaterial());
 
@@ -50,6 +51,15 @@ void UWaterBodyActorFactory::PostSpawnActor(UObject* Asset, AActor* NewActor)
 			WaterSpline->WaterSplineDefaults = WaterBodyDefaults->SplineDefaults;
 		}
 	}
+
+	// If the water body is spawned into a zone which is using local only tessellation, we must default to enabling static meshes.
+	if (const AWaterZone* WaterZone = WaterBodyComponent->GetWaterZone())
+	{
+		if (WaterZone->IsLocalOnlyTessellationEnabled())
+		{
+			WaterBodyComponent->SetWaterBodyStaticMeshEnabled(true);
+		}
+	}
 }
 
 // If WaterSpline's owning actor class is a BP class, don't allow to override WaterSplineDefaults
diff --git a/Engine/Plugins/Experimental/Water/Source/Editor/Private/WaterEditorModule.cpp b/Engine/Plugins/Experimental/Water/Source/Editor/Private/WaterEditorModule.cpp
--- a/Engine/Plugins/Experimental/Water/Source/Editor/Private/WaterEditorModule.cpp
+++ b/Engine/Plugins/Experimental/Water/Source/Editor/Private/WaterEditorModule.cpp
@@ -26,6 +26,7 @@
 #include "WaterBodyBrushCacheContainerThumbnailRenderer.h"
 #include "WaterWavesEditorToolkit.h"
 #include "WaterRuntimeSettings.h"
+#include "WaterBodyOceanComponent.h"
 
 #define LOCTEXT_NAMESPACE "WaterEditor"
 
@@ -321,6 +322,13 @@ void FWaterEditorModule::OnLevelActorAddedToWorld(AActor* Actor)
 				UE_LOG(LogWaterEditor, Warning, TEXT("Could not find Water Zone class %s to spawn"), *WaterEditorSettings->GetWaterZoneClassPath().GetAssetPathString());
 			}
 		}
+
+		// If the actor is an ocean, we can help the user by initializing the ocean extent to fully fill the zone to which it belongs:
+		check(WaterBodyActor->GetWaterBodyComponent());
+		if (UWaterBodyOceanComponent* OceanComponent = Cast<UWaterBodyOceanComponent>(WaterBodyActor->GetWaterBodyComponent()))
+		{
+			OceanComponent->FillWaterZoneWithOcean();
+		}
 	}
 }
 
@@ -364,6 +372,7 @@ void FWaterEditorModule::AddWaterCollisionProfile()
 	}
 }
 
+
 IMPLEMENT_MODULE(FWaterEditorModule, WaterEditor);
 
 #undef LOCTEXT_NAMESPACE
diff --git a/Engine/Plugins/Experimental/Water/Source/Editor/Private/WaterEditorSettings.cpp b/Engine/Plugins/Experimental/Water/Source/Editor/Private/WaterEditorSettings.cpp
--- a/Engine/Plugins/Experimental/Water/Source/Editor/Private/WaterEditorSettings.cpp
+++ b/Engine/Plugins/Experimental/Water/Source/Editor/Private/WaterEditorSettings.cpp
@@ -42,6 +42,11 @@ UMaterialInterface* FWaterBodyDefaults::GetWaterMaterial() const
 	return WaterMaterial.LoadSynchronous();
 }
 
+UMaterialInterface* FWaterBodyDefaults::GetWaterStaticMeshMaterial() const
+{
+	return WaterStaticMeshMaterial.LoadSynchronous();
+}
+
 UMaterialInterface* FWaterBodyDefaults::GetWaterHLODMaterial() const
 {
 	return WaterHLODMaterial.LoadSynchronous();
@@ -64,6 +69,7 @@ FWaterBodyRiverDefaults::FWaterBodyRiverDefaults()
 	BrushDefaults.HeightmapSettings.FalloffSettings.ZOffset = 16.0f;
 	
 	WaterMaterial = FSoftObjectPath(TEXT("/Water/Materials/WaterSurface/Water_Material_River.Water_Material_River"));
+	WaterStaticMeshMaterial = FSoftObjectPath(TEXT("/Water/Materials/WaterSurface/LODs/Water_Material_River_LOD.Water_Material_River_LOD"));
 }
 
 UMaterialInterface* FWaterBodyRiverDefaults::GetRiverToOceanTransitionMaterial() const
@@ -87,6 +93,7 @@ FWaterBodyLakeDefaults::FWaterBodyLakeDefaults()
 	BrushDefaults.HeightmapSettings.FalloffSettings.ZOffset = 0.0f;
 
 	WaterMaterial = FSoftObjectPath(TEXT("/Water/Materials/WaterSurface/Water_Material_Lake.Water_Material_Lake"));
+	WaterStaticMeshMaterial = FSoftObjectPath(TEXT("/Water/Materials/WaterSurface/LODs/Water_Material_Lake_LOD.Water_Material_Lake_LOD"));
 }
 
 FWaterBodyOceanDefaults::FWaterBodyOceanDefaults()
@@ -100,6 +107,7 @@ FWaterBodyOceanDefaults::FWaterBodyOceanDefaults()
 	BrushDefaults.HeightmapSettings.FalloffSettings.ZOffset = 32.0f;
 
 	WaterMaterial = FSoftObjectPath(TEXT("/Water/Materials/WaterSurface/Water_Material_Ocean.Water_Material_Ocean"));
+	WaterStaticMeshMaterial = FSoftObjectPath(TEXT("/Water/Materials/WaterSurface/LODs/Water_Material_Ocean_LOD.Water_Material_Ocean_LOD"));
 }
 
 FWaterBodyCustomDefaults::FWaterBodyCustomDefaults()
diff --git a/Engine/Plugins/Experimental/Water/Source/Editor/Public/WaterEditorSettings.h b/Engine/Plugins/Experimental/Water/Source/Editor/Public/WaterEditorSettings.h
--- a/Engine/Plugins/Experimental/Water/Source/Editor/Public/WaterEditorSettings.h
+++ b/Engine/Plugins/Experimental/Water/Source/Editor/Public/WaterEditorSettings.h
@@ -74,6 +74,9 @@ public:
 	UMaterialInterface* GetWaterMaterial() const;
 	FSoftObjectPath GetWaterMaterialPath() const { return WaterMaterial.ToSoftObjectPath(); }
 
+	UMaterialInterface* GetWaterStaticMeshMaterial() const;
+	FSoftObjectPath GetWaterStaticMeshMaterialPath() const { return WaterStaticMeshMaterial.ToSoftObjectPath(); }
+
 	UMaterialInterface* GetWaterHLODMaterial() const;
 	FSoftObjectPath GetWaterHLODMaterialPath() const { return WaterHLODMaterial.ToSoftObjectPath(); }
 
@@ -84,6 +87,9 @@ protected:
 	UPROPERTY(EditAnywhere, config, Category = Rendering)
 	TSoftObjectPtr<UMaterialInterface> WaterMaterial;
 
+	UPROPERTY(EditAnywhere, config, Category = Rendering)
+	TSoftObjectPtr<UMaterialInterface> WaterStaticMeshMaterial;
+
 	UPROPERTY(EditAnywhere, config, Category = Rendering)
 	TSoftObjectPtr<UMaterialInterface> WaterHLODMaterial;
 
@@ -233,7 +239,7 @@ public:
 
 	bool GetUpdateLandscapeDuringInteractiveChanges() const { return bUpdateLandscapeDuringInteractiveChanges; }
 	bool GetUpdateWaterMeshDuringInteractiveChanges() const { return bUpdateWaterMeshDuringInteractiveChanges; }
-
+	
 public:
 	/** The texture group to use for generated textures such as the combined veloctiy and height texture */
 	UPROPERTY(EditAnywhere, config, Category = Rendering)
diff --git a/Engine/Plugins/Experimental/Water/Source/Runtime/Private/WaterBodyActor.cpp b/Engine/Plugins/Experimental/Water/Source/Runtime/Private/WaterBodyActor.cpp
--- a/Engine/Plugins/Experimental/Water/Source/Runtime/Private/WaterBodyActor.cpp
+++ b/Engine/Plugins/Experimental/Water/Source/Runtime/Private/WaterBodyActor.cpp
@@ -15,6 +15,8 @@
 #include "WaterBodyLakeComponent.h"
 #include "WaterBodyOceanComponent.h"
 #include "WaterBodyRiverComponent.h"
+#include "WaterBodyInfoMeshComponent.h"
+#include "WaterBodyStaticMeshComponent.h"
 #include "WaterVersion.h"
 
 #include UE_INLINE_GENERATED_CPP_BY_NAME(WaterBodyActor)
@@ -53,6 +55,12 @@ AWaterBody::AWaterBody(const FObjectInitializer& ObjectInitializer)
 	// Temporarily set the root component to the spline because the WaterBodyComponent has not yet been created
 	RootComponent = SplineComp;
 
+	WaterInfoMeshComponent = CreateDefaultSubobject<UWaterBodyInfoMeshComponent>(TEXT("WaterInfoMeshComponent"));
+	WaterInfoMeshComponent->SetupAttachment(RootComponent);
+
+	DilatedWaterInfoMeshComponent = CreateDefaultSubobject<UWaterBodyInfoMeshComponent>(TEXT("DilatedWaterInfoMeshComponent"));
+	DilatedWaterInfoMeshComponent->SetupAttachment(RootComponent);
+
 #if WITH_EDITORONLY_DATA
 	bAffectsLandscape_DEPRECATED = true;
 	CollisionProfileName_DEPRECATED = GetDefault<UWaterRuntimeSettings>()->GetDefaultWaterCollisionProfileName();
@@ -194,6 +202,16 @@ void AWaterBody::SetWaterWavesInternal(UWaterWavesBase* InWaterWaves)
 	}
 }
 
+void AWaterBody::SetWaterBodyStaticMeshComponents(TArrayView<TObjectPtr<UWaterBodyStaticMeshComponent>> NewComponentList, TConstArrayView<TObjectPtr<UWaterBodyStaticMeshComponent>> ComponentsToUnregister)
+{
+	for (const TObjectPtr<UWaterBodyStaticMeshComponent>& StaticMeshComponent : ComponentsToUnregister)
+	{
+		StaticMeshComponent.Get()->UnregisterComponent();
+		StaticMeshComponent.Get()->DestroyComponent();
+	}
+	WaterBodyStaticMeshComponents = NewComponentList;
+}
+
 static FName GetWaterBodyComponentName(EWaterBodyType Type)
 {
 	switch (Type)
@@ -572,6 +590,16 @@ void AWaterBody::PostRegisterAllComponents()
 		// At this point, the water mesh actor should be ready and we can setup the MID accordingly : 
 		// Needs to be done at the end so that all data needed by the MIDs (e.g. WaterBodyIndex) is up to date :
 		WaterBodyComponent->UpdateMaterialInstances();
+
+#if WITH_EDITOR
+		// We need to generate the water body render data for the first time here after the waterbody component is registered to the actor.
+		if (GetLinkerCustomVersion(FFortniteMainBranchObjectVersion::GUID) < FFortniteMainBranchObjectVersion::WaterBodyStaticMeshComponents)
+		{
+			WaterBodyComponent->UpdateWaterBodyRenderData();
+		}
+#endif // WITH_EDITOR
+
+		WaterBodyComponent->OnPostRegisterAllComponents();
 	}
 }
 
diff --git a/Engine/Plugins/Experimental/Water/Source/Runtime/Private/WaterBodyComponent.cpp b/Engine/Plugins/Experimental/Water/Source/Runtime/Private/WaterBodyComponent.cpp
--- a/Engine/Plugins/Experimental/Water/Source/Runtime/Private/WaterBodyComponent.cpp
+++ b/Engine/Plugins/Experimental/Water/Source/Runtime/Private/WaterBodyComponent.cpp
@@ -24,10 +24,11 @@
 #include "WaterVersion.h"
 #include "Misc/UObjectToken.h"
 #include "Logging/MessageLog.h"
-#include "WaterBodySceneProxy.h"
 #include "Engine/StaticMesh.h"
-#include "DynamicMesh/DynamicMeshAABBTree3.h"
-#include "Operations/MeshPlaneCut.h"
+#include "WaterBodyStaticMeshComponent.h"
+#include "WaterBodyMeshComponent.h"
+#include "WaterBodyInfoMeshComponent.h"
+#include "LocalVertexFactory.h"
 #include "UObject/ICookInfo.h"
 #include "UObject/ObjectSaveContext.h"
 
@@ -40,6 +41,8 @@
 #include "StaticMeshAttributes.h"
 #include "WaterBodyHLODBuilder.h"
 #include "MeshMergeModule.h"
+#include "WaterBodyMeshBuilder.h"
+#include "MeshDescription.h"
 #endif
 
 #define LOCTEXT_NAMESPACE "Water"
@@ -61,22 +64,6 @@ TAutoConsoleVariable<float> CVarWaterOceanFallbackDepth(
 	TEXT("Depth to report for the ocean when no terrain is found under the query location. Not used when <= 0."),
 	ECVF_Default);
 
-
-FWaterBodyMeshSection::FWaterBodyMeshSection()
-	: Vertices()
-	, Indices()
-	, SectionBounds()
-{
-}
-
-FWaterBodyMeshSection::~FWaterBodyMeshSection() = default;
-
-uint32 FWaterBodyMeshSection::GetAllocatedSize() const
-{
-	return Vertices.GetAllocatedSize() + Indices.GetAllocatedSize();
-}
-
-
 const FName UWaterBodyComponent::WaterBodyIndexParamName(TEXT("WaterBodyIndex"));
 const FName UWaterBodyComponent::WaterBodyZOffsetParamName(TEXT("WaterBodyZOffset"));
 const FName UWaterBodyComponent::WaterVelocityAndHeightName(TEXT("WaterVelocityAndHeight"));
@@ -139,44 +126,6 @@ void UWaterBodyComponent::OnHiddenInGameChanged()
 	UpdateComponentVisibility(/* bAllowWaterZoneRebuild = */true);
 }
 
-FPrimitiveSceneProxy* UWaterBodyComponent::CreateSceneProxy()
-{
-	if (AffectsWaterMesh())
-	{
-		return new FWaterBodySceneProxy(this);
-	}
-	return nullptr;
-}
-
-void UWaterBodyComponent::GetUsedMaterials(TArray<UMaterialInterface*>& OutMaterialInterfaces, bool bGetDebugMaterials) const
-{
-	if (WaterInfoMID)
-	{
-		OutMaterialInterfaces.Add(WaterInfoMID);
-	}
-	if (WaterLODMID)
-	{
-		OutMaterialInterfaces.Add(WaterLODMID);
-	}
-}
-
-
-void UWaterBodyComponent::OnTessellatedWaterMeshBoundsChanged()
-{
-	// When the water info bounds changes, we need to push new bounds to the MIDs so they can address the texture correctly.
-	UpdateMaterialInstances();
-}
-
-void UWaterBodyComponent::PushTessellatedWaterMeshBoundsToProxy(const FBox2D& TessellatedWaterMeshBounds)
-{
-	OnTessellatedWaterMeshBoundsChanged();
-
-	if (SceneProxy)
-	{
-		static_cast<FWaterBodySceneProxy*>(SceneProxy)->OnTessellatedWaterMeshBoundsChanged_GameThread(TessellatedWaterMeshBounds);
-	}
-}
-
 bool UWaterBodyComponent::IsFlatSurface() const
 {
 	// Lakes and oceans have surfaces aligned with the XY plane
@@ -236,6 +185,24 @@ UWaterSplineComponent* UWaterBodyComponent::GetWaterSpline() const
 	return nullptr;
 }
 
+UWaterBodyInfoMeshComponent* UWaterBodyComponent::GetWaterInfoMeshComponent() const
+{
+	if (const AWaterBody* OwningWaterBody = GetWaterBodyActor())
+	{
+		return OwningWaterBody->WaterInfoMeshComponent.Get();
+	}
+	return nullptr;
+}
+
+UWaterBodyInfoMeshComponent* UWaterBodyComponent::GetDilatedWaterInfoMeshComponent() const
+{
+	if (const AWaterBody* OwningWaterBody = GetWaterBodyActor())
+	{
+		return OwningWaterBody->DilatedWaterInfoMeshComponent.Get();
+	}
+	return nullptr;
+}
+
 bool UWaterBodyComponent::IsWaterSplineClosedLoop() const
 {
 	return (GetWaterBodyType() == EWaterBodyType::Lake) || (GetWaterBodyType() == EWaterBodyType::Ocean);
@@ -287,8 +254,20 @@ void UWaterBodyComponent::GetBrushRenderDependencies(TSet<UObject*>& OutDependen
 
 void UWaterBodyComponent::SetWaterMaterial(UMaterialInterface* InMaterial)
 {
-	WaterMaterial = InMaterial;
-	UpdateMaterialInstances();
+	if (WaterMaterial != InMaterial)
+	{
+		WaterMaterial = InMaterial;
+		UpdateMaterialInstances();
+	}
+}
+
+void UWaterBodyComponent::SetWaterStaticMeshMaterial(UMaterialInterface* InMaterial)
+{
+	if (WaterStaticMeshMaterial != InMaterial)
+	{
+		WaterStaticMeshMaterial = InMaterial;
+		UpdateMaterialInstances();
+	}
 }
 
 UMaterialInstanceDynamic* UWaterBodyComponent::GetWaterMaterialInstance()
@@ -309,10 +288,10 @@ UMaterialInstanceDynamic* UWaterBodyComponent::GetWaterInfoMaterialInstance()
 	return WaterInfoMID;
 }
 
-UMaterialInstanceDynamic* UWaterBodyComponent::GetWaterLODMaterialInstance()
+UMaterialInstanceDynamic* UWaterBodyComponent::GetWaterStaticMeshMaterialInstance()
 {
-	CreateOrUpdateWaterLODMID();
-	return WaterLODMID;
+	CreateOrUpdateWaterStaticMeshMID();
+	return WaterStaticMeshMID;
 }
 
 void UWaterBodyComponent::SetUnderwaterPostProcessMaterial(UMaterialInterface* InMaterial)
@@ -865,7 +844,7 @@ void UWaterBodyComponent::UpdateMaterialInstances()
 {
 	CreateOrUpdateWaterMID();
 	CreateOrUpdateWaterInfoMID();
-	CreateOrUpdateWaterLODMID();
+	CreateOrUpdateWaterStaticMeshMID();
 	CreateOrUpdateUnderwaterPostProcessMID();
 }
 
@@ -926,18 +905,36 @@ void UWaterBodyComponent::CreateOrUpdateWaterInfoMID()
 	if (GetWorld())
 	{
 		WaterInfoMID = FWaterUtils::GetOrCreateTransientMID(WaterInfoMID, TEXT("WaterInfoMID"), WaterInfoMaterial, GetTransientMIDFlags());
+		if (UWaterBodyInfoMeshComponent* WaterInfoMesh = GetWaterInfoMeshComponent())
+		{
+			WaterInfoMesh->SetMaterial(0, WaterInfoMID);
+		}
+		if (UWaterBodyInfoMeshComponent* DilatedWaterMesh = GetDilatedWaterInfoMeshComponent())
+		{
+			DilatedWaterMesh->SetMaterial(0, WaterInfoMID);
+		}
 
 		SetDynamicParametersOnWaterInfoMID(WaterInfoMID);
 	}
 }
 
-void UWaterBodyComponent::CreateOrUpdateWaterLODMID()
+void UWaterBodyComponent::CreateOrUpdateWaterStaticMeshMID()
 {
 	if (GetWorld())
 	{
-		WaterLODMID = FWaterUtils::GetOrCreateTransientMID(WaterLODMID, TEXT("WaterLODMID"), WaterLODMaterial, GetTransientMIDFlags());
+		WaterStaticMeshMID = FWaterUtils::GetOrCreateTransientMID(WaterStaticMeshMID, TEXT("WaterLODMID"), WaterStaticMeshMaterial, GetTransientMIDFlags());
+		if (AWaterBody* WaterBodyActor = GetWaterBodyActor())
+		{
+			for (UWaterBodyStaticMeshComponent* WaterBodyStaticMeshComponent : WaterBodyActor->GetWaterBodyStaticMeshComponents())
+			{
+				if (IsValid(WaterBodyStaticMeshComponent))
+				{
+					WaterBodyStaticMeshComponent->SetMaterial(0, WaterStaticMeshMID);
+				}
+			}
+		}
 
-		SetDynamicParametersOnMID(WaterLODMID);
+		SetDynamicParametersOnMID(WaterStaticMeshMID);
 	}
 }
 
@@ -1080,6 +1077,10 @@ void UWaterBodyComponent::OnPostEditChangeProperty(FOnWaterBodyChangedParams& In
 	{
 		UpdateMaterialInstances();
 	}
+	else if (PropertyName == GET_MEMBER_NAME_CHECKED(UWaterBodyComponent, WaterStaticMeshMaterial))
+	{
+		UpdateMaterialInstances();
+	}
 	else if (PropertyName == GET_MEMBER_NAME_CHECKED(UWaterBodyComponent, TargetWaveMaskDepth))
 	{
 		RequestGPUWaveDataUpdate();
@@ -1100,9 +1101,13 @@ void UWaterBodyComponent::OnPostEditChangeProperty(FOnWaterBodyChangedParams& In
 		}
 		InOutOnWaterBodyChangedParams.bShapeOrPositionChanged = true;
 	}
+	else if (PropertyChangedEvent.GetMemberPropertyName() == GET_MEMBER_NAME_CHECKED(UWaterBodyComponent, StaticMeshSettings))
+	{
+		UpdateWaterBodyStaticMeshComponents();
+	}
 }
 
-TArray<TSharedRef<FTokenizedMessage>> UWaterBodyComponent::CheckWaterBodyStatus() const
+TArray<TSharedRef<FTokenizedMessage>> UWaterBodyComponent::CheckWaterBodyStatus()
 {
 	TArray<TSharedRef<FTokenizedMessage>> Result;
 
@@ -1119,13 +1124,31 @@ TArray<TSharedRef<FTokenizedMessage>> UWaterBodyComponent::CheckWaterBodyStatus(
 		}
 		else
 		{
-			if (AffectsWaterMesh() && (GetWaterZone() == nullptr))
+			if (AffectsWaterMesh())
 			{
-				Result.Add(FTokenizedMessage::Create(EMessageSeverity::Error)
-					->AddToken(FUObjectToken::Create(this))
-					->AddToken(FTextToken::Create(FText::Format(
-						LOCTEXT("MapCheck_Message_MissingWaterZone", "Water body {0} requires a WaterZone actor to be rendered. Please add one to the map. "),
-						FText::FromString(GetWaterBodyActor()->GetActorLabel())))));
+				AWaterZone* WaterZone = GetWaterZone();
+				if (WaterZone == nullptr)
+				{
+					Result.Add(FTokenizedMessage::Create(EMessageSeverity::Error)
+						->AddToken(FUObjectToken::Create(this))
+						->AddToken(FTextToken::Create(FText::Format(
+							LOCTEXT("MapCheck_Message_MissingWaterZone", "Water body {0} requires a WaterZone actor to be rendered. Please add one to the map. "),
+							FText::FromString(GetWaterBodyActor()->GetActorLabel())))));
+				}
+				else
+				{
+					if (WaterZone->IsLocalOnlyTessellationEnabled() && !StaticMeshSettings.bEnableWaterBodyStaticMesh)
+					{
+						Result.Add(FTokenizedMessage::Create(EMessageSeverity::Warning)
+							->AddToken(FUObjectToken::Create(this))
+							->AddToken(FTextToken::Create(FText::Format(
+								LOCTEXT("MapCheck_Message_LocalTessellationWithoutStaticMesh", "Water body {0} is rendered by a water zone with local tessellation enabled but does not have its own static mesh enabled to fallback on!"),
+								FText::FromString(GetWaterBodyActor()->GetActorLabel()))))
+							->AddToken(FActionToken::Create(LOCTEXT("MapCheck_MessageAction_EnableWaterBodyStaticMesh", "Click here to enable the static mesh"), FText(),
+								FOnActionTokenExecuted::CreateUObject(this, &UWaterBodyComponent::SetWaterBodyStaticMeshEnabled, true), true))
+								);
+					}
+				}
 			}
 
 			if (AffectsLandscape() && (FindLandscape() == nullptr))
@@ -1395,8 +1418,6 @@ void UWaterBodyComponent::UpdateAll(const FOnWaterBodyChangedParams& InParams)
 		if (bShapeOrPositionChanged)
 		{
 			FNavigationSystem::UpdateActorAndComponentData(*WaterBodyOwner);
-
-			UpdateWaterBodyRenderData();
 		}
 
 		UpdateComponentVisibility(/* bAllowWaterZoneRebuild = */true);
@@ -1410,7 +1431,6 @@ void UWaterBodyComponent::UpdateAll(const FOnWaterBodyChangedParams& InParams)
 	if (bShapeOrPositionChanged)
 	{
 		UpdateWaterZones();
-		UpdateWaterBodyRenderData();
 	}
 }
 
@@ -1452,6 +1472,11 @@ void UWaterBodyComponent::OnWaterBodyChanged(const FOnWaterBodyChangedParams& In
 	}
 
 #if WITH_EDITOR
+	if (InParams.PropertyChangedEvent.ChangeType == EPropertyChangeType::ValueSet)
+	{
+		UpdateWaterBodyRenderData();
+	}
+
 	AWaterBody* const WaterBodyActor = GetWaterBodyActor();
 	if (WaterBodyActor == nullptr)
 	{
@@ -1529,6 +1554,10 @@ void UWaterBodyComponent::PostLoad()
 	}
 
 #if WITH_EDITOR
+	if (GetLinkerCustomVersion(FFortniteMainBranchObjectVersion::GUID) < FFortniteMainBranchObjectVersion::WaterBodyStaticMeshComponents)
+	{
+		UpdateWaterBodyRenderData();
+	}
 	RegisterOnUpdateWavesData(GetWaterWaves(), /* bRegister = */true);
 #endif // WITH_EDITOR
 }
@@ -1551,6 +1580,11 @@ void UWaterBodyComponent::DeprecateData()
 		// Transfer info to sub-components :
 		ApplyNavigationSettings();
 	}
+
+	if (GetLinkerCustomVersion(FFortniteMainBranchObjectVersion::GUID) < FFortniteMainBranchObjectVersion::WaterBodyStaticMeshComponents)
+	{
+		WaterStaticMeshMaterial = WaterLODMaterial_DEPRECATED;
+	}
 #endif // WITH_EDITORONLY_DATA
 }
 
@@ -1606,21 +1640,6 @@ void UWaterBodyComponent::OnGenerateOverlapEventsChanged()
 void UWaterBodyComponent::GetResourceSizeEx(FResourceSizeEx& CumulativeResourceSize)
 {
 	Super::GetResourceSizeEx(CumulativeResourceSize);
-
-	uint32 SectionsSize = 0;
-	for (const FWaterBodyMeshSection& Section : WaterBodyMeshSections)
-	{
-		SectionsSize += Section.GetAllocatedSize();
-	}
-
-	// Account for all non-editor data properties :
-	CumulativeResourceSize.AddDedicatedSystemMemoryBytes(
-		WaterBodyMeshSections.GetAllocatedSize() 
-		+ SectionsSize
-		+ WaterBodyMeshIndices.GetAllocatedSize()
-		+ WaterBodyMeshVertices.GetAllocatedSize()
-		+ DilatedWaterBodyMeshIndices.GetAllocatedSize()
-		+ DilatedWaterBodyMeshVertices.GetAllocatedSize());
 }
 
 bool UWaterBodyComponent::SetDynamicParametersOnMID(UMaterialInstanceDynamic* InMID)
@@ -1646,9 +1665,9 @@ bool UWaterBodyComponent::SetDynamicParametersOnMID(UMaterialInstanceDynamic* In
 		const UWaterMeshComponent* WaterMeshComponent = WaterZone->GetWaterMeshComponent();
 		check(WaterMeshComponent);
 
-		// Location should be the bottom left of the zone
-		const FVector2D WaterInfoExtent = FVector2D(WaterZone->GetTessellatedWaterMeshExtent());
-		const FVector2D WaterInfoLocation = FVector2D(WaterZone->GetTessellatedWaterMeshCenter()) - (WaterInfoExtent / 2.f);
+		// Location is the bottom left of the zone
+		const FVector2D WaterInfoExtent = FVector2D(WaterZone->GetDynamicWaterMeshExtent());
+		const FVector2D WaterInfoLocation = FVector2D(WaterZone->GetDynamicWaterMeshCenter()) - (WaterInfoExtent / 2.f);
 
 		FVector4 WaterArea;
 		WaterArea.X = WaterInfoLocation.X;
@@ -1663,6 +1682,10 @@ bool UWaterBodyComponent::SetDynamicParametersOnMID(UMaterialInstanceDynamic* In
 		InMID->SetScalarParameterValue(WaterZMaxParamName, WaterHeightExtents.Y);
 		InMID->SetScalarParameterValue(GroundZMinParamName, GroundZMin);
 	}
+	else
+	{
+		InMID->SetDoubleVectorParameterValue(WaterAreaParamName, FVector4::Zero());
+	}
 
 	return true;
 }
@@ -1804,373 +1827,106 @@ UWaterWavesBase* UWaterBodyComponent::GetWaterWaves() const
 	return nullptr;
 }
 
-static inline FColor PackFlowData(float VelocityMagnitude, float DirectionAngle, float MaxVelocity)
+bool UWaterBodyComponent::GenerateWaterBodyMesh(UE::Geometry::FDynamicMesh3& OutMesh, UE::Geometry::FDynamicMesh3* OutDilatedMesh /* = nullptr */) const
 {
-	check((DirectionAngle >= 0.f) && (DirectionAngle <= TWO_PI));
-
-	const float NormalizedMagnitude = FMath::Clamp(VelocityMagnitude, 0.f, MaxVelocity) / MaxVelocity;
-	const float NormalizedAngle = DirectionAngle / TWO_PI;
-	const float MappedMagnitude = NormalizedMagnitude * TNumericLimits<uint16>::Max();
-	const float MappedAngle = NormalizedAngle * TNumericLimits<uint16>::Max();
-
-	const uint16 ResultMag = (uint16)MappedMagnitude;
-	const uint16 ResultAngle = (uint16)MappedAngle;
-	
-	FColor Result;
-	Result.R = (ResultMag >> 8) & 0xFF;
-	Result.G = (ResultMag >> 0) & 0xFF;
-	Result.B = (ResultAngle >> 8) & 0xFF;
-	Result.A = (ResultAngle >> 0) & 0xFF;
-
-	return Result;
+	 checkf(!AffectsWaterInfo(), TEXT("WaterBodyComponent affects water info but does not implement GenerateWaterBodyMesh!")); 
+	 return false;
 }
 
-
-void UWaterBodyComponent::RebuildWaterBodyInfoMesh()
+#if WITH_EDITOR
+void UWaterBodyComponent::SetWaterBodyStaticMeshEnabled(bool bEnabled)
 {
-	TRACE_CPUPROFILER_EVENT_SCOPE(UWaterBodyComponent::RebuildWaterBodyInfoMesh);
-
-	using namespace UE::Geometry;
-
-	WaterBodyMeshVertices.Empty();
-	WaterBodyMeshIndices.Empty();
-	DilatedWaterBodyMeshVertices.Empty();
-	DilatedWaterBodyMeshIndices.Empty();
-
-	FDynamicMesh3 WaterBodyMesh(EMeshComponents::VertexColors);
-	FDynamicMesh3 WaterBodyDilatedMesh(EMeshComponents::None);
+	bool bChanged = StaticMeshSettings.bEnableWaterBodyStaticMesh != bEnabled;
+	StaticMeshSettings.bEnableWaterBodyStaticMesh = bEnabled;
 
-	GenerateWaterBodyMesh(WaterBodyMesh, &WaterBodyDilatedMesh);
-	
-	const float MaxFlowVelocity = FWaterUtils::GetWaterMaxFlowVelocity(false);
-	WaterBodyMeshVertices.Reserve(WaterBodyMesh.VertexCount());
-	WaterBodyMeshIndices.Reserve(3 * WaterBodyMesh.TriangleCount());
-	for (int VertexIndex : WaterBodyMesh.VertexIndicesItr())
+	if (bChanged)
 	{
-		FVertexInfo VertexInfo;
-		const bool bWantColors = true;
-		const bool bWantNormals = false;
-		const bool bWantUVs = false;
-		WaterBodyMesh.GetVertex(VertexIndex, VertexInfo, bWantNormals, bWantColors, bWantUVs);
-
-		const FVector3f LocalSpacePosition = FVector3f(VertexInfo.Position);
-		FDynamicMeshVertex MeshVertex(LocalSpacePosition);
-
-		const float FlowSpeed = VertexInfo.Color.X;
-		const float FlowDirection = VertexInfo.Color.Y;
-		MeshVertex.Color = PackFlowData(FlowSpeed, FlowDirection, MaxFlowVelocity);
-
-		WaterBodyMeshVertices.Add(MeshVertex);
-	}
-	for (FIndex3i Triangle : WaterBodyMesh.TrianglesItr())
-	{
-		WaterBodyMeshIndices.Append({ (uint32)Triangle.A, (uint32)Triangle.B, (uint32)Triangle.C });
-	}
-
-	DilatedWaterBodyMeshIndices.Reserve(3 * WaterBodyDilatedMesh.TriangleCount());
-	DilatedWaterBodyMeshVertices.Reserve(WaterBodyDilatedMesh.VertexCount());
-	for (FVector3d Vertex : WaterBodyDilatedMesh.VerticesItr())
-	{
-		DilatedWaterBodyMeshVertices.Emplace(FVector3f(Vertex));
-	}
-	for (FIndex3i Triangle : WaterBodyDilatedMesh.TrianglesItr())
-	{
-		DilatedWaterBodyMeshIndices.Append( {(uint32)Triangle.A, (uint32)Triangle.B, (uint32)Triangle.C });
+		UpdateWaterBodyRenderData();
 	}
 }
 
-void UWaterBodyComponent::RebuildWaterBodyLODSections()
+void UWaterBodyComponent::UpdateWaterBodyRenderData()
 {
-	using namespace UE::Geometry;
-
-	WaterBodyMeshSections.Empty();
-
-	TRACE_CPUPROFILER_EVENT_SCOPE(UWaterBodyComponent::RebuildWaterBodyLODSections);
-
-	const AWaterZone* WaterZone = GetWaterZone();
-	check(WaterZone);
-
-	const UWaterMeshComponent* WaterMesh = WaterZone->GetWaterMeshComponent();
-	check(WaterMesh);
-
-	const float WaterLODSectionSize = WaterZone->GetNonTessellatedLODSectionSize();
-	const FVector2i WaterZoneExtentInLODSections = FVector2i(FMath::DivideAndRoundUp(WaterZone->GetZoneExtent(), FVector2d(WaterLODSectionSize)));
-
-	const float GridTileSize = WaterMesh->GetTileSize();
-	const FVector2d GridPosition = FVector2d(WaterMesh->GetComponentLocation().GridSnap(GridTileSize));
-	const FVector2d MeshOrigin = GridPosition - (GridTileSize * FVector2d(WaterMesh->GetExtentInTiles()));
-
-	FDynamicMesh3 EditMesh(EMeshComponents::VertexColors);
+	UpdateWaterInfoMeshComponents();
 
-	GenerateWaterBodyMesh(EditMesh);
+	UpdateWaterBodyStaticMeshComponents();
 
-	// Transform everything to world space. This is necessary to maintain an axis-aligned grid. The water mesh is guaranteed to always be axis-aligned in worldspace
-	// so we transform into that coordinate system to guarantee the calculations occur with axis-aligned section tiles.
-	for (int32 VID : EditMesh.VertexIndicesItr())
-	{
-		FVertexInfo VertexInfo = EditMesh.GetVertexInfo(VID);
-		VertexInfo.Position = GetComponentTransform().TransformPosition(VertexInfo.Position);
-		EditMesh.SetVertex(VID, VertexInfo.Position);
-	}
+	OnWaterBodyRenderDataUpdated();
+}
 
+void UWaterBodyComponent::OnWaterBodyRenderDataUpdated()
+{
+	const bool bAffectsWaterInfo = AffectsWaterInfo();
+	const bool bAffectsWaterMesh = AffectsWaterMesh();
 
-	// Slice the mesh along the grid defined by the Water Mesh
+	if (bAffectsWaterInfo)
 	{
-		TRACE_CPUPROFILER_EVENT_SCOPE(ComputeWaterBodyMeshSections::PartitionWaterBodyMesh);
-
-		// We start at an index of -1 and go to +1 here to ensure the bordering sections don't contain triangles that could extend outside the bounds of the cell.
-		for (int32 SectionIndexX = -1; SectionIndexX < WaterZoneExtentInLODSections.X + 1; ++SectionIndexX)
-		{
-			const FVector3d SlicePos = FVector3d(MeshOrigin + SectionIndexX * FVector2d(WaterLODSectionSize, 0.0), 0.0);
-			const FVector3d SliceNormal(-1.0, 0.0, 0.0);
-
-			FMeshPlaneCut Cut(&EditMesh, SlicePos, SliceNormal);
-			if (!ensure(Cut.CutWithoutDelete(false)))
-			{
-				continue;
-			}
-		}
-
-		for (int32 SectionIndexY = -1; SectionIndexY < WaterZoneExtentInLODSections.Y + 1; ++SectionIndexY)
+		if (AWaterZone* WaterZone = GetWaterZone())
 		{
-			const FVector3d SlicePos = FVector3d(MeshOrigin + SectionIndexY * FVector2d(0.0, WaterLODSectionSize), 0.0);
-			const FVector3d SliceNormal(0.0, -1.0, 0.0);
-
-			FMeshPlaneCut Cut(&EditMesh, SlicePos, SliceNormal);
-			if (!ensure(Cut.CutWithoutDelete(false)))
-			{
-				continue;
-			}
+			WaterZone->MarkForRebuild(EWaterZoneRebuildFlags::UpdateWaterInfoTexture);
 		}
 	}
 
-	FDynamicMeshAABBTree3 AABBTree(&EditMesh, true);
-
-	const uint32 NumberOfSections = WaterZoneExtentInLODSections.X * WaterZoneExtentInLODSections.Y;
-	TArray<FWaterBodyMeshSection> Sections;
-	Sections.SetNum(NumberOfSections);
-
-	// Need to access this before the ParallelFor since Foreground Worker threads don't count as gamethread.
-	const float MaxFlowVelocity = FWaterUtils::GetWaterMaxFlowVelocity(false);
-	ParallelFor(NumberOfSections, [&](int32 SectionIndex)
+	if (bAffectsWaterInfo || bAffectsWaterMesh)
 	{
-		TRACE_CPUPROFILER_EVENT_SCOPE(ComputeWaterBodyMeshSections::BuildMeshSection);
-		const int32 SectionIndexX = SectionIndex % WaterZoneExtentInLODSections.X;
-		const int32 SectionIndexY = SectionIndex / WaterZoneExtentInLODSections.X;
-
-		const FVector2d SectionPos = MeshOrigin + FVector2d(SectionIndexX, SectionIndexY) * WaterLODSectionSize;
-
-		const FAxisAlignedBox2d GridRectXY(SectionPos, SectionPos + WaterLODSectionSize);
-
-		FDynamicMeshAABBTree3::FTreeTraversal Traversal;
-		
-		int SelectAllDepth = TNumericLimits<int>::Max();
-		int CurrentDepth = -1;
-		
-		TArray<int> Triangles;
-
-		Traversal.NextBoxF =
-			[&GridRectXY, &SelectAllDepth, &CurrentDepth](const FAxisAlignedBox3d& Box, int Depth)
-		{
-			CurrentDepth = Depth;
-			if (Depth > SelectAllDepth)
-			{
-				// We are deeper than the depth whose AABB was first detected to be contained in the RectangleXY,
-				// descend and collect all leaf triangles
-				return true;
-			}
-			
-			SelectAllDepth = TNumericLimits<int>::Max();
-			
-			const FAxisAlignedBox2d BoxXY(FVector2d(Box.Min.X, Box.Min.Y), FVector2d(Box.Max.X, Box.Max.Y));
-			if (GridRectXY.Intersects(BoxXY))
-			{
-				if (GridRectXY.Contains(BoxXY))
-				{
-					SelectAllDepth = Depth;
-				}
-				
-				return true;		
-			}
-			return false;
-		};
-		
-
-		Traversal.NextTriangleF =
-			[&GridRectXY, &SelectAllDepth, &CurrentDepth, &Triangles, &EditMesh]
-			(int TriangleID)
-		{
-			if (CurrentDepth >= SelectAllDepth)
-			{
-				// This TriangleID is entirely contained in the selection rectangle so we can skip intersection testing
-				Triangles.Add(TriangleID);
-				return;
-			}
-
-			FAxisAlignedBox2d ExpandedGridRect(GridRectXY);
-			// Expand the rect by 1 unit to add an floating point error margin. 
-			ExpandedGridRect.Expand(1);
-		
-			// this triangle is either on the edge or intersecting with the edge but since we did the plane cuts we know it can't be crossing over into the next cell
-			FIndex3i Triangle = EditMesh.GetTriangle(TriangleID);
-			FVector3d VertexA;
-			FVector3d VertexB;
-			FVector3d VertexC;
-			EditMesh.GetTriVertices(TriangleID, VertexA, VertexB, VertexC);
-			if ((ExpandedGridRect.Contains(FVector2d(VertexA))) &&
-				(ExpandedGridRect.Contains(FVector2d(VertexB))) &&
-				(ExpandedGridRect.Contains(FVector2d(VertexC))))
-			{
-				Triangles.Add(TriangleID);
-			}
-		};
-			
-
-		AABBTree.DoTraversal(Traversal);
-
-		if (Triangles.Num() == 0)
-		{
-			return;
-		}
-
-		TArray<FDynamicMeshVertex> Vertices;
-		TArray<uint32> Indices;
-
-		bool bHasBoundaryTriangle = false;
-		if (GetWaterBodyType() != EWaterBodyType::River)
-		{
-			for (int32 TriangleID : Triangles)
-			{
-				if (EditMesh.IsBoundaryTriangle(TriangleID))
-				{
-					bHasBoundaryTriangle = true;
-					break;
-				}
-			}
-		}
-
-		// Simplification can occur for lakes/rivers since they have flat geometry and no vertex color data.
-		if (!bHasBoundaryTriangle && (GetWaterBodyType() != EWaterBodyType::River))
-		{
-			const float MeshZ = GetConstantSurfaceZ();
-			uint32 TriIndices[4];
-
-			auto AddQuadVertex = [&](FVector2D Position, int QuadIndex)
-			{
-				// Transform back into local space when rebuilding the mesh.
-				const FVector3f LocalSpacePosition = FVector3f(GetComponentTransform().InverseTransformPosition(FVector(Position, MeshZ)));
-				FDynamicMeshVertex MeshVertex(LocalSpacePosition);
-				MeshVertex.Color = FColor(0);
-				TriIndices[QuadIndex] = Vertices.Add(MeshVertex);
-			};
-
-			AddQuadVertex(GridRectXY.Min, 0);
-			AddQuadVertex(GridRectXY.Max, 1);
-			AddQuadVertex(FVector2D(GridRectXY.Min.X, GridRectXY.Max.Y), 2);
-			AddQuadVertex(FVector2D(GridRectXY.Max.X, GridRectXY.Min.Y), 3);
-
-			Indices.Append({ TriIndices[0], TriIndices[2], TriIndices[1], TriIndices[0], TriIndices[1], TriIndices[3] });
-		}
-		else
-		{
-			TMap<uint32, uint32> VertexMap;
-			for (int32 TriangleID : Triangles)
-			{
-				FIndex3i Triangle = EditMesh.GetTriangle(TriangleID);
-				
-				TArray<uint32, TInlineAllocator<3>> TriIndices;
-				for (int32 TriIndex = 0; TriIndex < 3; ++TriIndex)
-				{
-					if (uint32* VertexIndexPtr = VertexMap.Find(Triangle[TriIndex]))
-					{
-						TriIndices.Add(*VertexIndexPtr);
-					}
-					else
-					{
-						FVertexInfo VertexInfo;
-						const bool bWantColors = true;
-						const bool bWantNormals = false;
-						const bool bWantUVs = false;
-						EditMesh.GetVertex(Triangle[TriIndex], VertexInfo, bWantNormals, bWantColors, bWantUVs);
-
-						// Transform back into local space when rebuilding the mesh.
-						const FVector3f LocalSpacePosition = FVector3f(GetComponentTransform().InverseTransformPosition(VertexInfo.Position));
-						FDynamicMeshVertex MeshVertex(LocalSpacePosition);
-						MeshVertex.Color = PackFlowData(VertexInfo.Color.X, VertexInfo.Color.Y, MaxFlowVelocity);
-
-						uint32 NewIndex = Vertices.Add(MeshVertex);
-						VertexMap.Add(Triangle[TriIndex], NewIndex);
-						TriIndices.Add(NewIndex);
-					}
-				}
+		MarkRenderStateDirty();
+	}
+}
 
-				Indices.Append(TriIndices);
-			}
-		}
+void UWaterBodyComponent::UpdateWaterInfoMeshComponents()
+{
+	TRACE_CPUPROFILER_EVENT_SCOPE(UWaterBodyComponent::UpdateWaterInfoMeshComponents);
 
-		if (Vertices.Num() != 0)
-		{
-			Sections[SectionIndex].Vertices = MoveTemp(Vertices);
-			Sections[SectionIndex].Indices = MoveTemp(Indices);
-			Sections[SectionIndex].SectionBounds = FBox2D(FVector2D(SectionPos), FVector2D(SectionPos + WaterLODSectionSize));
-		}
-	});
+	// Not needed for CDOs or template actors
+	if (IsTemplate() || (GetWaterBodyActor() && GetWaterBodyActor()->bIsEditorPreviewActor))
+	{
+		return;
+	}
 
-	for (FWaterBodyMeshSection& Section : Sections)
+	if (!AffectsWaterInfo())
 	{
-		if (Section.IsValid())
-		{
-			WaterBodyMeshSections.Emplace(MoveTemp(Section));
-		}
+		return;
 	}
-}
 
+	FWaterBodyMeshBuilder MeshBuilder;
+	MeshBuilder.BuildWaterInfoMeshes( this, GetWaterInfoMeshComponent(), GetDilatedWaterInfoMeshComponent());
+}
 
-void UWaterBodyComponent::UpdateWaterBodyRenderData()
+void UWaterBodyComponent::UpdateWaterBodyStaticMeshComponents()
 {
-	const bool bAffectsWaterInfo = AffectsWaterInfo();
-	const bool bAffectsWaterMesh = AffectsWaterMesh();
+	TRACE_CPUPROFILER_EVENT_SCOPE(UWaterBodyComponent::UpdateWaterBodyStaticMeshComponents);
 
-	if (bAffectsWaterInfo)
+	// Not needed for CDOs or template actors
+	if (IsTemplate() || (GetOwner() && GetOwner()->bIsEditorPreviewActor))
 	{
-		RebuildWaterBodyInfoMesh();
-
-		if (AWaterZone* WaterZone = GetWaterZone())
-		{
-			WaterZone->MarkForRebuild(EWaterZoneRebuildFlags::UpdateWaterInfoTexture);
-		}
+		return;
 	}
 
-	if (bAffectsWaterMesh)
+	AWaterBody* WaterBodyActor = GetWaterBodyActor();
+	if (WaterBodyActor == nullptr)
 	{
-		checkf(bAffectsWaterInfo, TEXT("Invalid water body! WaterBody affects water mesh but doesn't affect water info"));
-
-		UpdateNonTessellatedMeshSections();
+		return;
 	}
 
-	if (bAffectsWaterInfo || bAffectsWaterMesh)
-	{
-		MarkRenderStateDirty();
-	}
-}
+	const TArray<TObjectPtr<UWaterBodyStaticMeshComponent>>& WaterBodyStaticMeshComponents = WaterBodyActor->GetWaterBodyStaticMeshComponents();
 
-void UWaterBodyComponent::UpdateNonTessellatedMeshSections()
-{
-	if (AWaterZone* WaterZone = GetWaterZone(); WaterZone && WaterZone->IsNonTessellatedLODMeshEnabled())
+	const bool bShouldUseStaticMesh = StaticMeshSettings.bEnableWaterBodyStaticMesh && (GetWaterMeshOverride() == nullptr) && (GetWaterBodyType() != EWaterBodyType::Transition);
+
+	if (bShouldUseStaticMesh)
 	{
-		RebuildWaterBodyLODSections();
+		FWaterBodyMeshBuilder LODMeshBuilder;
+		TArray<TObjectPtr<UWaterBodyStaticMeshComponent>> NewStaticMeshComponents = LODMeshBuilder.BuildWaterBodyStaticMesh(this, StaticMeshSettings, WaterBodyStaticMeshComponents);
+		
+		const int32 RemoveStartIndex = NewStaticMeshComponents.Num();
+		const int32 NumToRemove = FMath::Max(WaterBodyStaticMeshComponents.Num() - NewStaticMeshComponents.Num(), 0);
+		WaterBodyActor->SetWaterBodyStaticMeshComponents(NewStaticMeshComponents, TConstArrayView<TObjectPtr<UWaterBodyStaticMeshComponent>>(WaterBodyStaticMeshComponents.GetData() + RemoveStartIndex, NumToRemove));
 	}
 	else
 	{
-		// Clear the mesh sections if LOD is disabled
-		WaterBodyMeshSections.Empty();
+		WaterBodyActor->SetWaterBodyStaticMeshComponents({}, WaterBodyStaticMeshComponents);
 	}
-
-	MarkRenderStateDirty();
 }
 
-#if WITH_EDITOR
-
 void UWaterBodyComponent::CreateWaterSpriteComponent()
 {
 	UTexture2D* Texture = nullptr;
@@ -2264,47 +2020,7 @@ FMeshDescription UWaterBodyComponent::GetHLODMeshDescription() const
 		return MeshDescription;
 	}
 	
-	FStaticMeshAttributes StaticMeshAttributes(MeshDescription);
-	StaticMeshAttributes.Register();
-
-	TVertexAttributesRef<FVector3f> VertexPositions = StaticMeshAttributes.GetVertexPositions();
-
-	const int32 NumVertices = WaterBodyMeshVertices.Num();
-	const int32 NumTriangles = WaterBodyMeshIndices.Num() / 3;
-
-	MeshDescription.ReserveNewVertices(NumVertices);
-	MeshDescription.ReserveNewVertexInstances(NumVertices);
-	MeshDescription.ReserveNewTriangles(NumTriangles);
-
-	FPolygonGroupID PolygonGroupID = MeshDescription.CreatePolygonGroup();
-
-	// Positions
-	for (int32 VertexIndex = 0; VertexIndex < NumVertices; ++VertexIndex)
-	{
-		FVertexID VertexID = MeshDescription.CreateVertex();
-		VertexPositions[VertexID] = WaterBodyMeshVertices[VertexIndex].Position;
-	}
-
-	// Triangles
-	for (int32 TriangleIndex = 0; TriangleIndex < NumTriangles; ++TriangleIndex)
-	{
-		TStaticArray<FVertexInstanceID, 3> VertexInstanceIDs;
-
-		for (int32 Corner = 0; Corner < 3; ++Corner)
-		{
-			uint32 VertexIndex = WaterBodyMeshIndices[TriangleIndex * 3 + Corner];
-
-			FVertexID VertexID(VertexIndex);
-			FVertexInstanceID VertexInstanceID = MeshDescription.CreateVertexInstance(VertexID);
-			
-			VertexInstanceIDs[Corner] = VertexInstanceID;
-		}
-
-		// Create a triangle
-		MeshDescription.CreateTriangle(PolygonGroupID, VertexInstanceIDs);
-	}
-	
-	return MeshDescription;
+	return FWaterBodyMeshBuilder().BuildMeshDescription(this);
 }
 
 UMaterialInterface* UWaterBodyComponent::GetHLODMaterial() const
@@ -2320,5 +2036,3 @@ void UWaterBodyComponent::SetHLODMaterial(UMaterialInterface* InMaterial)
 #endif // WITH_EDITOR
 
 #undef LOCTEXT_NAMESPACE
-
-
diff --git a/Engine/Plugins/Experimental/Water/Source/Runtime/Private/WaterBodyCustomComponent.cpp b/Engine/Plugins/Experimental/Water/Source/Runtime/Private/WaterBodyCustomComponent.cpp
--- a/Engine/Plugins/Experimental/Water/Source/Runtime/Private/WaterBodyCustomComponent.cpp
+++ b/Engine/Plugins/Experimental/Water/Source/Runtime/Private/WaterBodyCustomComponent.cpp
@@ -134,7 +134,7 @@ void UWaterBodyCustomComponent::BeginUpdateWaterBody()
 }
 
 #if WITH_EDITOR
-TArray<TSharedRef<FTokenizedMessage>> UWaterBodyCustomComponent::CheckWaterBodyStatus() const
+TArray<TSharedRef<FTokenizedMessage>> UWaterBodyCustomComponent::CheckWaterBodyStatus() 
 {
 	TArray<TSharedRef<FTokenizedMessage>> StatusMessages = Super::CheckWaterBodyStatus();
 
diff --git a/Engine/Plugins/Experimental/Water/Source/Runtime/Private/WaterBodyInfoMeshComponent.cpp b/Engine/Plugins/Experimental/Water/Source/Runtime/Private/WaterBodyInfoMeshComponent.cpp
new file mode 100644
--- /dev/null
+++ b/Engine/Plugins/Experimental/Water/Source/Runtime/Private/WaterBodyInfoMeshComponent.cpp
@@ -0,0 +1,31 @@
+// Copyright Epic Games, Inc. All Rights Reserved.
+
+#pragma once
+
+#include "WaterBodyInfoMeshComponent.h"
+#include "StaticMeshSceneProxy.h"
+
+#include UE_INLINE_GENERATED_CPP_BY_NAME(WaterBodyInfoMeshComponent)
+
+
+UWaterBodyInfoMeshComponent::UWaterBodyInfoMeshComponent(const FObjectInitializer& ObjectInitializer)
+	: Super(ObjectInitializer)
+{
+	bAffectDistanceFieldLighting = false;
+	bSelectable = false;
+}
+
+FPrimitiveSceneProxy* UWaterBodyInfoMeshComponent::CreateSceneProxy()
+{
+	if (!CanCreateSceneProxy())
+	{
+		return nullptr;
+	}
+	return new FWaterBodyInfoMeshSceneProxy(this);
+}
+
+FWaterBodyInfoMeshSceneProxy::FWaterBodyInfoMeshSceneProxy(UWaterBodyInfoMeshComponent* Component)
+	: FStaticMeshSceneProxy(Component, true)
+{
+	SetEnabled(false);
+}
diff --git a/Engine/Plugins/Experimental/Water/Source/Runtime/Private/WaterBodyInfoMeshComponent.h b/Engine/Plugins/Experimental/Water/Source/Runtime/Private/WaterBodyInfoMeshComponent.h
new file mode 100644
--- /dev/null
+++ b/Engine/Plugins/Experimental/Water/Source/Runtime/Private/WaterBodyInfoMeshComponent.h
@@ -0,0 +1,39 @@
+// Copyright Epic Games, Inc. All Rights Reserved.
+
+#pragma once
+
+#include "WaterBodyMeshComponent.h"
+#include "StaticMeshSceneProxy.h"
+
+#include "WaterBodyInfoMeshComponent.generated.h"
+
+class UWaterBodyComponent;
+
+/**
+ * WaterBodyMeshComponent used to render the water body into the water info texture.
+ * Utilizes a custom scene proxy to allow hiding the mesh outside of all other passes besides the water info passes.
+ */
+UCLASS(MinimalAPI)
+class UWaterBodyInfoMeshComponent : public UWaterBodyMeshComponent
+{
+	GENERATED_BODY()
+public:
+	UWaterBodyInfoMeshComponent(const FObjectInitializer& ObjectInitializer);
+
+	virtual FPrimitiveSceneProxy* CreateSceneProxy() override;
+};
+
+struct FWaterBodyInfoMeshSceneProxy : public FStaticMeshSceneProxy
+{
+public:
+	FWaterBodyInfoMeshSceneProxy(UWaterBodyInfoMeshComponent* Component);
+
+	SIZE_T GetTypeHash() const
+	{
+		static size_t UniquePointer;
+		return reinterpret_cast<size_t>(&UniquePointer);
+	}
+
+	void SetEnabled(bool bInEnabled) { SetForceHidden(!bInEnabled); }
+};
+
diff --git a/Engine/Plugins/Experimental/Water/Source/Runtime/Private/WaterBodyLakeComponent.cpp b/Engine/Plugins/Experimental/Water/Source/Runtime/Private/WaterBodyLakeComponent.cpp
--- a/Engine/Plugins/Experimental/Water/Source/Runtime/Private/WaterBodyLakeComponent.cpp
+++ b/Engine/Plugins/Experimental/Water/Source/Runtime/Private/WaterBodyLakeComponent.cpp
@@ -11,13 +11,11 @@
 #include "Materials/MaterialInstanceDynamic.h"
 #include "Polygon2.h"
 #include "ConstrainedDelaunay2.h"
-#include "Operations/InsetMeshRegion.h"
+#include "DynamicMesh/DynamicMeshAttributeSet.h"
+#include "Curve/PolygonOffsetUtils.h"
 
 #include UE_INLINE_GENERATED_CPP_BY_NAME(WaterBodyLakeComponent)
 
-#if WITH_EDITOR
-#endif
-
 // ----------------------------------------------------------------------------------
 
 UWaterBodyLakeComponent::UWaterBodyLakeComponent(const FObjectInitializer& ObjectInitializer)
@@ -41,7 +39,7 @@ TArray<UPrimitiveComponent*> UWaterBodyLakeComponent::GetCollisionComponents(boo
 
 TArray<UPrimitiveComponent*> UWaterBodyLakeComponent::GetStandardRenderableComponents() const 
 {
-	TArray<UPrimitiveComponent*> Result;
+	TArray<UPrimitiveComponent*> Result = Super::GetStandardRenderableComponents();
 	if (LakeMeshComp != nullptr)
 	{
 		Result.Add(LakeMeshComp);
@@ -86,38 +84,79 @@ bool UWaterBodyLakeComponent::GenerateWaterBodyMesh(UE::Geometry::FDynamicMesh3&
 	}
 
 
+	check(OutMesh.Attributes());
+	FDynamicMeshColorOverlay* ColorOverlay = OutMesh.Attributes()->PrimaryColors();
+	FDynamicMeshNormalOverlay* NormalOverlay = OutMesh.Attributes()->PrimaryNormals();
+
 	for (const FVector2d& Vertex : Triangulation.Vertices)
 	{
 		FVertexInfo MeshVertex(FVector3d(Vertex.X, Vertex.Y, 0.f));
-		MeshVertex.Color = FVector3f(0.0);
-		MeshVertex.bHaveC = true;
 
 		OutMesh.AppendVertex(MeshVertex);
+		ColorOverlay->AppendElement(FVector4f(0.f));
+		NormalOverlay->AppendElement(FVector3f(0., 0., 1.));
 	}
 
 	for (const FIndex3i& Triangle : Triangulation.Triangles)
 	{
-		OutMesh.AppendTriangle(Triangle);
+		const int TriangleID = OutMesh.AppendTriangle(Triangle);
+		ColorOverlay->SetTriangle(TriangleID, Triangle);
+		NormalOverlay->SetTriangle(TriangleID, Triangle);
 	}
 
 	if (ShapeDilation > 0.f && OutDilatedMesh)
 	{
-		// Inset the mesh by -ShapeDilation to effectively expand the mesh
-		OutDilatedMesh->Copy(OutMesh);
-		FInsetMeshRegion Inset(OutDilatedMesh);
-		Inset.InsetDistance = -1 * ShapeDilation / 2.f;
+		TArray<FGeneralPolygon2d> OffsetPolys;
+		UE::Geometry::PolygonsOffset(
+			ShapeDilation / 2.f,
+			{ LakePoly },
+			OffsetPolys,
+			false,
+			2.0);
+
 
-		Inset.Triangles.Reserve(OutDilatedMesh->TriangleCount());
-		for (int32 Idx : OutDilatedMesh->TriangleIndicesItr())
+		FConstrainedDelaunay2d DilationTriangulation;
+		DilationTriangulation.FillRule = FConstrainedDelaunay2d::EFillRule::Positive;
+
+		for (const FGeneralPolygon2d& Poly : OffsetPolys)
 		{
-			Inset.Triangles.Add(Idx);
+			if (Poly.SignedArea() <= 0.)
+			{
+				UE_LOG(LogWater, Warning, TEXT("Failed to apply offset for shape dilation (%s"), *GetOwner()->GetActorNameOrLabel());
+				continue;
+			}
+
+			DilationTriangulation.Add(Poly);
 		}
-		
-		if (!Inset.Apply())
+
+		if (!DilationTriangulation.Triangulate())
+		{
+			UE_LOG(LogWater, Warning, TEXT("Failed to triangulate dilated lake mesh (%s"), *GetOwner()->GetActorNameOrLabel());
+			return false;
+		}
+
+		if (DilationTriangulation.Triangles.Num() == 0)
 		{
-			UE_LOG(LogWater, Warning, TEXT("Failed to apply mesh inset for shape dilation (%s"), *GetOwner()->GetActorNameOrLabel());
 			return false;
 		}
+
+		FDynamicMeshColorOverlay* DilatedColorOverlay = OutDilatedMesh->Attributes()->PrimaryColors();
+		FDynamicMeshNormalOverlay* DilatedNormalOverlay = OutDilatedMesh->Attributes()->PrimaryNormals();
+		for (const FVector2d& Vertex : DilationTriangulation.Vertices)
+		{
+			FVertexInfo MeshVertex(FVector3d(Vertex.X, Vertex.Y, 0.));
+
+			OutDilatedMesh->AppendVertex(MeshVertex);
+			DilatedColorOverlay->AppendElement(FVector4f(0.f));
+			DilatedNormalOverlay->AppendElement(FVector3f(0., 0., 1.));
+		}
+
+		for (const FIndex3i& Triangle : DilationTriangulation.Triangles)
+		{
+			const int TriangleID = OutDilatedMesh->AppendTriangle(Triangle);
+			DilatedColorOverlay->SetTriangle(TriangleID, Triangle);
+			DilatedNormalOverlay->SetTriangle(TriangleID, Triangle);
+		}
 	}
 
 	return true;
diff --git a/Engine/Plugins/Experimental/Water/Source/Runtime/Private/WaterBodyMeshBuilder.cpp b/Engine/Plugins/Experimental/Water/Source/Runtime/Private/WaterBodyMeshBuilder.cpp
new file mode 100644
--- /dev/null
+++ b/Engine/Plugins/Experimental/Water/Source/Runtime/Private/WaterBodyMeshBuilder.cpp
@@ -0,0 +1,573 @@
+// Copyright Epic Games, Inc. All Rights Reserved.
+
+#include "WaterBodyMeshBuilder.h"
+
+#include "WaterBodyComponent.h"
+#include "WaterBodyStaticMeshComponent.h"
+#include "WaterBodyInfoMeshComponent.h"
+
+#if WITH_EDITOR
+#include "Materials/MaterialInstanceDynamic.h"
+#include "WaterUtils.h"
+#include "DynamicMesh/DynamicMeshAABBTree3.h"
+#include "Operations/MeshPlaneCut.h"
+#include "Engine/StaticMesh.h"
+#include "DynamicMesh/DynamicMesh3.h"
+#include "DynamicMeshToMeshDescription.h"
+#include "StaticMeshAttributes.h"
+#include "DynamicMesh/DynamicMesh3.h"
+#include "PrimitiveSceneProxy.h"
+#include "StaticMeshResources.h"
+#include "DynamicMeshBuilder.h"
+#include "MeshDescriptionBuilder.h"
+#include "DynamicMesh/DynamicMeshAttributeSet.h"
+#include "DynamicMesh/MeshNormals.h"
+#include "DynamicMesh/MeshTangents.h"
+#include "Algo/RemoveIf.h"
+#include "PhysicsEngine/BodySetup.h"
+#include "Engine/CollisionProfile.h"
+#include "UObject/Package.h"
+#include "StaticMeshCompiler.h"
+#include "Algo/AnyOf.h"
+#endif // WITH_EDITOR
+
+#if WITH_EDITOR
+void FWaterBodyMeshBuilder::BuildWaterInfoMeshes(UWaterBodyComponent* WaterBodyComponent, UWaterBodyInfoMeshComponent* WaterInfoMeshComponent, UWaterBodyInfoMeshComponent* WaterInfoDilatedMeshComponent) const
+{
+	TRACE_CPUPROFILER_EVENT_SCOPE(FWaterBodyMeshBuilder::BuildWaterInfoMesh);
+	
+	using namespace UE::Geometry;
+
+	FDynamicMesh3 WaterBodyMesh(EMeshComponents::None);
+	FDynamicMesh3 WaterBodyDilatedMesh(EMeshComponents::None);
+	GetDynamicMesh(WaterBodyComponent, WaterBodyMesh, &WaterBodyDilatedMesh);
+
+	UMaterialInterface* WaterInfoMID = WaterBodyComponent->GetWaterInfoMaterialInstance();
+	UObject* Outer = WaterBodyComponent->GetPackage();
+
+	auto BuildWaterInfoMesh = [this, Outer, WaterInfoMID](const FDynamicMesh3& DynamicMesh, UWaterBodyInfoMeshComponent* MeshComponent, FName BaseName) -> UStaticMesh*
+	{
+		FName WaterInfoMeshName = MakeUniqueObjectName(Outer, UStaticMesh::StaticClass(), BaseName);
+		UStaticMesh* StaticMesh = CreateUStaticMesh(Outer, WaterInfoMeshName);
+
+		UpdateStaticMesh(StaticMesh, ConvertDynamicMeshToMeshDescription(DynamicMesh));
+
+		StaticMesh->AddMaterial(WaterInfoMID);
+
+		MeshComponent->SetStaticMesh(StaticMesh);
+
+		return StaticMesh;
+	};
+
+	TArray<UStaticMesh*> StaticMeshes = {
+		BuildWaterInfoMesh(WaterBodyMesh, WaterInfoMeshComponent, TEXT("WaterInfoMesh")),
+		BuildWaterInfoMesh(WaterBodyDilatedMesh, WaterInfoDilatedMeshComponent, TEXT("WaterInfoDilatedMesh")),
+	};
+
+	UStaticMesh::BatchBuild(StaticMeshes);
+}
+
+FMeshDescription FWaterBodyMeshBuilder::BuildMeshDescription(const UWaterBodyComponent* WaterBodyComponent) const
+{
+	using namespace UE::Geometry;
+
+	FDynamicMesh3 WaterBodyMesh(EMeshComponents::None);
+	GetDynamicMesh(WaterBodyComponent, WaterBodyMesh);
+
+	return ConvertDynamicMeshToMeshDescription(WaterBodyMesh);
+}
+
+TArray<TObjectPtr<UWaterBodyStaticMeshComponent>> FWaterBodyMeshBuilder::BuildWaterBodyStaticMesh(UWaterBodyComponent* WaterBodyComponent, const FWaterBodyStaticMeshSettings& Settings, const TArray<TObjectPtr<UWaterBodyStaticMeshComponent>>& ReusableComponents) const
+{
+	TRACE_CPUPROFILER_EVENT_SCOPE(FWaterBodyMeshBuilder::BuildWaterBodyStaticMesh);
+	
+	using namespace UE::Geometry;
+
+	FDynamicMesh3 EditMesh(EMeshComponents::None);
+
+	GetDynamicMesh(WaterBodyComponent, EditMesh);
+
+	const FBox MeshBounds3d = FBox(EditMesh.GetBounds());
+	const FAxisAlignedBox2d MeshBounds2d(FVector2d(MeshBounds3d.Min), FVector2d(MeshBounds3d.Max));
+	const double WaterMeshSectionSize = Settings.SectionSize;
+	const EWaterBodyType WaterBodyType = WaterBodyComponent->GetWaterBodyType();
+	const double ConstantSurfaceZ = WaterBodyComponent->GetConstantSurfaceZ();
+
+	/** Mesh descriptions for sections which cannot be simplified to quads */
+	TArray<FMeshDescription> SectionMeshDescriptions;
+	std::atomic<uint32> NumSectionMeshDescriptions = {0};
+
+	/** Positions of sections which can be simplified as quads. */
+	TArray<FVector> QuadSectionPositions;
+	std::atomic<uint32> NumQuadSectionPositions = {0};
+
+	TArray<TObjectPtr<UWaterBodyStaticMeshComponent>> StaticMeshComponents;
+
+	// Skip trying to section the mesh if the mesh bounds are smaller than size of a section.
+	if (Settings.bSectionWaterBodyStaticMesh && (MeshBounds2d.Area() > WaterMeshSectionSize * WaterMeshSectionSize))
+	{
+		/**
+		* Sectioning algorithm outline:
+		* 
+		* 1. Cut source mesh along X and Y axes at intervals of the section size. This prepares the geometry to then be sorted into buckets for each cell
+		* 2. Sort the mesh into an AABB tree to be able to quickly test if vertices belong within a section's cell.
+		* 3. In parallel for each section within the bounds of the mesh:
+		*		a. Traverse the AABB tree to collect any triangles that are contained within or intersect with the bounds of that section.
+		*		b. Identify if any triangles within the section are boundary triangles (they lie on the edge of the mesh).
+		*		c. If this is the case, the geometry of the cell must be preserved since it defines the edge of our 2d polygon.
+		*		d. If this is not the case, a simple quad mesh can be substituted which completely fills the section.
+		*		e. A MeshDescription is then created for either the simplified quad or the triangles within the cell.
+		* 4. After all the section meshes are created, create or reuse UWaterBodyStaticMeshComponents to render the individual meshes separately.
+		*/
+
+		// Round up when considering the regions to slice
+		const uint32 NumberOfMeshSectionsX = FMath::DivideAndRoundUp(MeshBounds2d.Extents().X * 2., WaterMeshSectionSize);
+		const uint32 NumberOfMeshSectionsY = FMath::DivideAndRoundUp(MeshBounds2d.Extents().Y * 2., WaterMeshSectionSize);
+
+		{
+			TRACE_CPUPROFILER_EVENT_SCOPE(BuildWaterBodyStaticMesh::CutMesh);
+
+			for (uint32 SectionIndexX = 0; SectionIndexX < NumberOfMeshSectionsX; ++SectionIndexX)
+			{
+				const FVector3d SlicePos = FVector3d(MeshBounds2d.Min + SectionIndexX * FVector2d(WaterMeshSectionSize, 0.0), 0.0);
+				const FVector3d SliceNormal(-1.0, 0.0, 0.0);
+
+				FMeshPlaneCut Cut(&EditMesh, SlicePos, SliceNormal);
+				ensure(Cut.CutWithoutDelete(false));
+			}
+
+			for (uint32 SectionIndexY = 0; SectionIndexY < NumberOfMeshSectionsY; ++SectionIndexY)
+			{
+				const FVector3d SlicePos = FVector3d(MeshBounds2d.Min + SectionIndexY * FVector2d(0.0, WaterMeshSectionSize), 0.0);
+				const FVector3d SliceNormal(0.0, -1.0, 0.0);
+
+				FMeshPlaneCut Cut(&EditMesh, SlicePos, SliceNormal);
+				ensure(Cut.CutWithoutDelete(false));
+			}
+		}
+
+		FDynamicMeshAABBTree3 AABBTree(&EditMesh, true);
+
+		const FDynamicMeshColorOverlay* ColorOverlay = EditMesh.Attributes()->PrimaryColors();
+
+		const uint32 TotalNumberOfSections = NumberOfMeshSectionsX * NumberOfMeshSectionsY;
+		SectionMeshDescriptions.SetNum(TotalNumberOfSections);
+		QuadSectionPositions.SetNum(TotalNumberOfSections);
+
+		ParallelFor(TotalNumberOfSections, [&](int32 SectionIndex)
+		{
+			TRACE_CPUPROFILER_EVENT_SCOPE(BuildWaterBodyStaticMesh::BuildMeshSection);
+			const int32 SectionIndexX = SectionIndex % NumberOfMeshSectionsX;
+			const int32 SectionIndexY = SectionIndex / NumberOfMeshSectionsY;
+
+			const FVector2d SectionPos = MeshBounds2d.Min + FVector2d(SectionIndexX, SectionIndexY) * WaterMeshSectionSize;
+
+			const FAxisAlignedBox2d GridRectXY(SectionPos, SectionPos + WaterMeshSectionSize);
+
+			FDynamicMeshAABBTree3::FTreeTraversal Traversal;
+			
+			int SelectAllDepth = TNumericLimits<int>::Max();
+			int CurrentDepth = -1;
+			
+			TArray<int> Triangles;
+
+			bool bHasBoundaryTriangle = false;
+
+			{
+				TRACE_CPUPROFILER_EVENT_SCOPE(BuildWaterBodyStaticMesh::Traversal);
+
+				Traversal.NextBoxF =
+					[&GridRectXY, &SelectAllDepth, &CurrentDepth](const FAxisAlignedBox3d& Box, int Depth)
+				{
+					CurrentDepth = Depth;
+					if (Depth > SelectAllDepth)
+					{
+						// We are deeper than the depth whose AABB was first detected to be contained in the RectangleXY,
+						// descend and collect all leaf triangles
+						return true;
+					}
+					
+					SelectAllDepth = TNumericLimits<int>::Max();
+					
+					const FAxisAlignedBox2d BoxXY(FVector2d(Box.Min.X, Box.Min.Y), FVector2d(Box.Max.X, Box.Max.Y));
+					if (GridRectXY.Intersects(BoxXY))
+					{
+						if (GridRectXY.Contains(BoxXY))
+						{
+							SelectAllDepth = Depth;
+						}
+						
+						return true;		
+					}
+					return false;
+				};
+				
+				Traversal.NextTriangleF =
+					[&GridRectXY, &SelectAllDepth, &CurrentDepth, &Triangles, &EditMesh, &bHasBoundaryTriangle]
+					(int TriangleID)
+				{
+					if (CurrentDepth >= SelectAllDepth)
+					{
+						// This TriangleID is entirely contained in the selection rectangle so we can skip intersection testing
+						Triangles.Add(TriangleID);
+						bHasBoundaryTriangle |= EditMesh.IsBoundaryTriangle(TriangleID);
+						return;
+					}
+
+					FAxisAlignedBox2d ExpandedGridRect(GridRectXY);
+					// Expand the rect by 1 unit to add a floating point error margin. 
+					ExpandedGridRect.Expand(1);
+				
+					// this triangle is either on the edge or intersecting with the edge but since we did the plane cuts we know it can't be crossing over into the next cell
+					FIndex3i Triangle = EditMesh.GetTriangle(TriangleID);
+					FVector3d VertexA;
+					FVector3d VertexB;
+					FVector3d VertexC;
+					EditMesh.GetTriVertices(TriangleID, VertexA, VertexB, VertexC);
+					if ((ExpandedGridRect.Contains(FVector2d(VertexA))) &&
+						(ExpandedGridRect.Contains(FVector2d(VertexB))) &&
+						(ExpandedGridRect.Contains(FVector2d(VertexC))))
+					{
+						Triangles.Add(TriangleID);
+						bHasBoundaryTriangle |= EditMesh.IsBoundaryTriangle(TriangleID);
+					}
+				};
+					
+				AABBTree.DoTraversal(Traversal);
+			}
+
+			TArray<FVector3d> VertexPositions;
+			TArray<FVector4f> VertexColors;
+			TArray<FVector> VertexNormals;
+			TArray<FVector> VertexTangents;
+			TArray<uint32> Indices;
+
+			// Simplification can occur for lakes/rivers since they have flat geometry and no vertex color data.
+			//if (!bHasBoundaryTriangle && (WaterBodyType != EWaterBodyType::River) && Triangles.Num() > 0)
+			//{
+			//	const float MeshZ = ConstantSurfaceZ;
+			//	const uint32 QuadSectionIndex = NumQuadSectionPositions.fetch_add(1);
+			//	QuadSectionPositions[QuadSectionIndex] = FVector(SectionPos.X, SectionPos.Y, ConstantSurfaceZ);
+			//	return;
+			//}
+
+			TMap<uint32, uint32> VertexMap;
+			for (int32 TriangleID : Triangles)
+			{
+				FIndex3i Triangle = EditMesh.GetTriangle(TriangleID);
+				
+				TArray<uint32, TInlineAllocator<3>> TriIndices;
+				for (int32 TriIndex = 0; TriIndex < 3; ++TriIndex)
+				{
+					// Vertex data for this index already exists
+					if (uint32* VertexIndexPtr = VertexMap.Find(Triangle[TriIndex]))
+					{
+						TriIndices.Add(*VertexIndexPtr);
+					}
+					// Vertex data for this index doesn't exist. Add it now:
+					else
+					{
+						const int32 VertexID = Triangle[TriIndex];
+						const FVector3d VertexPos = EditMesh.GetVertex(VertexID);
+
+						// Only rivers have vertex colors for the velocity data.
+						const FVector4f Color = WaterBodyType == EWaterBodyType::River ? ColorOverlay->GetElement(VertexID) : FVector4f(0);
+
+						const uint32 NewIndex = VertexPositions.Add(VertexPos);
+						VertexColors.Add(Color);
+						VertexNormals.Add(FVector(0, 0, 1));
+						VertexTangents.Add(FVector(1, 0, 0));
+
+						VertexMap.Add(VertexID, NewIndex);
+
+						TriIndices.Add(NewIndex);
+					}
+				}
+
+				Indices.Append(TriIndices);
+			}
+
+			if (VertexPositions.Num() != 0)
+			{
+				const int32 NumVertices = VertexPositions.Num();
+				check(VertexNormals.Num() == NumVertices);
+				check(VertexColors.Num() == NumVertices);
+
+				const int32 NumTriangles = Indices.Num() / 3;
+
+				FMeshDescription MeshDescription;
+				FStaticMeshAttributes(MeshDescription).Register();
+				MeshDescription.Empty();
+
+				FMeshDescriptionBuilder Builder;
+				Builder.SetMeshDescription(&MeshDescription);
+				Builder.SuspendMeshDescriptionIndexing();
+
+				Builder.EnablePolyGroups();
+
+				Builder.SetNumUVLayers(1);
+				Builder.ReserveNewUVs(NumVertices, 0);
+
+				Builder.ReserveNewVertices(NumVertices);
+				
+				for (int32 VertexIndex = 0; VertexIndex < NumVertices; ++VertexIndex)
+				{
+					const FVertexID VertexID = Builder.AppendVertex(VertexPositions[VertexIndex]);
+					const FVertexInstanceID VertexInstance = Builder.AppendInstance(VertexID);
+
+					Builder.SetInstanceColor(VertexInstance, VertexColors[VertexIndex]);
+					Builder.SetInstanceNormal(VertexInstance, VertexNormals[VertexIndex]);
+					Builder.SetInstanceTangentSpace(VertexInstance, VertexNormals[VertexIndex], VertexTangents[VertexIndex], 1);
+				}
+
+				FPolygonGroupID GroupID = Builder.AppendPolygonGroup();
+
+				for (int32 TriangleIndex = 0; TriangleIndex < NumTriangles; ++TriangleIndex)
+				{
+					FVertexInstanceID Triangle[3];
+					for (int32 Vertex = 0; Vertex < 3; ++Vertex)
+					{
+						Triangle[Vertex] = Indices[TriangleIndex * 3 + Vertex];
+					}
+
+					Builder.AppendTriangle(Triangle[0], Triangle[1], Triangle[2], GroupID);
+				}
+
+				Builder.ResumeMeshDescriptionIndexing();
+
+				const uint32 SectionMeshDescriptionIndex = NumSectionMeshDescriptions.fetch_add(1);
+				SectionMeshDescriptions[SectionMeshDescriptionIndex] = MoveTemp(MeshDescription);
+			}
+		});
+	}
+	else
+	{
+		SectionMeshDescriptions.EmplaceAt(0, ConvertDynamicMeshToMeshDescription(EditMesh));
+		++NumSectionMeshDescriptions;
+	}
+
+	//SectionMeshDescriptions.SetNum(Algo::RemoveIf(SectionMeshDescriptions, [](const FMeshDescription& MeshDescription) { return MeshDescription.Vertices().Num() == 0; }));
+	SectionMeshDescriptions.SetNum(NumSectionMeshDescriptions);
+	QuadSectionPositions.SetNum(NumQuadSectionPositions);
+
+	// Attempt to reuse any of the components passed in from the ReusableComponents list
+	int32 NumComponentsRequired = SectionMeshDescriptions.Num() + QuadSectionPositions.Num();
+	StaticMeshComponents.Reserve(NumComponentsRequired);
+	for (int32 Index = 0; Index < FMath::Min(NumComponentsRequired, ReusableComponents.Num()); ++Index)
+	{
+		StaticMeshComponents.Add(ReusableComponents[Index]);
+	}
+	// Create any needed new components
+	AActor* ActorOwner = WaterBodyComponent->GetOwner();
+	for (int32 Index = StaticMeshComponents.Num(); Index < NumComponentsRequired; ++Index)
+	{
+		const FName Name = MakeUniqueObjectName(ActorOwner, UWaterBodyStaticMeshComponent::StaticClass());
+		UWaterBodyStaticMeshComponent* WaterBodyStaticMeshComponent = NewObject<UWaterBodyStaticMeshComponent>(ActorOwner, Name);
+
+		WaterBodyStaticMeshComponent->RegisterComponent();
+		WaterBodyStaticMeshComponent->AttachToComponent(WaterBodyComponent, FAttachmentTransformRules::SnapToTargetNotIncludingScale);
+
+		StaticMeshComponents.Add(WaterBodyStaticMeshComponent);
+	}
+
+	check(StaticMeshComponents.Num() == NumComponentsRequired);
+
+	TArray<UStaticMesh*> StaticMeshes;
+	StaticMeshes.Reserve(NumComponentsRequired);
+
+	UMaterialInterface* WaterStaticMeshMID = WaterBodyComponent->GetWaterStaticMeshMaterialInstance();
+
+	UObject* MeshComponentOuter = WaterBodyComponent->GetPackage();
+
+	int32 StaticMeshComponentIndex = 0;
+	{
+		TRACE_CPUPROFILER_EVENT_SCOPE(FWaterBodyMeshBuilder::SetupMeshComponents);
+
+		for (int32 Index = 0; Index < SectionMeshDescriptions.Num(); ++Index)
+		{
+			UWaterBodyStaticMeshComponent* WaterBodyStaticMeshComponent = StaticMeshComponents[StaticMeshComponentIndex];
+			const FMeshDescription& MeshDescription = SectionMeshDescriptions[StaticMeshComponentIndex];
+
+			check(IsValid(WaterBodyStaticMeshComponent));
+
+			FName WaterBodyMeshName = MakeUniqueObjectName(MeshComponentOuter, UStaticMesh::StaticClass(), TEXT("WaterBodyMesh"));
+			UStaticMesh* StaticMesh = CreateUStaticMesh(MeshComponentOuter, WaterBodyMeshName);
+
+			StaticMesh->AddMaterial(WaterStaticMeshMID);
+
+			UpdateStaticMesh(StaticMesh, MeshDescription);
+
+			WaterBodyStaticMeshComponent->SetStaticMesh(StaticMesh);
+
+			StaticMeshes.Add(StaticMesh);
+
+			++StaticMeshComponentIndex;
+		}
+	}
+
+	if (QuadSectionPositions.Num() > 0)
+	{
+		TRACE_CPUPROFILER_EVENT_SCOPE(FWaterBodyMeshBuilder::SetupQuadComponents);
+
+		// Create a shared UStaticMesh for all the quad sections used by this water body
+		FName WaterBodyMeshName = MakeUniqueObjectName(MeshComponentOuter, UStaticMesh::StaticClass(), TEXT("WaterBodyQuadMesh"));
+		UStaticMesh* QuadStaticMesh = CreateUStaticMesh(MeshComponentOuter, WaterBodyMeshName);
+
+		QuadStaticMesh->AddMaterial(WaterStaticMeshMID);
+
+		FMeshDescription QuadMeshDescription;
+		{
+			TArray<FVector> VertexPositions = { {0., 0., 0.}, {1., 0., 0}, {1., 1., 0.}, {0., 1., 0.} };
+			TArray<int32> Indices = { 0, 2, 1, 0, 3, 2 };
+			const int32 NumVertices = VertexPositions.Num();
+			const int32 NumTriangles = Indices.Num() / 3;
+
+			FStaticMeshAttributes(QuadMeshDescription).Register();
+			QuadMeshDescription.Empty();
+
+			FMeshDescriptionBuilder Builder;
+			Builder.SetMeshDescription(&QuadMeshDescription);
+			Builder.SuspendMeshDescriptionIndexing();
+
+			Builder.EnablePolyGroups();
+
+			Builder.SetNumUVLayers(1);
+			Builder.ReserveNewUVs(NumVertices, 0);
+
+			Builder.ReserveNewVertices(NumVertices);
+			
+			for (int32 VertexIndex = 0; VertexIndex < NumVertices; ++VertexIndex)
+			{
+				const FVertexID VertexID = Builder.AppendVertex(VertexPositions[VertexIndex]);
+				const FVertexInstanceID VertexInstance = Builder.AppendInstance(VertexID);
+
+				Builder.SetInstanceColor(VertexInstance, FVector4f(0.));
+				Builder.SetInstanceNormal(VertexInstance, FVector(0., 0., 1.));
+				Builder.SetInstanceTangentSpace(VertexInstance, FVector(0., 0., 1.0), FVector(1.0, .0, .0), 1);
+			}
+
+			FPolygonGroupID GroupID = Builder.AppendPolygonGroup();
+
+			for (int32 TriangleIndex = 0; TriangleIndex < NumTriangles; ++TriangleIndex)
+			{
+				FVertexInstanceID Triangle[3];
+				for (int32 Vertex = 0; Vertex < 3; ++Vertex)
+				{
+					Triangle[Vertex] = Indices[TriangleIndex * 3 + Vertex];
+				}
+
+				Builder.AppendTriangle(Triangle[0], Triangle[1], Triangle[2], GroupID);
+			}
+
+			Builder.ResumeMeshDescriptionIndexing();
+		}
+
+		StaticMeshes.Add(QuadStaticMesh);
+
+		UpdateStaticMesh(QuadStaticMesh, QuadMeshDescription);
+
+		for (const FVector& QuadMeshPosition : QuadSectionPositions)
+		{
+			UWaterBodyStaticMeshComponent* WaterBodyStaticMeshComponent = StaticMeshComponents[StaticMeshComponentIndex];
+
+			WaterBodyStaticMeshComponent->SetStaticMesh(QuadStaticMesh);
+			WaterBodyStaticMeshComponent->SetRelativeLocation(QuadMeshPosition);
+			WaterBodyStaticMeshComponent->SetRelativeScale3D(FVector(WaterMeshSectionSize, WaterMeshSectionSize, 1.0));
+
+			++StaticMeshComponentIndex;
+		}
+	}
+
+	check(SectionMeshDescriptions.Num() + QuadSectionPositions.Num() == StaticMeshComponents.Num());
+
+	UStaticMesh::BatchBuild(StaticMeshes);
+
+	return StaticMeshComponents;
+}
+
+void FWaterBodyMeshBuilder::GetDynamicMesh(const UWaterBodyComponent* WaterBodyComponent, UE::Geometry::FDynamicMesh3& InMesh, UE::Geometry::FDynamicMesh3* InDilatedMesh) const
+{
+	using namespace UE::Geometry;
+
+	for (FDynamicMesh3* Mesh : { &InMesh, InDilatedMesh })
+	{
+		if (Mesh)
+		{
+			Mesh->EnableAttributes();
+			Mesh->Attributes()->EnablePrimaryColors();
+			Mesh->Attributes()->EnableTangents();
+		}
+	}
+
+	WaterBodyComponent->GenerateWaterBodyMesh(InMesh, InDilatedMesh);
+
+	for (FDynamicMesh3* Mesh : { &InMesh, InDilatedMesh })
+	{
+		if (Mesh)
+		{
+			FDynamicMeshNormalOverlay* NormalOverlay = Mesh->Attributes()->PrimaryNormals();
+			FMeshNormals::InitializeOverlayToPerVertexNormals(NormalOverlay);
+			FMeshTangentsf::ComputeDefaultOverlayTangents(*Mesh);
+		}
+	}
+}
+
+FMeshDescription FWaterBodyMeshBuilder::ConvertDynamicMeshToMeshDescription(const UE::Geometry::FDynamicMesh3& Mesh) const
+{
+	FMeshDescription Result;
+	FStaticMeshAttributes StaticMeshAttributes(Result);
+	StaticMeshAttributes.Register();
+
+	FDynamicMeshToMeshDescription DynamicMeshToMeshDescription;
+	DynamicMeshToMeshDescription.ConversionOptions.bUpdateVtxColors = true;
+	// We recompute normals and tangents after generating the mesh so no need to recompute them here
+	DynamicMeshToMeshDescription.ConversionOptions.bUpdateNormals = false;
+	DynamicMeshToMeshDescription.ConversionOptions.bUpdateTangents = false;
+
+	DynamicMeshToMeshDescription.Convert(&Mesh, Result);
+	DynamicMeshToMeshDescription.UpdateUsingConversionOptions(&Mesh, Result);
+
+	return Result;
+}
+
+void FWaterBodyMeshBuilder::UpdateStaticMesh(UStaticMesh* WaterMesh, const FMeshDescription& MeshDescription) const
+{
+	if (!WaterMesh->IsSourceModelValid(0))
+	{
+		WaterMesh->AddSourceModel();
+	}
+
+	FStaticMeshSourceModel& SrcModel = WaterMesh->GetSourceModel(0);
+	SrcModel.BuildSettings.bRecomputeNormals = true;
+	SrcModel.BuildSettings.bRecomputeTangents = true;
+	SrcModel.BuildSettings.bRemoveDegenerates = false;
+	SrcModel.BuildSettings.bUseHighPrecisionTangentBasis = false;
+	SrcModel.BuildSettings.bUseFullPrecisionUVs = false;
+	WaterMesh->CreateMeshDescription(0, MeshDescription);
+
+	UStaticMesh::FCommitMeshDescriptionParams CommitParams;
+	CommitParams.bMarkPackageDirty = true;
+	CommitParams.bUseHashAsGuid = true;
+	WaterMesh->CommitMeshDescription(0,CommitParams);
+
+	WaterMesh->ImportVersion = EImportStaticMeshVersion::LastVersion;
+}
+
+UStaticMesh* FWaterBodyMeshBuilder::CreateUStaticMesh(UObject* Outer, FName MeshName) const
+{
+	UStaticMesh* StaticMesh = NewObject<UStaticMesh>(Outer, MeshName);
+
+	// Disable navigation
+	StaticMesh->MarkAsNotHavingNavigationData();
+
+	if (UBodySetup* BodySetup = StaticMesh->GetBodySetup())
+	{
+		BodySetup->DefaultInstance.SetCollisionProfileName(UCollisionProfile::NoCollision_ProfileName);
+		BodySetup->CollisionTraceFlag = CTF_UseSimpleAsComplex;
+		// We won't ever enable collisions (since collisions are handled by the dedicated water body collision components), ensure we don't even cook or load any collision data on this mesh: 
+		BodySetup->bNeverNeedsCookedCollisionData = true;
+	}
+	return StaticMesh;
+}
+#endif // WITH_EDITOR
+
diff --git a/Engine/Plugins/Experimental/Water/Source/Runtime/Private/WaterBodyMeshBuilder.h b/Engine/Plugins/Experimental/Water/Source/Runtime/Private/WaterBodyMeshBuilder.h
new file mode 100644
--- /dev/null
+++ b/Engine/Plugins/Experimental/Water/Source/Runtime/Private/WaterBodyMeshBuilder.h
@@ -0,0 +1,45 @@
+// Copyright Epic Games, Inc. All Rights Reserved.
+
+#pragma once
+
+#include "MeshDescription.h"
+#include "WaterBodyStaticMeshSettings.h"
+
+class UWaterBodyComponent;
+class UWaterBodyStaticMeshComponent;
+class UWaterBodyInfoMeshComponent;
+class UStaticMesh;
+
+namespace UE::Geometry { class FDynamicMesh3; }
+
+#if WITH_EDITOR
+class FWaterBodyMeshBuilder
+{
+public:
+	/** 
+	 * Builds and initializes the UWaterBodyStaticMeshComponents for a WaterBodyComponent to be able to render itself without relying on the water info texture or water mesh.
+	 * The function will attempt to reuse static mesh components passed in via the ReusableComponents array.
+	 * Note: components generated by this process will be attached and registered to the UWaterBodyComponent passed in.
+	 */
+	TArray<TObjectPtr<UWaterBodyStaticMeshComponent>> BuildWaterBodyStaticMesh(UWaterBodyComponent* WaterBodyComponent, const FWaterBodyStaticMeshSettings& Settings, const TArray<TObjectPtr<UWaterBodyStaticMeshComponent>>& ReusableComponents) const;
+
+	/** 
+	 * Builds and assigns static meshes for the WaterBody Info mesh and Dilated Mesh directly to the components passed.
+	 */
+	void BuildWaterInfoMeshes(UWaterBodyComponent* WaterBodyComponent, UWaterBodyInfoMeshComponent* WaterBodyInfoMesh, UWaterBodyInfoMeshComponent* WaterBodyDilatedMesh) const;
+
+	/**
+	 * Builds a mesh description out of the water body geometry.
+	 */
+	FMeshDescription BuildMeshDescription(const UWaterBodyComponent* WaterBodyComponent) const;
+private:
+	/** Helper function to retrieve the dynamic mesh representations of a water body component with all the vertex attributes initialized. Optionally does the same for the DilatedMesh if the parameter is non-null */
+	void GetDynamicMesh(const UWaterBodyComponent* WaterBodyComponent, UE::Geometry::FDynamicMesh3& Mesh, UE::Geometry::FDynamicMesh3* DilateMesh = nullptr) const;
+	
+	FMeshDescription ConvertDynamicMeshToMeshDescription(const UE::Geometry::FDynamicMesh3& Mesh) const;
+
+	void UpdateStaticMesh(UStaticMesh* Mesh, const FMeshDescription& MeshDescription) const;
+
+	UStaticMesh* CreateUStaticMesh(UObject* Outer, FName MeshName) const;
+};
+#endif // WITH_EDITOR
diff --git a/Engine/Plugins/Experimental/Water/Source/Runtime/Private/WaterBodyMeshComponent.cpp b/Engine/Plugins/Experimental/Water/Source/Runtime/Private/WaterBodyMeshComponent.cpp
new file mode 100644
--- /dev/null
+++ b/Engine/Plugins/Experimental/Water/Source/Runtime/Private/WaterBodyMeshComponent.cpp
@@ -0,0 +1,84 @@
+// Copyright Epic Games, Inc. All Rights Reserved.
+
+#include "WaterBodyMeshComponent.h"
+#include "WaterBodyActor.h"
+#include "WaterBodyComponent.h"
+#include "DynamicMeshToMeshDescription.h"
+#include "Engine/CollisionProfile.h"
+#include "PhysicsEngine/BodySetup.h"
+#include "WaterModule.h"
+#include "StaticMeshResources.h"
+#include "Engine/StaticMesh.h"
+
+#if WITH_EDITOR
+#include "UObject/Package.h"
+#include "StaticMeshAttributes.h"
+#include "DynamicMesh/DynamicMesh3.h"
+#include "DynamicMeshToMeshDescription.h"
+#include "DynamicMesh/DynamicMeshAttributeSet.h"
+#include "DynamicMesh/MeshNormals.h"
+#include "DynamicMesh/MeshTangents.h"
+#endif 
+
+#include UE_INLINE_GENERATED_CPP_BY_NAME(WaterBodyMeshComponent)
+
+UWaterBodyMeshComponent::UWaterBodyMeshComponent(const FObjectInitializer& ObjectInitializer)
+	: Super(ObjectInitializer)
+{
+	SetMobility(EComponentMobility::Static);
+
+	SetCollisionProfileName(UCollisionProfile::NoCollision_ProfileName);
+	SetGenerateOverlapEvents(false);
+	SetCollisionEnabled(ECollisionEnabled::NoCollision);
+
+	CanCharacterStepUpOn = ECanBeCharacterBase::ECB_No;
+	DepthPriorityGroup = SDPG_World;
+
+	bComputeFastLocalBounds = true;
+	bComputeBoundsOnceForGame = true;
+
+	AlwaysLoadOnServer = false;
+	bSelectable = true;
+
+	bCanEverAffectNavigation = false;
+}
+
+bool UWaterBodyMeshComponent::CanCreateSceneProxy() const
+{
+	if (GetStaticMesh() == nullptr)
+	{
+		UE_LOG(LogWater, Verbose, TEXT("Skipping CreateSceneProxy for WaterBodyMeshComponent %s (StaticMesh is null)"), *GetFullName());
+		return false;
+	}
+
+	// Prevent accessing the RenderData during async compilation. The RenderState will be recreated when compilation finishes.
+	if (GetStaticMesh()->IsCompiling())
+	{
+		UE_LOG(LogWater, Verbose, TEXT("Skipping CreateSceneProxy for WaterBodyMeshComponent %s (StaticMesh is not ready)"), *GetFullName());
+		return false;
+	}
+
+	if (GetStaticMesh()->GetRenderData() == nullptr)
+	{
+		UE_LOG(LogWater, Verbose, TEXT("Skipping CreateSceneProxy for WaterBodyMeshComponent %s (RenderData is null)"), *GetFullName());
+		return false;
+	}
+
+	// By now the compilation should be finished so having null render data is not valid.
+	if (!GetStaticMesh()->GetRenderData()->IsInitialized())
+	{
+		UE_LOG(LogWater, Warning, TEXT("Skipping CreateSceneProxy for WaterBodyMeshComponent %s (RenderData is not initialized)"), *GetFullName());
+		return false;
+	}
+
+	const FStaticMeshLODResourcesArray& LODResources = GetStaticMesh()->GetRenderData()->LODResources;
+	const int32 SMCurrentMinLOD = GetStaticMesh()->GetMinLODIdx();
+	const int32 EffectiveMinLOD = bOverrideMinLOD ? MinLOD : SMCurrentMinLOD;
+	if (LODResources.Num() == 0	|| LODResources[FMath::Clamp<int32>(EffectiveMinLOD, 0, LODResources.Num()-1)].VertexBuffers.StaticMeshVertexBuffer.GetNumVertices() == 0)
+	{
+		UE_LOG(LogWater, Warning, TEXT("Skipping CreateSceneProxy for WaterBodyMeshComponent %s (LOD problems)"), *GetFullName());
+		return false;
+	}
+
+	return true;
+}
diff --git a/Engine/Plugins/Experimental/Water/Source/Runtime/Private/WaterBodyOceanComponent.cpp b/Engine/Plugins/Experimental/Water/Source/Runtime/Private/WaterBodyOceanComponent.cpp
--- a/Engine/Plugins/Experimental/Water/Source/Runtime/Private/WaterBodyOceanComponent.cpp
+++ b/Engine/Plugins/Experimental/Water/Source/Runtime/Private/WaterBodyOceanComponent.cpp
@@ -10,11 +10,17 @@
 #include "WaterSubsystem.h"
 #include "ConstrainedDelaunay2.h"
 #include "Operations/InsetMeshRegion.h"
+#include "DynamicMesh/DynamicMeshAttributeSet.h"
+#include "Curve/PolygonOffsetUtils.h"
+
+#if WITH_EDITOR
+#include "Misc/UObjectToken.h"
+#include "WaterRuntimeSettings.h"
+#endif // WITH_EDITOR
 
 #include UE_INLINE_GENERATED_CPP_BY_NAME(WaterBodyOceanComponent)
 
-#if WITH_EDITOR
-#endif
+#define LOCTEXT_NAMESPACE "Water"
 
 // ----------------------------------------------------------------------------------
 
@@ -22,6 +28,7 @@ UWaterBodyOceanComponent::UWaterBodyOceanComponent(const FObjectInitializer& Obj
 	: Super(ObjectInitializer)
 {
 	CollisionExtents = FVector(50000.f, 50000.f, 10000.f);
+	OceanExtents = FVector2D(51200., 51200.);
 
 	// @todo_water : Remove these checks (Once AWaterBody is no more Blueprintable, these methods should become PURE_VIRTUAL and this class should overload them)
 	check(IsFlatSurface());
@@ -82,6 +89,22 @@ void UWaterBodyOceanComponent::SetCollisionExtents(const FVector& NewExtents)
 	UpdateAll(Params);
 }
 
+void UWaterBodyOceanComponent::SetOceanExtent(const FVector2D& NewExtents)
+{
+	OceanExtents = NewExtents;
+
+	UpdateWaterBodyRenderData();
+}
+
+void UWaterBodyOceanComponent::FillWaterZoneWithOcean()
+{
+	if (const AWaterZone* WaterZone = GetWaterZone())
+	{
+		OceanExtents = WaterZone->GetZoneExtent();
+		UpdateWaterBodyRenderData();
+	}
+}
+
 void UWaterBodyOceanComponent::OnPostEditChangeProperty(FOnWaterBodyChangedParams& InOutOnWaterBodyChangedParams)
 {
 	Super::OnPostEditChangeProperty(InOutOnWaterBodyChangedParams);
@@ -97,27 +120,8 @@ const TCHAR* UWaterBodyOceanComponent::GetWaterSpriteTextureName() const
 {
 	return TEXT("/Water/Icons/WaterBodyOceanSprite");
 }
-#endif
 
-void UWaterBodyOceanComponent::Reset()
-{
-	for (UBoxComponent* Component : CollisionBoxes)
-	{
-		if (Component)
-		{
-			Component->DestroyComponent();
-		}
-	}
-	CollisionBoxes.Reset();
-	for (UOceanCollisionComponent* Component : CollisionHullSets)
-	{
-		if (Component)
-		{
-			Component->DestroyComponent();
-		}
-	}
-	CollisionHullSets.Reset();
-}
+#endif
 
 // Adds a quad to a dynamic mesh reusing passed vertex indices if they are not INDEX_NONE
 static void AddAABBQuadToDynamicMesh(UE::Geometry::FDynamicMesh3& EditMesh, FVector3d Min, FVector3d Max, TArray<int32, TInlineAllocator<4>>& InOutVertices)
@@ -129,45 +133,52 @@ static void AddAABBQuadToDynamicMesh(UE::Geometry::FDynamicMesh3& EditMesh, FVec
 	int32 VertexCIndex = InOutVertices[2];
 	int32 VertexDIndex = InOutVertices[3];
 
+	FDynamicMeshColorOverlay* ColorOverlay = EditMesh.Attributes()->PrimaryColors();
+	FDynamicMeshNormalOverlay* NormalOverlay = EditMesh.Attributes()->PrimaryNormals();
+
 	if (VertexAIndex == INDEX_NONE)
 	{ 
 		FVertexInfo VertexA;
 		VertexA.Position = Min;
-		VertexA.Color = FVector3f(0.0);
-		VertexA.bHaveC = true;
 		VertexAIndex = EditMesh.AppendVertex(VertexA);
+		ColorOverlay->AppendElement(FVector4f(0.0));
+		NormalOverlay->AppendElement(FVector3f(0., 0., 1.));
 	}
 	
 	if (VertexBIndex == INDEX_NONE)
 	{
 		FVertexInfo VertexB;
 		VertexB.Position = FVector3d(Max.X, Min.Y, 0);
-		VertexB.Color = FVector3f(0.0);
-		VertexB.bHaveC = true;
 		VertexBIndex = EditMesh.AppendVertex(VertexB);
+		ColorOverlay->AppendElement(FVector4f(0.0));
+		NormalOverlay->AppendElement(FVector3f(0., 0., 1.));
 	}
 
 	if (VertexCIndex == INDEX_NONE)
 	{
 		FVertexInfo VertexC;
 		VertexC.Position = Max;
-		VertexC.Color = FVector3f(0.0);
-		VertexC.bHaveC = true;
 		VertexCIndex = EditMesh.AppendVertex(VertexC);
+		ColorOverlay->AppendElement(FVector4f(0.0));
+		NormalOverlay->AppendElement(FVector3f(0., 0., 1.));
 	}
 
 	if (VertexDIndex == INDEX_NONE)
 	{
 		FVertexInfo VertexD;
 		VertexD.Position = FVector3d(Min.X, Max.Y, 0);
-		VertexD.Color = FVector3f(0.0);
-		VertexD.bHaveC = true;
 		VertexDIndex = EditMesh.AppendVertex(VertexD);
+		ColorOverlay->AppendElement(FVector4f(0.0));
+		NormalOverlay->AppendElement(FVector3f(0., 0., 1.));
 	}
 
 	const int TriOne = EditMesh.AppendTriangle(VertexAIndex, VertexDIndex, VertexBIndex);
 	const int TriTwo = EditMesh.AppendTriangle(VertexBIndex, VertexDIndex, VertexCIndex);
+
 	check(TriOne != INDEX_NONE && TriTwo != INDEX_NONE);
+	
+	ColorOverlay->SetTriangle(TriOne, FIndex3i(VertexAIndex, VertexDIndex, VertexBIndex));
+	ColorOverlay->SetTriangle(TriTwo, FIndex3i(VertexBIndex, VertexDIndex, VertexCIndex));
 
 	InOutVertices = {VertexAIndex, VertexBIndex, VertexCIndex, VertexDIndex};
 }
@@ -187,7 +198,7 @@ bool UWaterBodyOceanComponent::GenerateWaterBodyMesh(UE::Geometry::FDynamicMesh3
 		return false;
 	}
 
-	if (!WaterZone)
+	if (WaterZone == nullptr)
 	{
 		return false;
 	}
@@ -207,8 +218,7 @@ bool UWaterBodyOceanComponent::GenerateWaterBodyMesh(UE::Geometry::FDynamicMesh3
 	IslandBounds.Expand(1);
 
 	// #todo: account for scale
-	const FVector2d OceanExtents = WaterZone->GetZoneExtent(); // file the entire water zone with an ocean
-	const FAxisAlignedBox2d OceanBounds = FAxisAlignedBox2d(FVector2d(GetComponentTransform().InverseTransformPositionNoScale(WaterZone->GetActorLocation())), OceanExtents.X / 2.0);
+	const FAxisAlignedBox2d OceanBounds = FAxisAlignedBox2d(FVector2d(GetComponentTransform().InverseTransformPositionNoScale(WaterZone->GetRootComponent()->GetRelativeLocation())), OceanExtents.X / 2.0);
 	FPolygon2d IslandBoundingPolygon = FPolygon2d::MakeRectangle(IslandBounds.Center(), IslandBounds.Extents().X * 2., IslandBounds.Extents().Y * 2.);
 
 	FConstrainedDelaunay2d Triangulation;
@@ -242,14 +252,17 @@ bool UWaterBodyOceanComponent::GenerateWaterBodyMesh(UE::Geometry::FDynamicMesh3
 	int32 IslandTopLeft = INDEX_NONE;
 	FVector3d IslandTopLeftVertex;
 
+	check(OutMesh.Attributes());
+	FDynamicMeshColorOverlay* ColorOverlay = OutMesh.Attributes()->PrimaryColors();
+	FDynamicMeshNormalOverlay* NormalOverlay = OutMesh.Attributes()->PrimaryNormals();
+
 	for (const FVector2d& Vertex : Triangulation.Vertices)
 	{
-		FVertexInfo VertexInfo;
-		VertexInfo.Position = FVector3d(Vertex, 0.f);
-		VertexInfo.Color = FVector3f(0.0);
-		VertexInfo.bHaveC = true;
+		const FVertexInfo VertexInfo(FVector3d(Vertex, 0.f));
 
-		int32 Index = OutMesh.AppendVertex(VertexInfo);
+		const int32 Index = OutMesh.AppendVertex(VertexInfo);
+		ColorOverlay->AppendElement(FVector4f(0.0));
+		NormalOverlay->AppendElement(FVector3f(0., 0., 1.));
 
 		// Collect the corner vertices of the island bounding box so we can stitch the outer quads to them.
 		if ((IslandBottomLeft == INDEX_NONE) || (VertexInfo.Position.X < IslandBottomLeftVertex.X || VertexInfo.Position.Y < IslandBottomLeftVertex.Y))
@@ -280,7 +293,9 @@ bool UWaterBodyOceanComponent::GenerateWaterBodyMesh(UE::Geometry::FDynamicMesh3
 
 	for (const FIndex3i& Triangle : Triangulation.Triangles)
 	{
-		OutMesh.AppendTriangle(Triangle);
+		const int TriangleID = OutMesh.AppendTriangle(Triangle);
+		ColorOverlay->SetTriangle(TriangleID, Triangle);
+		NormalOverlay->SetTriangle(TriangleID, Triangle);
 	}
 
 	// Add the bounding quads:
@@ -288,90 +303,160 @@ bool UWaterBodyOceanComponent::GenerateWaterBodyMesh(UE::Geometry::FDynamicMesh3
 	// Each quad has four vertices laid out like so
 	// D --- C
 	// |  \  |
-	// A --- B
-	
-	TArray<int32, TInlineAllocator<4>> Vertices({ INDEX_NONE, INDEX_NONE, INDEX_NONE, INDEX_NONE });
-
-	// Bottom left quad
-	Vertices[2] = IslandBottomLeft;
-	AddAABBQuadToDynamicMesh(OutMesh, FVector3d(OceanBounds.Min, 0), FVector3d(IslandBounds.Min, 0), Vertices);
-	// We need the B and C vertices of the first quad to attach to the last quad
-	int32 FirstQuadB = Vertices[1];
-	int32 FirstQuadC = Vertices[2];
-
-	// Left middle quad
-	Vertices[0] = Vertices[3];
-	Vertices[1] = Vertices[2];
-	Vertices[2] = IslandTopLeft;
-	Vertices[3] = INDEX_NONE;
-	AddAABBQuadToDynamicMesh(OutMesh, FVector3d(OceanBounds.Min.X, IslandBounds.Min.Y, 0), FVector3d(IslandBounds.Min.X, IslandBounds.Max.Y, 0), Vertices);
-
-	// Top left quad
-	Vertices[0] = Vertices[3];
-	Vertices[1] = Vertices[2];
-	Vertices[2] = INDEX_NONE;
-	Vertices[3] = INDEX_NONE;
-	AddAABBQuadToDynamicMesh(OutMesh, FVector3d(OceanBounds.Min.X, IslandBounds.Max.Y, 0), FVector3d(IslandBounds.Min.X, OceanBounds.Max.Y, 0), Vertices);
-
-	// Top middle quad
-	Vertices[0] = Vertices[1];
-	Vertices[3] = Vertices[2];
-	Vertices[1] = IslandTopRight;
-	Vertices[2] = INDEX_NONE;
-	AddAABBQuadToDynamicMesh(OutMesh, FVector3d(IslandBounds.Min.X, IslandBounds.Max.Y, 0), FVector3d(IslandBounds.Max.X, OceanBounds.Max.Y, 0), Vertices);
-
-	// Top right quad
-	Vertices[0] = Vertices[1];
-	Vertices[3] = Vertices[2];
-	Vertices[1] = INDEX_NONE;
-	Vertices[2] = INDEX_NONE;
-	AddAABBQuadToDynamicMesh(OutMesh, FVector3d(IslandBounds.Max, 0), FVector3d(OceanBounds.Max, 0), Vertices);
-
-	// Middle right quad
-	Vertices[3] = Vertices[0];
-	Vertices[2] = Vertices[1];
-	Vertices[1] = INDEX_NONE;
-	Vertices[0] = IslandBottomRight;
-	AddAABBQuadToDynamicMesh(OutMesh, FVector3d(IslandBounds.Max.X, IslandBounds.Min.Y, 0), FVector3d(OceanBounds.Max.X, IslandBounds.Max.Y, 0), Vertices);
-
-	// Bottom right quad
-	Vertices[3] = Vertices[0];
-	Vertices[2] = Vertices[1];
-	Vertices[1] = INDEX_NONE;
-	Vertices[0] = INDEX_NONE;
-	AddAABBQuadToDynamicMesh(OutMesh, FVector3d(IslandBounds.Max.X, OceanBounds.Min.Y, 0), FVector3d(OceanBounds.Max.X, IslandBounds.Min.Y, 0), Vertices);
-
-	// Middle bottom quad
-	Vertices[1] = Vertices[0];
-	Vertices[2] = Vertices[3];
-	Vertices[0] = FirstQuadB;
-	Vertices[3] = FirstQuadC;
-	AddAABBQuadToDynamicMesh(OutMesh, FVector3d(IslandBounds.Min.X, OceanBounds.Min.Y, 0), FVector3d(IslandBounds.Max.X, IslandBounds.Min.Y, 0), Vertices);
+	// A --- Bo
+	{
+		TArray<int32, TInlineAllocator<4>> Vertices({ INDEX_NONE, INDEX_NONE, INDEX_NONE, INDEX_NONE });
+
+		// Bottom left quad
+		Vertices[2] = IslandBottomLeft;
+		AddAABBQuadToDynamicMesh(OutMesh, FVector3d(OceanBounds.Min, 0), FVector3d(IslandBounds.Min, 0), Vertices);
+		// We need the B and C vertices of the first quad to attach to the last quad
+		int32 FirstQuadB = Vertices[1];
+		int32 FirstQuadC = Vertices[2];
+
+		// Left middle quad
+		Vertices[0] = Vertices[3];
+		Vertices[1] = Vertices[2];
+		Vertices[2] = IslandTopLeft;
+		Vertices[3] = INDEX_NONE;
+		AddAABBQuadToDynamicMesh(OutMesh, FVector3d(OceanBounds.Min.X, IslandBounds.Min.Y, 0), FVector3d(IslandBounds.Min.X, IslandBounds.Max.Y, 0), Vertices);
+
+		// Top left quad
+		Vertices[0] = Vertices[3];
+		Vertices[1] = Vertices[2];
+		Vertices[2] = INDEX_NONE;
+		Vertices[3] = INDEX_NONE;
+		AddAABBQuadToDynamicMesh(OutMesh, FVector3d(OceanBounds.Min.X, IslandBounds.Max.Y, 0), FVector3d(IslandBounds.Min.X, OceanBounds.Max.Y, 0), Vertices);
+
+		// Top middle quad
+		Vertices[0] = Vertices[1];
+		Vertices[3] = Vertices[2];
+		Vertices[1] = IslandTopRight;
+		Vertices[2] = INDEX_NONE;
+		AddAABBQuadToDynamicMesh(OutMesh, FVector3d(IslandBounds.Min.X, IslandBounds.Max.Y, 0), FVector3d(IslandBounds.Max.X, OceanBounds.Max.Y, 0), Vertices);
+
+		// Top right quad
+		Vertices[0] = Vertices[1];
+		Vertices[3] = Vertices[2];
+		Vertices[1] = INDEX_NONE;
+		Vertices[2] = INDEX_NONE;
+		AddAABBQuadToDynamicMesh(OutMesh, FVector3d(IslandBounds.Max, 0), FVector3d(OceanBounds.Max, 0), Vertices);
+
+		// Middle right quad
+		Vertices[3] = Vertices[0];
+		Vertices[2] = Vertices[1];
+		Vertices[1] = INDEX_NONE;
+		Vertices[0] = IslandBottomRight;
+		AddAABBQuadToDynamicMesh(OutMesh, FVector3d(IslandBounds.Max.X, IslandBounds.Min.Y, 0), FVector3d(OceanBounds.Max.X, IslandBounds.Max.Y, 0), Vertices);
+
+		// Bottom right quad
+		Vertices[3] = Vertices[0];
+		Vertices[2] = Vertices[1];
+		Vertices[1] = INDEX_NONE;
+		Vertices[0] = INDEX_NONE;
+		AddAABBQuadToDynamicMesh(OutMesh, FVector3d(IslandBounds.Max.X, OceanBounds.Min.Y, 0), FVector3d(OceanBounds.Max.X, IslandBounds.Min.Y, 0), Vertices);
+
+		// Middle bottom quad
+		Vertices[1] = Vertices[0];
+		Vertices[2] = Vertices[3];
+		Vertices[0] = FirstQuadB;
+		Vertices[3] = FirstQuadC;
+		AddAABBQuadToDynamicMesh(OutMesh, FVector3d(IslandBounds.Min.X, OceanBounds.Min.Y, 0), FVector3d(IslandBounds.Max.X, IslandBounds.Min.Y, 0), Vertices);
+	}
 	
 	if (ShapeDilation > 0.f && OutDilatedMesh)
 	{
 		TRACE_CPUPROFILER_EVENT_SCOPE(DilateOceanMesh);
 
-		// Inset the mesh by -ShapeDilation to effectively expand the mesh
-		OutDilatedMesh->Copy(OutMesh);
-		FInsetMeshRegion Inset(OutDilatedMesh);
-		Inset.InsetDistance = -1 * ShapeDilation / 2.f;
+		FPolygon2d OceanBoundingPoly;
+		OceanBoundingPoly.AppendVertex(FVector2D(OceanBounds.Min));
+		OceanBoundingPoly.AppendVertex(FVector2D(OceanBounds.Min.X, OceanBounds.Max.Y));
+		OceanBoundingPoly.AppendVertex(FVector2D(OceanBounds.Max));
+		OceanBoundingPoly.AppendVertex(FVector2D(OceanBounds.Max.X, OceanBounds.Min.Y));
 
-		Inset.Triangles.Reserve(OutDilatedMesh->TriangleCount());
-		for (int32 Idx : OutDilatedMesh->TriangleIndicesItr())
+		FGeneralPolygon2d OceanPoly(OceanBoundingPoly);
+		if (Island.IsClockwise())
 		{
-			Inset.Triangles.Add(Idx);
+			Island.Reverse();
 		}
-		
-		if (!Inset.Apply())
+		OceanPoly.AddHole(Island);
+
+		TArray<FGeneralPolygon2d> OffsetPolys;
+		UE::Geometry::PolygonsOffset(
+			ShapeDilation / 2.f,
+			{ OceanPoly },
+			OffsetPolys,
+			false,
+			2.0);
+
+		FConstrainedDelaunay2d DilationTriangulation;
+		DilationTriangulation.FillRule = FConstrainedDelaunay2d::EFillRule::Positive;
+
+		for (const FGeneralPolygon2d& Poly : OffsetPolys)
+		{
+			if (Poly.SignedArea() <= 0.)
+			{
+				UE_LOG(LogWater, Warning, TEXT("Failed to apply offset for shape dilation (%s"), *GetOwner()->GetActorNameOrLabel());
+				continue;
+			}
+
+			DilationTriangulation.Add(Poly);
+		}
+
+		if (!DilationTriangulation.Triangulate())
+		{
+			UE_LOG(LogWater, Warning, TEXT("Failed to triangulate dilated ocean mesh (%s"), *GetOwner()->GetActorNameOrLabel());
+			return false;
+		}
+
+		if (DilationTriangulation.Triangles.Num() == 0)
 		{
-			UE_LOG(LogWater, Warning, TEXT("Failed to apply mesh inset for shape dilation (%s)"), *GetOwner()->GetActorNameOrLabel());
 			return false;
 		}
-	}	
+
+		FDynamicMeshColorOverlay* DilatedColorOverlay = OutDilatedMesh->Attributes()->PrimaryColors();
+		FDynamicMeshNormalOverlay* DilatedNormalOverlay = OutDilatedMesh->Attributes()->PrimaryNormals();
+		for (const FVector2d& Vertex : DilationTriangulation.Vertices)
+		{
+			FVertexInfo MeshVertex(FVector3d(Vertex.X, Vertex.Y, 0.));
+
+			OutDilatedMesh->AppendVertex(MeshVertex);
+			DilatedColorOverlay->AppendElement(FVector4f(0.f));
+			DilatedNormalOverlay->AppendElement(FVector3f(0., 0., 1.));
+		}
+
+		for (const FIndex3i& Triangle : DilationTriangulation.Triangles)
+		{
+			const int TriangleID = OutDilatedMesh->AppendTriangle(Triangle);
+			DilatedColorOverlay->SetTriangle(TriangleID, Triangle);
+			DilatedNormalOverlay->SetTriangle(TriangleID, Triangle);
+		}
+	}
+
 	return true;
 }
 
+
+void UWaterBodyOceanComponent::Reset()
+{
+	for (UBoxComponent* Component : CollisionBoxes)
+	{
+		if (Component)
+		{
+			Component->DestroyComponent();
+		}
+	}
+	CollisionBoxes.Reset();
+	for (UOceanCollisionComponent* Component : CollisionHullSets)
+	{
+		if (Component)
+		{
+			Component->DestroyComponent();
+		}
+	}
+	CollisionHullSets.Reset();
+}
+
 void UWaterBodyOceanComponent::PostLoad()
 {
 	Super::PostLoad();
@@ -380,11 +465,25 @@ void UWaterBodyOceanComponent::PostLoad()
 #endif // WITH_EDITORONLY_DATA
 }
 
+void UWaterBodyOceanComponent::OnPostRegisterAllComponents()
+{
+	Super::OnPostRegisterAllComponents();
+
+#if WITH_EDITOR
+	// In this version the ocean was changed to not be strongly coupled to the water zone and instead rely on the user to keep the mesh extent in sync with the zone extent
+	if (GetLinkerCustomVersion(FFortniteMainBranchObjectVersion::GUID) < FFortniteMainBranchObjectVersion::WaterBodyStaticMeshComponents)
+	{
+		// Trigger a rebuild of the render data manually since this doesn't occur on load anymore. 
+		// The mesh extent will be out of sync for old oceans so ensure they rebuild now:
+		FillWaterZoneWithOcean();
+	}
+#endif // WITH_EDITOR
+}
+	
 FBoxSphereBounds UWaterBodyOceanComponent::CalcBounds(const FTransform& LocalToWorld) const
 {
 	if (AWaterZone* WaterZone = GetWaterZone())
 	{
-		const FVector2D OceanExtents = WaterZone->GetZoneExtent();
 		const FVector Min(FVector(-OceanExtents / 2.0, -1.0 * GetChannelDepth()));
 		const FVector Max(FVector(OceanExtents / 2.0, 0.0));
 		return FBoxSphereBounds(FBox(Min, Max)).TransformBy(LocalToWorld);
@@ -392,6 +491,60 @@ FBoxSphereBounds UWaterBodyOceanComponent::CalcBounds(const FTransform& LocalToW
 	return FBoxSphereBounds(EForceInit::ForceInit).TransformBy(LocalToWorld);
 }
 
+#if WITH_EDITOR
+
+void UWaterBodyOceanComponent::OnWaterBodyRenderDataUpdated()
+{
+	Super::OnWaterBodyRenderDataUpdated();
+
+	// Store the location of the zone when the ocean render data is created so we can validate this later and check for inconsistencies:
+	if (const AWaterZone* WaterZone = GetWaterZone())
+	{
+		SavedZoneLocation = FVector2D(WaterZone->GetActorLocation());
+	}
+}
+
+TArray<TSharedRef<FTokenizedMessage>> UWaterBodyOceanComponent::CheckWaterBodyStatus()
+{
+	TArray<TSharedRef<FTokenizedMessage>> Result = Super::CheckWaterBodyStatus();
+
+	if (const AWaterZone* WaterZone = GetWaterZone())
+	{
+		auto DisableWarnOnMismatchExtent = []()
+		{
+			if (UWaterRuntimeSettings* WaterRuntimeSettings = GetMutableDefault<UWaterRuntimeSettings>())
+			{
+				WaterRuntimeSettings->SetShouldWarnOnMismatchOceanExtent(false);
+			}
+		};
+
+		const UWaterRuntimeSettings* WaterRuntimeSettings = GetDefault<UWaterRuntimeSettings>();
+		check(WaterRuntimeSettings);
+
+		if (WaterRuntimeSettings->ShouldWarnOnMismatchOceanExtent())
+		{
+			if ((WaterZone->GetZoneExtent() != OceanExtents)
+			 || (FVector2D(WaterZone->GetActorLocation()) != SavedZoneLocation))
+			{
+				Result.Add(FTokenizedMessage::Create(EMessageSeverity::Warning)
+					->AddToken(FUObjectToken::Create(this))
+					->AddToken(FTextToken::Create(FText::Format(
+						LOCTEXT("MapCheck_Message_MismatchedOceanExtent", "WaterBodyOcean ({0}) has a serialized mesh which does not match the WaterZone it belongs to ({1})."),
+						FText::FromString(*GetNameSafe(this)),
+						FText::FromString(*GetNameSafe(WaterZone)))))
+					->AddToken(FActionToken::Create(LOCTEXT("MapCheck_MessageAction_SaveOcean", "Click here to fill the zone with ocean."), FText(),
+						FOnActionTokenExecuted::CreateUObject(this, &UWaterBodyOceanComponent::FillWaterZoneWithOcean)))
+					->AddToken(FActionToken::Create(LOCTEXT("MapCheck_MessageAction_SaveOcean", "If this is desired behavior, click here to disable this warning."), FText(),
+						FOnActionTokenExecuted::CreateLambda(DisableWarnOnMismatchExtent)))
+						);
+			}
+		}
+	}
+	return Result;
+}
+
+#endif // WITH_EDITOR
+
 void UWaterBodyOceanComponent::OnUpdateBody(bool bWithExclusionVolumes)
 {
 	AActor* OwnerActor = GetOwner();
@@ -506,3 +659,5 @@ void UWaterBodyOceanComponent::OnUpdateBody(bool bWithExclusionVolumes)
 		Reset();
 	}
 }
+
+#undef LOCTEXT_NAMESPACE
diff --git a/Engine/Plugins/Experimental/Water/Source/Runtime/Private/WaterBodyRiverComponent.cpp b/Engine/Plugins/Experimental/Water/Source/Runtime/Private/WaterBodyRiverComponent.cpp
--- a/Engine/Plugins/Experimental/Water/Source/Runtime/Private/WaterBodyRiverComponent.cpp
+++ b/Engine/Plugins/Experimental/Water/Source/Runtime/Private/WaterBodyRiverComponent.cpp
@@ -13,15 +13,13 @@
 #include "Engine/StaticMesh.h"
 #include "Engine/World.h"
 #include "DynamicMesh/DynamicMesh3.h"
+#include "DynamicMesh/DynamicMeshAttributeSet.h"
 
 // for working around Chaos issue
 #include "Chaos/Convex.h"
 
 #include UE_INLINE_GENERATED_CPP_BY_NAME(WaterBodyRiverComponent)
 
-#if WITH_EDITOR
-#endif
-
 // ----------------------------------------------------------------------------------
 
 UWaterBodyRiverComponent::UWaterBodyRiverComponent(const FObjectInitializer& ObjectInitializer)
@@ -48,12 +46,11 @@ TArray<UPrimitiveComponent*> UWaterBodyRiverComponent::GetCollisionComponents(bo
 
 TArray<UPrimitiveComponent*> UWaterBodyRiverComponent::GetStandardRenderableComponents() const
 {
-	TArray<UPrimitiveComponent*> Result;
-	Result.Reserve(SplineMeshComponents.Num());
+	TArray<UPrimitiveComponent*> Result = Super::GetStandardRenderableComponents();
+	Result.Reserve(Result.Num() + SplineMeshComponents.Num());
 
-	Algo::TransformIf(SplineMeshComponents, Result, 
-		[](USplineMeshComponent* SplineComp) { return (SplineComp != nullptr); }, 
-		[](USplineMeshComponent* SplineComp) { return SplineComp; });
+	Algo::CopyIf(SplineMeshComponents, Result, 
+		[](USplineMeshComponent* SplineComp) { return (SplineComp != nullptr); });
 
 	return Result;
 }
@@ -137,12 +134,13 @@ static void AddVerticesForRiverSplineStep(
 	FVertexInfo Left(FVector3d(Pos - OutwardDistance));
 	FVertexInfo Right(FVector3d(Pos + OutwardDistance));
 
-	const FVector3f FlowData = FVector3f(Velocity, FlowDirection, 0.f);
+	const FVector4f FlowData = FWaterUtils::PackFlowData(Velocity, FlowDirection);
 	Left.Color = FlowData;
 	Right.Color = FlowData;
-	Left.bHaveC = true;
-	Right.bHaveC = true;
 
+	check(OutMesh.Attributes());
+	FDynamicMeshColorOverlay* Colors = OutMesh.Attributes()->PrimaryColors();
+	FDynamicMeshNormalOverlay* Normals = OutMesh.Attributes()->PrimaryNormals();
 	// Append regular geometry to mesh:
 	{
 		/* Non - dilated river segment geometry:
@@ -153,15 +151,23 @@ static void AddVerticesForRiverSplineStep(
 		const int32 BaseIndex = OutMesh.GetVerticesBuffer().Num();
 
 		OutMesh.AppendVertex(Left);
+		Colors->AppendElement(FlowData);
+		Normals->AppendElement(FVector3f(0., 0., 1.));
+
 		OutMesh.AppendVertex(Right);
+		Colors->AppendElement(FlowData);
+		Normals->AppendElement(FVector3f(0., 0., 1.));
 
 		if (BaseIndex != 0)
 		{
-			OutMesh.AppendTriangle(BaseIndex - 2, BaseIndex + 1, BaseIndex - 1);
-			OutMesh.AppendTriangle(BaseIndex - 2, BaseIndex, BaseIndex + 1);
+			int TriangleID1 = OutMesh.AppendTriangle(BaseIndex - 2, BaseIndex + 1, BaseIndex - 1);
+			int TriangleID2 = OutMesh.AppendTriangle(BaseIndex - 2, BaseIndex, BaseIndex + 1);
+
+			Colors->SetTriangle(TriangleID1, FIndex3i(BaseIndex - 2, BaseIndex + 1, BaseIndex - 1));
+			Colors->SetTriangle(TriangleID2, FIndex3i(BaseIndex - 2, BaseIndex, BaseIndex + 1));
 		}
 	}
-	
+
 	const float DilationAmount = Component->ShapeDilation;
 	// If dilation is required, append dilated geometry:
 	if (DilationAmount > 0.f && OutDilatedMesh)
diff --git a/Engine/Plugins/Experimental/Water/Source/Runtime/Private/WaterBodySceneProxy.cpp b/Engine/Plugins/Experimental/Water/Source/Runtime/Private/WaterBodySceneProxy.cpp
deleted file mode 100644
diff --git a/Engine/Plugins/Experimental/Water/Source/Runtime/Private/WaterBodyStaticMeshComponent.h b/Engine/Plugins/Experimental/Water/Source/Runtime/Private/WaterBodyStaticMeshComponent.h
new file mode 100644
--- /dev/null
+++ b/Engine/Plugins/Experimental/Water/Source/Runtime/Private/WaterBodyStaticMeshComponent.h
@@ -0,0 +1,22 @@
+// Copyright Epic Games, Inc. All Rights Reserved.
+
+#pragma once
+
+#include "CoreMinimal.h"
+#include "UObject/ObjectMacros.h"
+#include "Components/StaticMeshComponent.h"
+#include "WaterBodyMeshComponent.h"
+
+#include "WaterBodyStaticMeshComponent.generated.h"
+
+class UWaterBodyComponent;
+
+/**
+* WaterBodyMeshComponent used to when rendering water statically without relying on the dynamic tessellation of the water zone/water mesh.
+*/
+UCLASS(MinimalAPI)
+class UWaterBodyStaticMeshComponent : public UWaterBodyMeshComponent
+{
+	GENERATED_BODY()
+public:
+};
diff --git a/Engine/Plugins/Experimental/Water/Source/Runtime/Private/WaterInfoRendering.cpp b/Engine/Plugins/Experimental/Water/Source/Runtime/Private/WaterInfoRendering.cpp
--- a/Engine/Plugins/Experimental/Water/Source/Runtime/Private/WaterInfoRendering.cpp
+++ b/Engine/Plugins/Experimental/Water/Source/Runtime/Private/WaterInfoRendering.cpp
@@ -8,7 +8,6 @@
 #include "Modules/ModuleManager.h"
 #include "RenderCaptureInterface.h"
 #include "RHIStaticStates.h"
-#include "WaterBodySceneProxy.h"
 #include "Engine/TextureRenderTarget2D.h"
 #include "SceneCaptureRendering.h"
 #include "PostProcess/DrawRectangle.h"
@@ -20,6 +19,7 @@
 #include "LandscapeModule.h"
 #include "TextureResource.h"
 #include "WaterBodyComponent.h"
+#include "WaterBodyInfoMeshComponent.h"
 #include "Containers/StridedView.h"
 
 #include "RenderGraphBuilder.h"
@@ -183,7 +183,8 @@ public:
 		SHADER_PARAMETER_STRUCT_REF(FViewUniformShaderParameters, View)
 		SHADER_PARAMETER_STRUCT_INCLUDE(FSceneTextureShaderParameters, SceneTextures)
 		SHADER_PARAMETER_RDG_TEXTURE(Texture2D, WaterInfoTexture)
-		SHADER_PARAMETER(FVector2f, WaterHeightExtents)
+		SHADER_PARAMETER(float, WaterZMin)
+		SHADER_PARAMETER(float, WaterZMax)
 		SHADER_PARAMETER(float, GroundZMin)
 		SHADER_PARAMETER(float, CaptureZ)
 		SHADER_PARAMETER(int, BlurRadius)
@@ -246,10 +247,11 @@ static void FinalizeWaterInfo(
 		PassParameters->RenderTargets[0] = FRenderTargetBinding(OutputTexture, ERenderTargetLoadAction::ENoAction);
 		PassParameters->SceneTextures = GetSceneTextureShaderParameters(View);
 		PassParameters->WaterInfoTexture = WaterInfoTexture;
-		PassParameters->BlurRadius = Params.VelocityBlurRadius;
-		PassParameters->CaptureZ = Params.CaptureZ;
-		PassParameters->WaterHeightExtents = Params.WaterHeightExtents;
+		PassParameters->WaterZMin = Params.WaterHeightExtents.X;
+		PassParameters->WaterZMax = Params.WaterHeightExtents.Y;
 		PassParameters->GroundZMin = Params.GroundZMin;
+		PassParameters->CaptureZ = Params.CaptureZ;
+		PassParameters->BlurRadius = Params.VelocityBlurRadius;
 
 		TShaderMapRef<FScreenVS> VertexShader(ShaderMap);
 		TShaderMapRef<FWaterInfoFinalizePS> PixelShader(ShaderMap, PixelPermutationVector);
@@ -299,21 +301,37 @@ static FMatrix BuildOrthoMatrix(float InOrthoWidth, float InOrthoHeight)
 
 // ---------------------------------------------------------------------------------------------------------------------
 
-static void SetWaterBodiesWithinWaterInfoPass(FSceneInterface* Scene, const FSceneView& View, EWaterInfoPass InPass)
+struct FEnableWaterInfoMeshesScoped
 {
-	TRACE_CPUPROFILER_EVENT_SCOPE(WaterInfo::SetWaterBodiesWithinWaterInfoPass);
-	if (View.ShowOnlyPrimitives.IsSet())
+public:
+	FEnableWaterInfoMeshesScoped(FSceneInterface* Scene, const FSceneView& View)
 	{
-		for (FPrimitiveSceneProxy* PrimProxy : Scene->GetPrimitiveSceneProxies())
+		TRACE_CPUPROFILER_EVENT_SCOPE(WaterInfo::EnableWaterInfoMeshes);
+		if (View.ShowOnlyPrimitives.IsSet())
 		{
-			if (PrimProxy && View.ShowOnlyPrimitives->Contains(PrimProxy->GetPrimitiveComponentId()))
+			for (FPrimitiveSceneProxy* PrimProxy : Scene->GetPrimitiveSceneProxies())
 			{
-				FWaterBodySceneProxy* WaterProxy = (FWaterBodySceneProxy*)PrimProxy;
-				WaterProxy->SetWithinWaterInfoPass(InPass);
+				if (PrimProxy && View.ShowOnlyPrimitives->Contains(PrimProxy->GetPrimitiveComponentId()))
+				{
+					FWaterBodyInfoMeshSceneProxy* WaterProxy = (FWaterBodyInfoMeshSceneProxy*)PrimProxy;
+					WaterProxy->SetEnabled(true);
+					Proxies.Add(WaterProxy);
+				}
 			}
 		}
 	}
-}
+
+	~FEnableWaterInfoMeshesScoped()
+	{
+		TRACE_CPUPROFILER_EVENT_SCOPE(WaterInfo::DisableWaterInfoMeshes);
+		for (FWaterBodyInfoMeshSceneProxy* Proxy : Proxies)
+		{
+			Proxy->SetEnabled(false);
+		}
+	}
+private:
+	TArray<FWaterBodyInfoMeshSceneProxy*> Proxies;
+};
 
 /**
  * Sets up the landscape LOD override for water info rendering and restores it when the instance goes out of scope.
@@ -425,11 +443,13 @@ static void UpdateWaterInfoRendering_RenderThread(
 	TRefCountPtr<IPooledRenderTarget> ExtractedDepthTexture;
 	TRefCountPtr<IPooledRenderTarget> ExtractedColorTexture;
 
-	FSceneInterface* Scene = Params.Scene;
 
 	// Depth-only pass for actors which are considered the ground for water rendering
 	{
+		TRACE_CPUPROFILER_EVENT_SCOPE(WaterInfo::DepthPass);
+
 		FSceneRenderer* DepthRenderer = Params.DepthRenderer;
+		FSceneInterface* Scene = DepthRenderer->ViewFamily.Scene;
 
 		DepthRenderer->RenderThreadBegin(RHICmdList);
 
@@ -483,7 +503,12 @@ static void UpdateWaterInfoRendering_RenderThread(
 
 	// Render the water bodies' data including flow, zoffset, depth
 	{
+		TRACE_CPUPROFILER_EVENT_SCOPE(WaterInfo::ColorPass);
+
 		FSceneRenderer* ColorRenderer = Params.ColorRenderer;
+		FSceneInterface* Scene = ColorRenderer->ViewFamily.Scene;
+
+		FEnableWaterInfoMeshesScoped ScopedWaterInfoMeshEnable(ColorRenderer->Scene, ColorRenderer->Views[0]);
 
 		ColorRenderer->RenderThreadBegin(RHICmdList);
 
@@ -496,8 +521,6 @@ static void UpdateWaterInfoRendering_RenderThread(
 		// We need to execute the pre-render view extensions before we do any view dependent work.
 		FSceneRenderer::ViewExtensionPreRender_RenderThread(GraphBuilder, ColorRenderer);
 
-		SetWaterBodiesWithinWaterInfoPass(Scene, View, EWaterInfoPass::Color);
-		
 		FRDGTextureRef TargetTexture = RegisterExternalTexture(GraphBuilder, RenderTarget->GetRenderTargetTexture(), TEXT("WaterColorTarget"));
 
 		FRDGTextureDesc ColorTextureDesc(TargetTexture->Desc);
@@ -535,7 +558,12 @@ static void UpdateWaterInfoRendering_RenderThread(
 
 	// Depth-only pass for water body dilated sections which get composited back into the water info texture at a lower priority than regular water body data
 	{
+		TRACE_CPUPROFILER_EVENT_SCOPE(WaterInfo::DilationPass);
+
 		FSceneRenderer* DilationRenderer = Params.DilationRenderer;
+		FSceneInterface* Scene = DilationRenderer->ViewFamily.Scene;
+
+		FEnableWaterInfoMeshesScoped ScopedWaterInfoMeshEnable(DilationRenderer->Scene, DilationRenderer->Views[0]);
 
 		DilationRenderer->RenderThreadBegin(RHICmdList);
 
@@ -549,8 +577,6 @@ static void UpdateWaterInfoRendering_RenderThread(
 		// We need to execute the pre-render view extensions before we do any view dependent work.
 		FSceneRenderer::ViewExtensionPreRender_RenderThread(GraphBuilder, DilationRenderer);
 
-		SetWaterBodiesWithinWaterInfoPass(Scene, View, EWaterInfoPass::Dilation);
-
 		FRDGTextureRef TargetTexture = RegisterExternalTexture(GraphBuilder, RenderTarget->GetRenderTargetTexture(), TEXT("WaterDilationTarget"));
 
 		FRDGTextureRef DepthTexture = GraphBuilder.RegisterExternalTexture(ExtractedDepthTexture);
@@ -594,7 +620,6 @@ static void UpdateWaterInfoRendering_RenderThread(
 		AddCopyTexturePass(GraphBuilder, FinalizedTexture, ShaderResourceTexture);
 		GraphBuilder.Execute();
 
-		SetWaterBodiesWithinWaterInfoPass(Scene, View, EWaterInfoPass::None);
 		DilationRenderer->RenderThreadEnd(RHICmdList);
 	}
 
@@ -639,6 +664,9 @@ static FSceneRenderer* CreateWaterInfoSceneRenderer(const FCreateWaterInfoSceneR
 	ShowFlags.Fog = 0;
 	ShowFlags.VolumetricFog = 0;
 	ShowFlags.DynamicShadows = 0;
+
+	ShowFlags.SetDisableOcclusionQueries(true);
+	ShowFlags.SetVirtualShadowMapCaching(false);
 	
 	FSceneViewFamilyContext ViewFamily(FSceneViewFamily::ConstructionValues(
 		Params.RenderTarget,
@@ -660,7 +688,8 @@ static FSceneRenderer* CreateWaterInfoSceneRenderer(const FCreateWaterInfoSceneR
 	ViewInitOptions.ProjectionMatrix = Params.ProjectionMatrix;
 	ViewInitOptions.LODDistanceFactor = 0.001f;
 	ViewInitOptions.OverlayColor = FLinearColor::Black;
-	ViewInitOptions.bIsSceneCapture = true;
+	// Must be set to false to prevent the renders from using different VSM page pool sizes leading to unnecessary reallocations.
+	ViewInitOptions.bIsSceneCapture = false;
 
 	if (ViewFamily.Scene->GetWorld() != nullptr && ViewFamily.Scene->GetWorld()->GetWorldSettings() != nullptr)
 	{
@@ -682,14 +711,28 @@ static FSceneRenderer* CreateWaterInfoSceneRenderer(const FCreateWaterInfoSceneR
 
 	ViewFamily.SetScreenPercentageInterface(new FLegacyScreenPercentageDriver(ViewFamily, 1.f));
 
-	ViewFamily.ViewExtensions = GEngine->ViewExtensions->GatherActiveExtensions(FSceneViewExtensionContext(Params.Scene));
+	ViewFamily.ViewExtensions = GEngine->ViewExtensions->GatherActiveExtensions(FSceneViewExtensionContext(ViewFamily.Scene));
 	for (const FSceneViewExtensionRef& Extension : ViewFamily.ViewExtensions)
 	{
 		Extension->SetupViewFamily(ViewFamily);
 		Extension->SetupView(ViewFamily, *View);
 	}
 
-	return FSceneRenderer::CreateSceneRenderer(&ViewFamily, nullptr);
+	FSceneRenderer* SceneRenderer = FSceneRenderer::CreateSceneRenderer(&ViewFamily, nullptr);
+	FSceneViewStateReference* ViewStateReference = new FSceneViewStateReference();
+	ViewStateReference->Allocate(ViewFamily.Scene->GetFeatureLevel());
+	FSceneViewStateInterface* ViewStateInterface = ViewStateReference->GetReference();
+
+	if (UseVirtualShadowMaps(SceneRenderer->ShaderPlatform, ViewFamily.Scene->GetFeatureLevel()))
+	{
+		// This is needed to prevent the default cache being invalidated on every frame.
+		if (ViewStateInterface)
+		{
+			ViewStateInterface->AddVirtualShadowMapCache(ViewFamily.Scene);
+		}
+	}
+
+	return SceneRenderer;
 }
 
 void UpdateWaterInfoRendering(
@@ -705,11 +748,13 @@ void UpdateWaterInfoRendering(
 	{
 		return;
 	}
-	const FVector ZoneExtent = Context.ZoneToRender->GetTessellatedWaterMeshExtent();
+	const FVector ZoneExtent = Context.ZoneToRender->GetDynamicWaterMeshExtent();
 
-	FVector ViewLocation = Context.ZoneToRender->GetTessellatedWaterMeshCenter();
+	FVector ViewLocation = Context.ZoneToRender->GetDynamicWaterMeshCenter();
 	ViewLocation.Z = Context.CaptureZ;
 
+	const FBox2D CaptureBounds(FVector2D(ViewLocation - ZoneExtent), FVector2D(ViewLocation + ZoneExtent));
+
 	// Zone rendering always happens facing towards negative z.
 	const FVector LookAt = ViewLocation - FVector(0.f, 0.f, 1.f);
 	
@@ -749,20 +794,34 @@ void UpdateWaterInfoRendering(
 	CreateSceneRendererParams.ShowOnlyPrimitives = MoveTemp(ComponentsToRenderInDepthPass);
 	FSceneRenderer* DepthRenderer = CreateWaterInfoSceneRenderer(CreateSceneRendererParams);
 	
-	TSet<FPrimitiveComponentId> ComponentsToRenderInColorAndDilationPass;
+	TSet<FPrimitiveComponentId> ComponentsToRenderInColorPass;
+	TSet<FPrimitiveComponentId> ComponentsToRenderInDilationPass;
 	if (Context.WaterBodies.Num() > 0)
 	{
-		ComponentsToRenderInColorAndDilationPass.Reserve(Context.WaterBodies.Num());
+		ComponentsToRenderInColorPass.Reserve(Context.WaterBodies.Num());
+		ComponentsToRenderInDilationPass.Reserve(Context.WaterBodies.Num());
 		for (const UWaterBodyComponent* WaterBodyToRender : Context.WaterBodies)
 		{
-			ComponentsToRenderInColorAndDilationPass.Add(WaterBodyToRender->ComponentId);
+			if (!IsValid(WaterBodyToRender))
+			{
+				continue;
+			}
+
+			// Perform our own simple culling based on the known Capture bounds:
+			const FBox WaterBodyBounds = WaterBodyToRender->Bounds.GetBox();
+			if (CaptureBounds.Intersect(FBox2D(FVector2D(WaterBodyBounds.Min), FVector2D(WaterBodyBounds.Max))))
+			{
+				ComponentsToRenderInColorPass.Add(WaterBodyToRender->GetWaterInfoMeshComponent()->ComponentId);
+				ComponentsToRenderInDilationPass.Add(WaterBodyToRender->GetDilatedWaterInfoMeshComponent()->ComponentId);
+			}
 		}
 	}
 	CreateSceneRendererParams.CaptureSource = SCS_SceneColorSceneDepth;
-	CreateSceneRendererParams.ShowOnlyPrimitives = MoveTemp(ComponentsToRenderInColorAndDilationPass);
+	CreateSceneRendererParams.ShowOnlyPrimitives = MoveTemp(ComponentsToRenderInColorPass);
 	FSceneRenderer* ColorRenderer = CreateWaterInfoSceneRenderer(CreateSceneRendererParams);
 
 	CreateSceneRendererParams.CaptureSource = SCS_DeviceDepth;
+	CreateSceneRendererParams.ShowOnlyPrimitives = MoveTemp(ComponentsToRenderInDilationPass);
 	FSceneRenderer* DilationRenderer = CreateWaterInfoSceneRenderer(CreateSceneRendererParams);
 
 	FTextureRenderTargetResource* TextureRenderTargetResource = Context.TextureRenderTarget->GameThread_GetRenderTargetResource();
diff --git a/Engine/Plugins/Experimental/Water/Source/Runtime/Private/WaterMeshComponent.cpp b/Engine/Plugins/Experimental/Water/Source/Runtime/Private/WaterMeshComponent.cpp
--- a/Engine/Plugins/Experimental/Water/Source/Runtime/Private/WaterMeshComponent.cpp
+++ b/Engine/Plugins/Experimental/Water/Source/Runtime/Private/WaterMeshComponent.cpp
@@ -606,18 +606,6 @@ void UWaterMeshComponent::PostEditChangeProperty(FPropertyChangedEvent& Property
 	{
 		const FName PropertyName = PropertyThatChanged->GetFName();
 
-		// Properties that require water body geometry to be rebuilt since they depend on the mesh grid layout
-		if (PropertyName == GET_MEMBER_NAME_CHECKED(UWaterMeshComponent, TileSize)
-			|| PropertyName == GET_MEMBER_NAME_CHECKED(UWaterMeshComponent, ExtentInTiles)
-			|| PropertyName == TEXT("RelativeScale3D")
-			|| PropertyName == TEXT("RelativeLocation")
-			|| PropertyName == TEXT("RelativeRotation"))
-		{
-			AWaterZone* WaterZone = GetOwner<AWaterZone>();
-			check(WaterZone);
-			WaterZone->MarkForRebuild(EWaterZoneRebuildFlags::UpdateWaterBodyLODSections);
-		}
-
 		// Properties that needs the scene proxy to be rebuilt
 		if (PropertyName == GET_MEMBER_NAME_CHECKED(UWaterMeshComponent, LODScale)
 			|| PropertyName == GET_MEMBER_NAME_CHECKED(UWaterMeshComponent, TessellationFactor)
@@ -633,15 +621,5 @@ void UWaterMeshComponent::PostEditChangeProperty(FPropertyChangedEvent& Property
 		}
 	}
 }
-
-void UWaterMeshComponent::PostEditComponentMove(bool bFinished)
-{
-	if (bFinished)
-	{
-		AWaterZone* WaterZone = GetOwner<AWaterZone>();
-		check(WaterZone);
-		WaterZone->MarkForRebuild(EWaterZoneRebuildFlags::UpdateWaterBodyLODSections);
-	}
-}
 #endif
 
diff --git a/Engine/Plugins/Experimental/Water/Source/Runtime/Private/WaterUtils.cpp b/Engine/Plugins/Experimental/Water/Source/Runtime/Private/WaterUtils.cpp
--- a/Engine/Plugins/Experimental/Water/Source/Runtime/Private/WaterUtils.cpp
+++ b/Engine/Plugins/Experimental/Water/Source/Runtime/Private/WaterUtils.cpp
@@ -124,3 +124,17 @@ float FWaterUtils::GetWaterMaxFlowVelocity(bool bIsRenderThread)
 {
 	return (bIsRenderThread ? CVarWaterMaxFlowVelocity.GetValueOnRenderThread() : CVarWaterMaxFlowVelocity.GetValueOnGameThread());
 }
+
+FVector4f FWaterUtils::PackFlowData(float VelocityMagnitude, float DirectionAngle)
+{
+	check((DirectionAngle >= 0.f) && (DirectionAngle <= TWO_PI));
+
+	const float MaxVelocity = GetWaterMaxFlowVelocity(false);
+
+	FVector4f Result;
+	Result.X = FMath::Clamp(VelocityMagnitude / MaxVelocity, 0.f, 1.f);
+	Result.Y = FMath::Clamp(DirectionAngle / TWO_PI, 0.f, 1.f);
+
+	return Result;
+}
+
diff --git a/Engine/Plugins/Experimental/Water/Source/Runtime/Private/WaterViewExtension.cpp b/Engine/Plugins/Experimental/Water/Source/Runtime/Private/WaterViewExtension.cpp
--- a/Engine/Plugins/Experimental/Water/Source/Runtime/Private/WaterViewExtension.cpp
+++ b/Engine/Plugins/Experimental/Water/Source/Runtime/Private/WaterViewExtension.cpp
@@ -41,9 +41,6 @@ void FWaterViewExtension::SetupView(FSceneViewFamily& InViewFamily, FSceneView&
 	const TWeakObjectPtr<UWorld> WorldPtr = GetWorld();
 	check(WorldPtr.IsValid())
 
-	// #todo_water: hacky way to prevent updating from within the water info passes. Alternative solution might be to set a flag from within the info render
-	// similarly to the SetWithinWaterInfoPass functions for water body components 
-	//This won't be necessary when waterinfo passes don't utilize scene renderers.
 	static bool bUpdatingWaterInfo = false;
 	if (!bUpdatingWaterInfo)
 	{
@@ -52,43 +49,44 @@ void FWaterViewExtension::SetupView(FSceneViewFamily& InViewFamily, FSceneView&
 		for (const TPair<AWaterZone*, UE::WaterInfo::FRenderingContext>& Pair : WaterInfoContextsToRender)
 		{
 			AWaterZone* WaterZone = Pair.Key;
+			check(WaterZone);
 
-			if (WaterZone->IsNonTessellatedLODMeshEnabled())
+			if (WaterZone->IsLocalOnlyTessellationEnabled())
 			{
-				const UWaterMeshComponent* WaterMesh= WaterZone->GetWaterMeshComponent();
+				UWaterMeshComponent* WaterMesh= WaterZone->GetWaterMeshComponent();
 				check(WaterMesh);
 
-				const float SectionSize = WaterZone->GetNonTessellatedLODSectionSize();
+				const double TileSize = WaterMesh->GetTileSize();
 
-				FVector TessellatedWaterMeshCenter = ViewLocation.GridSnap(SectionSize);
+				const FVector WaterMeshCenter = ViewLocation.GridSnap(TileSize);
 
-				WaterZone->SetTessellatedWaterMeshCenter(TessellatedWaterMeshCenter);
+				const FVector WaterMeshHalfExtent = WaterZone->GetDynamicWaterMeshExtent() / 2.0;
+				const FBox2D WaterMeshBounds(FVector2D(WaterMeshCenter - WaterMeshHalfExtent), FVector2D(WaterMeshCenter + WaterMeshHalfExtent));
 
-				const FVector TessellatedWaterMeshHalfExtent = WaterZone->GetTessellatedWaterMeshExtent() / 2.0;
-				const FBox2D TessellatedWaterMeshBounds(FVector2D(TessellatedWaterMeshCenter - TessellatedWaterMeshHalfExtent), FVector2D(TessellatedWaterMeshCenter + TessellatedWaterMeshHalfExtent));
+				WaterZone->SetLocalTessellationCenter(WaterMeshCenter);
+				WaterMesh->PushTessellatedWaterMeshBoundsToPoxy(WaterMeshBounds.ExpandBy(2 * TileSize));
 
-				WaterZone->GetWaterMeshComponent()->PushTessellatedWaterMeshBoundsToPoxy(TessellatedWaterMeshBounds);
-
-				WaterZone->ForEachWaterBodyComponent([TessellatedWaterMeshBounds](UWaterBodyComponent* WaterBodyComponent)
-				{
-					WaterBodyComponent->PushTessellatedWaterMeshBoundsToProxy(TessellatedWaterMeshBounds);
-
-					return true;
-				});
-
-				const FBox2D CurrentCellBounds(FVector2D(ViewLocation) - FVector2D(SectionSize / 2.0), FVector2D(ViewLocation) + FVector2D(SectionSize / 2.0));
+				const FBox2D CurrentCellBounds(FVector2D(ViewLocation) - FVector2D(TileSize / 2.0), FVector2D(ViewLocation) + FVector2D(TileSize / 2.0));
 				UpdateBounds.Add(WaterZone, CurrentCellBounds);
 			}
 
+			// Update the material instances now that the tessellated mesh bounds have changed.
+			WaterZone->ForEachWaterBodyComponent([](UWaterBodyComponent* WaterBodyComponent)
+			{
+				WaterBodyComponent->UpdateMaterialInstances();
+				return true;
+			});
+
 			const UE::WaterInfo::FRenderingContext& Context(Pair.Value);
 			UE::WaterInfo::UpdateWaterInfoRendering(WorldPtr.Get()->Scene, Context);
 		}
 
 		WaterInfoContextsToRender.Empty();
 
+		// Check if the view location is no longer within the current update bounds of a water zone and if so, queue an update for it.
 		for (AWaterZone* WaterZone : TActorRange<AWaterZone>(WorldPtr.Get()))
 		{
-			if (WaterZone->IsNonTessellatedLODMeshEnabled())
+			if (WaterZone->IsLocalOnlyTessellationEnabled())
 			{
 				const FBox2D* ZoneUpdateBounds = UpdateBounds.Find(WaterZone);
 				if (ZoneUpdateBounds == nullptr || !ZoneUpdateBounds->IsInside(FVector2D(ViewLocation)))
diff --git a/Engine/Plugins/Experimental/Water/Source/Runtime/Private/WaterZoneActor.cpp b/Engine/Plugins/Experimental/Water/Source/Runtime/Private/WaterZoneActor.cpp
--- a/Engine/Plugins/Experimental/Water/Source/Runtime/Private/WaterZoneActor.cpp
+++ b/Engine/Plugins/Experimental/Water/Source/Runtime/Private/WaterZoneActor.cpp
@@ -45,7 +45,7 @@ AWaterZone::AWaterZone(const FObjectInitializer& Initializer)
 	WaterMesh = CreateDefaultSubobject<UWaterMeshComponent>(TEXT("WaterMesh"));
 	SetRootComponent(WaterMesh);
 	ZoneExtent = FVector2D(51200., 51200.);
-	TessellatedWaterMeshExtent = FVector(35000., 35000., 10000.);
+	LocalTessellationExtent = FVector(35000., 35000., 10000.);
 	
 #if	WITH_EDITOR
 	// Setup bounds component
@@ -70,6 +70,8 @@ AWaterZone::AWaterZone(const FObjectInitializer& Initializer)
 #endif // WITH_EDITOR
 
 #if WITH_EDITORONLY_DATA
+	TessellatedWaterMeshExtent_DEPRECATED = FVector(35000., 35000., 10000.);
+
 	bIsSpatiallyLoaded = false;
 #endif
 }
@@ -134,6 +136,12 @@ void AWaterZone::PostLoad()
 	Super::PostLoad();
 
 #if WITH_EDITORONLY_DATA
+	if (GetLinkerCustomVersion(FFortniteMainBranchObjectVersion::GUID) < FFortniteMainBranchObjectVersion::WaterBodyStaticMeshComponents)
+	{
+		LocalTessellationExtent = TessellatedWaterMeshExtent_DEPRECATED;
+		bEnableLocalOnlyTessellation = bEnableNonTesselatedLODMesh_DEPRECATED;
+	}
+
 	// WaterMesh ExtentInTiles returns the half extent.
 	const FVector2D ExtentInTiles = 2.0 * FVector2D(WaterMesh->GetExtentInTiles());
 	ZoneExtent = FVector2D(ExtentInTiles * WaterMesh->GetTileSize());
@@ -149,17 +157,6 @@ void AWaterZone::DeclareConstructClasses(TArray<FTopLevelAssetPath>& OutConstruc
 }
 #endif
 
-void AWaterZone::PostRegisterAllComponents()
-{
-	Super::PostRegisterAllComponents();
-
-	if (UpdateOverlappingWaterBodies())
-	{
-		MarkForRebuild(EWaterZoneRebuildFlags::UpdateWaterBodyLODSections);
-	}
-
-}
-
 void AWaterZone::MarkForRebuild(EWaterZoneRebuildFlags Flags)
 {
 	if (EnumHasAnyFlags(Flags, EWaterZoneRebuildFlags::UpdateWaterMesh))
@@ -172,11 +169,6 @@ void AWaterZone::MarkForRebuild(EWaterZoneRebuildFlags Flags)
 		UE_LOG(LogWater, Verbose, TEXT("AWaterZone::MarkForRebuild (UpdateWaterInfoTexture)"));
 		bNeedsWaterInfoRebuild = true;
 	}
-	if (EnumHasAnyFlags(Flags, EWaterZoneRebuildFlags::UpdateWaterBodyLODSections) && bEnableNonTessellatedLODMesh)
-	{
-		UE_LOG(LogWater, Verbose, TEXT("AWaterZone::MarkForRebuild (UpdateWaterBodyLODSections)"));
-		bNeedsNonTessellatedMeshRebuild = true;
-	}
 }
 
 void AWaterZone::ForEachWaterBodyComponent(TFunctionRef<bool(UWaterBodyComponent*)> Predicate) const
@@ -241,16 +233,6 @@ void AWaterZone::Update()
 			bNeedsWaterInfoRebuild = false;
 		}
 	}
-	if (bNeedsNonTessellatedMeshRebuild)
-	{
-		ForEachWaterBodyComponent([](UWaterBodyComponent* WaterBodyComponent)
-		{
-			WaterBodyComponent->UpdateNonTessellatedMeshSections();
-			return true;
-		});
-
-		bNeedsNonTessellatedMeshRebuild = false;
-	}
 	
 	if (WaterMesh)
 	{
@@ -270,10 +252,6 @@ void AWaterZone::PostEditMove(bool bFinished)
 
 	// Ensure that the water mesh is rebuilt if it moves
 	EWaterZoneRebuildFlags RebuildFlags = EWaterZoneRebuildFlags::UpdateWaterInfoTexture | EWaterZoneRebuildFlags::UpdateWaterMesh;
-	if (bFinished)
-	{
-		RebuildFlags |= EWaterZoneRebuildFlags::UpdateWaterBodyLODSections;
-	}
 
 	UpdateOverlappingWaterBodies();
 
@@ -321,15 +299,11 @@ void AWaterZone::PostEditChangeProperty(FPropertyChangedEvent& PropertyChangedEv
 	{
 		MarkForRebuild(EWaterZoneRebuildFlags::UpdateWaterInfoTexture);
 	}
-	else if (PropertyName == GET_MEMBER_NAME_CHECKED(AWaterZone, TessellatedWaterMeshExtent))
+	else if (PropertyName == GET_MEMBER_NAME_CHECKED(AWaterZone, LocalTessellationExtent))
 	{
 		MarkForRebuild(EWaterZoneRebuildFlags::All);
 	}
-	else if (PropertyName == GET_MEMBER_NAME_CHECKED(AWaterZone, bEnableNonTessellatedLODMesh))
-	{
-		MarkForRebuild(EWaterZoneRebuildFlags::All);
-	}
-	else if (PropertyName == GET_MEMBER_NAME_CHECKED(AWaterZone, NonTessellatedLODSectionScale))
+	else if (PropertyName == GET_MEMBER_NAME_CHECKED(AWaterZone, bEnableLocalOnlyTessellation))
 	{
 		MarkForRebuild(EWaterZoneRebuildFlags::All);
 	}
@@ -538,22 +512,17 @@ bool AWaterZone::UpdateWaterInfoTexture()
 			WaterViewExtension.Pin()->MarkWaterInfoTextureForRebuild(Context);
 		}
 
-		for (UWaterBodyComponent* Component : WaterBodiesToRender)
-		{
-			Component->UpdateMaterialInstances();
-		}
-
 		UE_LOG(LogWater, Verbose, TEXT("Queued Water Info texture update"));
 	}
 
 	return true;
 }
 
-FVector AWaterZone::GetTessellatedWaterMeshCenter() const
+FVector AWaterZone::GetDynamicWaterMeshCenter() const
 {
-	if (IsNonTessellatedLODMeshEnabled())
+	if (IsLocalOnlyTessellationEnabled())
 	{
-		return TessellatedWaterMeshCenter;
+		return LocalTessellationCenter;
 	}
 	else
 	{
@@ -561,11 +530,11 @@ FVector AWaterZone::GetTessellatedWaterMeshCenter() const
 	}
 }
 
-FVector AWaterZone::GetTessellatedWaterMeshExtent() const
+FVector AWaterZone::GetDynamicWaterMeshExtent() const
 {
-	if (IsNonTessellatedLODMeshEnabled())
+	if (IsLocalOnlyTessellationEnabled())
 	{
-		return TessellatedWaterMeshExtent.GridSnap(GetNonTessellatedLODSectionSize());
+		return LocalTessellationExtent;
 	}
 	else
 	{
@@ -574,12 +543,6 @@ FVector AWaterZone::GetTessellatedWaterMeshExtent() const
 	}
 }
 
-float AWaterZone::GetNonTessellatedLODSectionSize() const
-{
-	check(WaterMesh)
-	return WaterMesh->GetTileSize() * NonTessellatedLODSectionScale;
-}
-
 void AWaterZone::OnLevelAddedToWorld(ULevel* InLevel, UWorld* InWorld)
 {
 	OnLevelChanged(InLevel, InWorld);
diff --git a/Engine/Plugins/Experimental/Water/Source/Runtime/Public/WaterBodyActor.h b/Engine/Plugins/Experimental/Water/Source/Runtime/Public/WaterBodyActor.h
--- a/Engine/Plugins/Experimental/Water/Source/Runtime/Public/WaterBodyActor.h
+++ b/Engine/Plugins/Experimental/Water/Source/Runtime/Public/WaterBodyActor.h
@@ -8,6 +8,7 @@
 #include "WaterBodyActor.generated.h"
 
 class UWaterSplineComponent;
+class UWaterBodyInfoMeshComponent;
 class AWaterBodyIsland;
 class AWaterBodyExclusionVolume;
 class ALandscapeProxy;
@@ -112,6 +113,15 @@ protected:
 	UPROPERTY(EditAnywhere, BlueprintReadOnly, Instanced, Category = Wave, DisplayName = "Waves Source", meta = (Tooltip = ""))
 	TObjectPtr<UWaterWavesBase> WaterWaves = nullptr;
 
+	UPROPERTY(NonTransactional, TextExportTransient, NonPIEDuplicateTransient)
+	TObjectPtr<UWaterBodyInfoMeshComponent> WaterInfoMeshComponent;
+
+	UPROPERTY(NonTransactional, TextExportTransient, NonPIEDuplicateTransient)
+	TObjectPtr<UWaterBodyInfoMeshComponent> DilatedWaterInfoMeshComponent;
+
+	UPROPERTY(NonTransactional, TextExportTransient, NonPIEDuplicateTransient)
+	TArray<TObjectPtr<UWaterBodyStaticMeshComponent>> WaterBodyStaticMeshComponents;
+
 #if WITH_EDITOR
 	virtual void PostEditMove(bool bFinished) override;
 	virtual void PreEditChange(FProperty* PropertyThatWillChange) override;
@@ -121,9 +131,15 @@ protected:
 
 	void SetWaterWavesInternal(UWaterWavesBase* InWaterWaves);
 
+	const TArray<TObjectPtr<UWaterBodyStaticMeshComponent>>& GetWaterBodyStaticMeshComponents() const { return WaterBodyStaticMeshComponents; }
+
+	/** Sets up a new list of water body static mesh components. */
+	void SetWaterBodyStaticMeshComponents(TArrayView<TObjectPtr<UWaterBodyStaticMeshComponent>> NewComponentList, TConstArrayView<TObjectPtr<UWaterBodyStaticMeshComponent>> ComponentsToUnregister = {});
+
 // ----------------------------------------------------------------------------------
 // Deprecated
 
+#pragma region Deprecated
 public:
 	UE_DEPRECATED(4.27, "Moved to WaterBodyComponent")
 	UFUNCTION(BlueprintCallable, Category = Rendering, meta = (DeprecatedFunction))
@@ -250,6 +266,8 @@ protected:
 	UPROPERTY()
 	float ShapeDilation_DEPRECATED = 4096.0f;
 #endif // WITH_EDITORONLY_DATA
+
+#pragma endregion // deprecated
 };
 
 #if UE_ENABLE_INCLUDE_ORDER_DEPRECATED_IN_5_2
diff --git a/Engine/Plugins/Experimental/Water/Source/Runtime/Public/WaterBodyComponent.h b/Engine/Plugins/Experimental/Water/Source/Runtime/Public/WaterBodyComponent.h
--- a/Engine/Plugins/Experimental/Water/Source/Runtime/Public/WaterBodyComponent.h
+++ b/Engine/Plugins/Experimental/Water/Source/Runtime/Public/WaterBodyComponent.h
@@ -7,6 +7,7 @@
 #include "WaterBodyWeightmapSettings.h"
 #include "WaterBodyHeightmapSettings.h"
 #include "WaterCurveSettings.h"
+#include "WaterBodyStaticMeshSettings.h"
 #include "WaterSplineMetadata.h"
 #include "WaterBodyTypes.h"
 
@@ -32,6 +33,8 @@ struct FPostProcessVolumeProperties;
 
 class UWaterSplineComponent;
 struct FOnWaterSplineDataChangedParams;
+class UWaterBodyStaticMeshComponent;
+class UWaterBodyInfoMeshComponent;
 class AWaterBodyIsland;
 class AWaterBodyExclusionVolume;
 class AWaterZone;
@@ -39,7 +42,6 @@ class ALandscapeProxy;
 class UMaterialInstanceDynamic;
 class FTokenizedMessage;
 namespace UE::Geometry { class FDynamicMesh3; }
-struct FDynamicMeshVertex;
 struct FMeshDescription;
 
 // ----------------------------------------------------------------------------------
@@ -91,25 +93,6 @@ enum class EWaterBodyStatus : uint8
 };
 
 
-// ----------------------------------------------------------------------------------
-
-struct FWaterBodyMeshSection
-{
-public:
-	FWaterBodyMeshSection();
-	~FWaterBodyMeshSection();
-
-	bool IsValid() const { return (bool)SectionBounds.bIsValid; }
-
-	uint32 GetAllocatedSize() const;
-
-	TArray<FDynamicMeshVertex> Vertices;
-	TArray<uint32> Indices;
-
-	FBox2D SectionBounds;
-};
-
-
 // ----------------------------------------------------------------------------------
 
 struct FOnWaterBodyChangedParams
@@ -141,14 +124,13 @@ class WATER_API UWaterBodyComponent : public UPrimitiveComponent
 
 	friend class AWaterBody;
 	friend class FWaterBodySceneProxy;
+	friend class FWaterBodyMeshBuilder;
 public:
 	virtual bool AffectsLandscape() const;
 	virtual bool AffectsWaterMesh() const;
 	virtual bool AffectsWaterInfo() const;
 	virtual bool CanEverAffectWaterMesh() const { return true; }
 
-	UE_DEPRECATED(5.1, "Use the version of this function taking FOnWaterBodyChangedParams in parameter")
-	void UpdateAll(bool bShapeOrPositionChanged) {}
 	void UpdateAll(const FOnWaterBodyChangedParams& InParams);
 
 #if WITH_EDITOR
@@ -164,6 +146,10 @@ public:
 	virtual FMeshDescription GetHLODMeshDescription() const;
 	virtual UMaterialInterface* GetHLODMaterial() const;
 	virtual void SetHLODMaterial(UMaterialInterface* InMaterial);
+
+	void UpdateWaterBodyRenderData();
+
+	void SetWaterBodyStaticMeshEnabled(bool bEnabled);
 #endif //WITH_EDITOR
 
 	/** Returns whether the body supports waves */
@@ -181,7 +167,7 @@ public:
 
 	/** Retrieves the list of primitive components that this water body uses when not being rendered by the water mesh (e.g. the static mesh component used when WaterMeshOverride is specified) */
 	UFUNCTION(BlueprintCallable, Category = Rendering)
-	virtual TArray<UPrimitiveComponent*> GetStandardRenderableComponents() const { return TArray<UPrimitiveComponent*>(); }
+	virtual TArray<UPrimitiveComponent*> GetStandardRenderableComponents() const { return {}; }
 
 	/** Returns the body's collision component bounds */
 	virtual FBox GetCollisionComponentBounds() const;
@@ -230,19 +216,22 @@ public:
 	UFUNCTION(BlueprintCallable, Category = Rendering)
 	UMaterialInterface* GetWaterMaterial() const { return WaterMaterial; }
 
-	/** Returns material used to render the water LOD mesh sections */
-	UMaterialInterface* GetWaterLODMaterial() const { return WaterLODMaterial; }
+	/** Returns material used to render the water as a static mesh */
+	UMaterialInterface* GetWaterStaticMeshMaterial() const { return WaterStaticMeshMaterial; }
 
 	/** Sets water material */
 	void SetWaterMaterial(UMaterialInterface* InMaterial);
 
+	/** Sets water static mesh material */
+	void SetWaterStaticMeshMaterial(UMaterialInterface* InMaterial);
+
 	/** Returns water MID */
 	UFUNCTION(BlueprintCallable, Category = Rendering)
 	UMaterialInstanceDynamic* GetWaterMaterialInstance();
 
-	/** Returns water LOD MID */
+	/** Returns water static mesh MID */
 	UFUNCTION(BlueprintCallable, Category = Rendering)
-	UMaterialInstanceDynamic* GetWaterLODMaterialInstance();
+	UMaterialInstanceDynamic* GetWaterStaticMeshMaterialInstance();
 
 	/** Returns under water post process MID */
 	UFUNCTION(BlueprintCallable, Category = Rendering)
@@ -331,10 +320,6 @@ public:
 	virtual void OnUnregister() override;
 	virtual void PostDuplicate(bool bDuplicateForPie) override;
 
-	UE_DEPRECATED(5.2, "Use version which takes FOnWaterBodyChangedParams")
-	UFUNCTION(BlueprintCallable, Category=Water, meta=(Deprecated = "5.2"))
-	void OnWaterBodyChanged(bool bShapeOrPositionChanged, bool bWeightmapSettingsChanged = false);
-
 	void OnWaterBodyChanged(const FOnWaterBodyChangedParams& InParams);
 
 	/** Fills wave-related information at the given world position and for this water depth.
@@ -397,21 +382,33 @@ public:
 	/** Registers or this water body with corresponding overlapping water zones and unregisters it from any old zones if they are no longer overlapping. */
 	void UpdateWaterZones();
 
-	void UpdateWaterBodyRenderData();
-	void UpdateNonTessellatedMeshSections();
-
-	virtual FPrimitiveSceneProxy* CreateSceneProxy() override;
-	virtual void GetUsedMaterials(TArray<UMaterialInterface*>& OutMaterials, bool bGetDebugMaterials) const override;
-
-	void PushTessellatedWaterMeshBoundsToProxy(const FBox2D& TessellatedWaterMeshBounds);
-	
 	/** Set the navigation area class */
 	void SetNavAreaClass(TSubclassOf<UNavAreaBase> NewWaterNavAreaClass) { WaterNavAreaClass = NewWaterNavAreaClass; }
 
+	/* Generates the mesh representation of the water body */
+	virtual bool GenerateWaterBodyMesh(UE::Geometry::FDynamicMesh3& OutMesh, UE::Geometry::FDynamicMesh3* OutDilatedMesh = nullptr) const;
+
+	UWaterBodyInfoMeshComponent* GetWaterInfoMeshComponent() const;
+	UWaterBodyInfoMeshComponent* GetDilatedWaterInfoMeshComponent() const;
+
+	const FWaterBodyStaticMeshSettings& GetWaterBodyStaticMeshSettings() const { return StaticMeshSettings; }
 	
 	UE_DEPRECATED(5.1, "Renamed to CanEverAffectWaterMesh")
 	virtual bool CanAffectWaterMesh() const { return true; }
 
+	UE_DEPRECATED(5.1, "Use the version of this function taking FOnWaterBodyChangedParams in parameter")
+	void UpdateAll(bool bShapeOrPositionChanged) {}
+
+	UE_DEPRECATED(5.2, "Use version which takes FOnWaterBodyChangedParams")
+	UFUNCTION(BlueprintCallable, Category=Water, meta=(Deprecated = "5.2"))
+	void OnWaterBodyChanged(bool bShapeOrPositionChanged, bool bWeightmapSettingsChanged = false);
+
+	UE_DEPRECATED(5.3, "Renamed to GetWaterStaticMeshMaterial")
+	UMaterialInterface* GetWaterLODMaterial() const { return GetWaterStaticMeshMaterial(); }
+
+	UE_DEPRECATED(5.3, "Renamed to SetWaterStaticMeshMaterial")
+	void SetWaterLODMaterial(UMaterialInterface* InMaterial) { SetWaterStaticMeshMaterial(InMaterial); }
+
 protected:
 	//~ Begin UActorComponent interface.
 	virtual bool IsHLODRelevant() const override;
@@ -442,8 +439,8 @@ protected:
 
 	virtual void OnUpdateBody(bool bWithExclusionVolumes) {}
 
-	/* Generates the meshes used to render the Water Info texture */
-	virtual bool GenerateWaterBodyMesh(UE::Geometry::FDynamicMesh3& OutMesh, UE::Geometry::FDynamicMesh3* OutDilatedMesh = nullptr) const { checkf(!AffectsWaterInfo(), TEXT("WaterBodyComponent affects water info but does not implement GenerateWaterBodyMesh!")); return false; }
+	/** Called when the WaterBodyActor has had all its components registered. */
+	virtual void OnPostRegisterAllComponents() {}
 
 	/** Returns navigation area class */
 	TSubclassOf<UNavAreaBase> GetNavAreaClass() const { return WaterNavAreaClass; }
@@ -472,13 +469,15 @@ protected:
 	/** Validates this component's data */
 	virtual void CheckForErrors() override;
 
-	virtual TArray<TSharedRef<FTokenizedMessage>> CheckWaterBodyStatus() const;
+	virtual TArray<TSharedRef<FTokenizedMessage>> CheckWaterBodyStatus();
 
 	virtual const TCHAR* GetWaterSpriteTextureName() const { return TEXT("/Water/Icons/WaterSprite"); }
 
 	virtual bool IsIconVisible() const { return true; }
 
 	virtual FVector GetWaterSpriteLocation() const { return GetComponentLocation(); }
+
+	virtual void OnWaterBodyRenderDataUpdated();
 #endif // WITH_EDITOR
 
 	EWaterBodyQueryFlags CheckAndAjustQueryFlags(EWaterBodyQueryFlags InQueryFlags) const;
@@ -486,7 +485,7 @@ protected:
 	void UpdateExclusionVolumes();
 	bool UpdateWaterHeight();
 	virtual void CreateOrUpdateWaterMID();
-	void CreateOrUpdateWaterLODMID();
+	void CreateOrUpdateWaterStaticMeshMID();
 	void CreateOrUpdateUnderwaterPostProcessMID();
 	void CreateOrUpdateWaterInfoMID();
 	void PrepareCurrentPostProcessSettings();
@@ -495,9 +494,6 @@ protected:
 	void RequestGPUWaveDataUpdate();
 	EObjectFlags GetTransientMIDFlags() const; 
 	void DeprecateData();
-	void RebuildWaterBodyInfoMesh();
-	void RebuildWaterBodyLODSections();
-	void OnTessellatedWaterMeshBoundsChanged();
 
 	virtual void Serialize(FArchive& Ar) override;
 	virtual void PostLoad() override;
@@ -522,6 +518,9 @@ protected:
 
 	void CreateWaterSpriteComponent();
 
+	void UpdateWaterInfoMeshComponents();
+	void UpdateWaterBodyStaticMeshComponents();
+
 	virtual TSubclassOf<class UHLODBuilder> GetCustomHLODBuilderClass() const override;
 #endif // WITH_EDITOR
 
@@ -573,8 +572,8 @@ public:
 	UPROPERTY(Category = HLOD, EditAnywhere, BlueprintReadOnly, meta = (DisplayName = "Water HLOD Material"))
 	TObjectPtr<UMaterialInterface> WaterHLODMaterial;
 
-	UPROPERTY(Category = Rendering, EditAnywhere, BlueprintReadOnly, meta = (DisplayName = "Water LOD Material"))
-	TObjectPtr<UMaterialInterface> WaterLODMaterial;
+	UPROPERTY(Category = Rendering, EditAnywhere, BlueprintReadOnly)
+	TObjectPtr<UMaterialInterface> WaterStaticMeshMaterial;
 
 	/** Post process material to apply when the camera goes underwater (only available when bGenerateCollisions is true because collisions are needed to detect if it's under water). */
 	UPROPERTY(Category = Rendering, EditAnywhere, BlueprintReadOnly, meta = (EditCondition = "bGenerateCollisions", DisplayAfter = "WaterMaterial"))
@@ -602,6 +601,9 @@ public:
 	UPROPERTY(Category = Terrain, EditAnywhere, BlueprintReadWrite)
 	bool bAffectsLandscape;
 
+	UPROPERTY(Category = Water, EditAnywhere)
+	FWaterBodyStaticMeshSettings StaticMeshSettings;
+
 protected:
 
 	/** Unique Id for accessing (wave, ... ) data in GPU buffers */
@@ -628,8 +630,8 @@ protected:
 	UPROPERTY(Category = Debug, VisibleInstanceOnly, Transient, NonPIEDuplicateTransient, TextExportTransient, meta = (DisplayAfter = "WaterMaterial"))
 	TObjectPtr<UMaterialInstanceDynamic> WaterMID;
 
-	UPROPERTY(Category = Debug, VisibleInstanceOnly, Transient, NonPIEDuplicateTransient, TextExportTransient, meta = (DisplayAfter = "WaterLODMaterial"))
-	TObjectPtr<UMaterialInstanceDynamic> WaterLODMID;
+	UPROPERTY(Category = Debug, VisibleInstanceOnly, Transient, NonPIEDuplicateTransient, TextExportTransient, meta = (DisplayAfter = "WaterStaticMeshMaterial"))
+	TObjectPtr<UMaterialInstanceDynamic> WaterStaticMeshMID;
 
 	UPROPERTY(Category = Debug, VisibleInstanceOnly, Transient, NonPIEDuplicateTransient, TextExportTransient, meta = (DisplayAfter = "UnderwaterPostProcessMaterial"))
 	TObjectPtr<UMaterialInstanceDynamic> UnderwaterPostProcessMID;
@@ -647,7 +649,7 @@ protected:
 	UPROPERTY(Transient)
 	mutable TWeakObjectPtr<ALandscapeProxy> Landscape;
 
-	UPROPERTY(Category = Water, VisibleAnywhere, AdvancedDisplay)
+	UPROPERTY(Category = Water, VisibleInstanceOnly, AdvancedDisplay)
 	TSoftObjectPtr<AWaterZone> OwningWaterZone;
 
 	UPROPERTY(Category = Water, EditAnywhere, BlueprintReadOnly, AdvancedDisplay)
@@ -660,15 +662,6 @@ protected:
 	UPROPERTY(EditAnywhere, BlueprintReadWrite, Category = Navigation, meta = (EditCondition = "bCanAffectNavigation && bGenerateCollisions"))
 	TSubclassOf<UNavAreaBase> WaterNavAreaClass;
 
-	// #todo_water [roey]: serialize and don't rebuild on load.
-	TArray<FWaterBodyMeshSection> WaterBodyMeshSections;
-
-	TArray<FDynamicMeshVertex> WaterBodyMeshVertices;
-	TArray<uint32> WaterBodyMeshIndices;
-
-	TArray<FDynamicMeshVertex> DilatedWaterBodyMeshVertices;
-	TArray<uint32> DilatedWaterBodyMeshIndices;
-
 	/** If the Water Material assigned to this component has Fixed Depth enabled, this is the depth that is passed. */
 	UPROPERTY(Category = Water, EditAnywhere, AdvancedDisplay)
 	double FixedWaterDepth = 512.0;
@@ -679,6 +672,9 @@ protected:
 	UPROPERTY()
 	TArray<TLazyObjectPtr<AWaterBodyExclusionVolume>> ExclusionVolumes_DEPRECATED;
 
+	UPROPERTY()
+	TObjectPtr<UMaterialInterface> WaterLODMaterial_DEPRECATED;
+
 	UPROPERTY()
 	bool bFillCollisionUnderWaterBodiesForNavmesh_DEPRECATED;
 
diff --git a/Engine/Plugins/Experimental/Water/Source/Runtime/Public/WaterBodyCustomComponent.h b/Engine/Plugins/Experimental/Water/Source/Runtime/Public/WaterBodyCustomComponent.h
--- a/Engine/Plugins/Experimental/Water/Source/Runtime/Public/WaterBodyCustomComponent.h
+++ b/Engine/Plugins/Experimental/Water/Source/Runtime/Public/WaterBodyCustomComponent.h
@@ -30,7 +30,7 @@ protected:
 	
 	virtual FPrimitiveSceneProxy* CreateSceneProxy() override;
 #if WITH_EDITOR
-	virtual TArray<TSharedRef<FTokenizedMessage>> CheckWaterBodyStatus() const override;
+	virtual TArray<TSharedRef<FTokenizedMessage>> CheckWaterBodyStatus() override;
 
 	virtual const TCHAR* GetWaterSpriteTextureName() const override;
 
diff --git a/Engine/Plugins/Experimental/Water/Source/Runtime/Public/WaterBodyMeshComponent.h b/Engine/Plugins/Experimental/Water/Source/Runtime/Public/WaterBodyMeshComponent.h
new file mode 100644
--- /dev/null
+++ b/Engine/Plugins/Experimental/Water/Source/Runtime/Public/WaterBodyMeshComponent.h
@@ -0,0 +1,29 @@
+// Copyright Epic Games, Inc. All Rights Reserved.
+
+#pragma once
+
+#include "CoreMinimal.h"
+#include "UObject/ObjectMacros.h"
+#include "Components/StaticMeshComponent.h"
+
+#include "WaterBodyMeshComponent.generated.h"
+
+class UWaterBodyComponent;
+struct FMeshDescription;
+class UMaterialInterface;
+
+/**
+ * Base class for meshes used to render water bodies without relying on the water zone/water mesh.
+ */
+UCLASS(MinimalAPI)
+class UWaterBodyMeshComponent : public UStaticMeshComponent
+{
+	GENERATED_UCLASS_BODY()
+public:
+
+protected:
+	virtual bool CanCreateSceneProxy() const;
+};
+
+
+// #todo_water: UWaterBodyInstancedStaticMeshComponent for the quad meshes to be instanced
diff --git a/Engine/Plugins/Experimental/Water/Source/Runtime/Public/WaterBodyOceanComponent.h b/Engine/Plugins/Experimental/Water/Source/Runtime/Public/WaterBodyOceanComponent.h
--- a/Engine/Plugins/Experimental/Water/Source/Runtime/Public/WaterBodyOceanComponent.h
+++ b/Engine/Plugins/Experimental/Water/Source/Runtime/Public/WaterBodyOceanComponent.h
@@ -26,6 +26,12 @@ public:
 	
 #if WITH_EDITOR
 	void SetCollisionExtents(const FVector& NewExtents);
+
+	void SetOceanExtent(const FVector2D& NewExtents);
+
+	/** Rebuilds the ocean mesh to completely fill the zone to which it belongs. */
+	UFUNCTION(CallInEditor, Category = Water)
+	void FillWaterZoneWithOcean();
 #endif // WITH_EDITOR
 
 	UE_DEPRECATED(5.1, "Oceans no longer rely on the visual extent parameter making this obsolete. Instead they will be guaranteed to fill the entire water zone to which they belong.")
@@ -41,12 +47,17 @@ protected:
 	virtual bool GenerateWaterBodyMesh(UE::Geometry::FDynamicMesh3& OutMesh, UE::Geometry::FDynamicMesh3* OutDilatedMesh = nullptr) const override;
 
 	virtual void PostLoad() override;
+	virtual void OnPostRegisterAllComponents() override;
 	virtual FBoxSphereBounds CalcBounds(const FTransform& LocalToWorld) const;
 
 #if WITH_EDITOR
 	virtual void OnPostEditChangeProperty(FOnWaterBodyChangedParams& InOutOnWaterBodyChangedParams) override;
 
 	virtual const TCHAR* GetWaterSpriteTextureName() const override;
+
+	virtual TArray<TSharedRef<FTokenizedMessage>> CheckWaterBodyStatus();
+
+	virtual void OnWaterBodyRenderDataUpdated() override;
 #endif
 protected:
 	UPROPERTY(NonPIEDuplicateTransient)
@@ -55,14 +66,25 @@ protected:
 	UPROPERTY(NonPIEDuplicateTransient)
 	TArray<TObjectPtr<UOceanCollisionComponent>> CollisionHullSets;
 
-	UPROPERTY()
-	FVector2D VisualExtents_DEPRECATED;
-
 	UPROPERTY(Category = Collision, EditAnywhere, BlueprintReadOnly)
 	FVector CollisionExtents;
 
+	// #todo_water: should we make this editor-only? It's not needed at runtime since the mesh will never regenerate at runtime
+	UPROPERTY(Category = Water, EditAnywhere, BlueprintReadOnly)
+	FVector2D OceanExtents;
+
 	UPROPERTY(Transient)
 	float HeightOffset = 0.0f;
+
+#if WITH_EDITORONLY_DATA
+
+	/** Serialize the offset relative to the water zone to check if the serialized mesh and the current water zone have a mismatch. */
+	UPROPERTY()
+	FVector2D SavedZoneLocation;
+
+	UPROPERTY()
+	FVector2D VisualExtents_DEPRECATED;
+#endif // WITH_EDITORONLY_DATA
 };
 
 #if UE_ENABLE_INCLUDE_ORDER_DEPRECATED_IN_5_2
diff --git a/Engine/Plugins/Experimental/Water/Source/Runtime/Public/WaterBodySceneProxy.h b/Engine/Plugins/Experimental/Water/Source/Runtime/Public/WaterBodySceneProxy.h
deleted file mode 100644
diff --git a/Engine/Plugins/Experimental/Water/Source/Runtime/Public/WaterBodyStaticMeshSettings.h b/Engine/Plugins/Experimental/Water/Source/Runtime/Public/WaterBodyStaticMeshSettings.h
new file mode 100644
--- /dev/null
+++ b/Engine/Plugins/Experimental/Water/Source/Runtime/Public/WaterBodyStaticMeshSettings.h
@@ -0,0 +1,24 @@
+// Copyright Epic Games, Inc. All Rights Reserved.
+
+#pragma once
+
+#include "WaterBodyStaticMeshSettings.generated.h"
+
+USTRUCT()
+struct FWaterBodyStaticMeshSettings
+{
+	GENERATED_BODY()
+
+	/** If enabled, the water body will generate static meshes to use outside of the dynamic mesh generated by the water zone. */
+	UPROPERTY(Category = StaticMesh, EditAnywhere)
+	bool bEnableWaterBodyStaticMesh = false;
+
+	/** If enabled, the water body static mesh will be split into separate meshes to allow them to stream in and out as needed to save memory. */
+	// #todo_water [roey]: Re-enable this when UE-184177 and UE-184211: 
+	//UPROPERTY(Category = StaticMesh, EditAnywhere, meta = (EditCondition = "bEnableWaterBodyStaticMesh"))
+	bool bSectionWaterBodyStaticMesh = false;
+
+	/** The size of each section tile if sectioning is enabled. */
+	UPROPERTY(Category = StaticMesh, EditAnywhere, AdvancedDisplay, meta = (EditCondition = "bSectionWaterBodyStaticMesh"))
+	double SectionSize = 8192.0;
+};
diff --git a/Engine/Plugins/Experimental/Water/Source/Runtime/Public/WaterMeshComponent.h b/Engine/Plugins/Experimental/Water/Source/Runtime/Public/WaterMeshComponent.h
--- a/Engine/Plugins/Experimental/Water/Source/Runtime/Public/WaterMeshComponent.h
+++ b/Engine/Plugins/Experimental/Water/Source/Runtime/Public/WaterMeshComponent.h
@@ -133,7 +133,6 @@ private:
 #if WITH_EDITOR
 	//~ Begin USceneComponent Interface
 	virtual void PostEditChangeProperty(FPropertyChangedEvent& PropertyChangedEvent) override;
-	virtual void PostEditComponentMove(bool bFinished) override;
 	//~ Begin USceneComponent Interface
 #endif
 };
diff --git a/Engine/Plugins/Experimental/Water/Source/Runtime/Public/WaterRuntimeSettings.h b/Engine/Plugins/Experimental/Water/Source/Runtime/Public/WaterRuntimeSettings.h
--- a/Engine/Plugins/Experimental/Water/Source/Runtime/Public/WaterRuntimeSettings.h
+++ b/Engine/Plugins/Experimental/Water/Source/Runtime/Public/WaterRuntimeSettings.h
@@ -41,6 +41,10 @@ public:
 
 	TSubclassOf<UWaterBodyCustomComponent> GetWaterBodyCustomComponentClass() const;
 
+	//#todo_water: how can we put these settins on the editor settings but still access them from the WaterModule?
+	bool ShouldWarnOnMismatchOceanExtent() const { return bWarnOnMismatchOceanExtent; }
+	void SetShouldWarnOnMismatchOceanExtent(bool bEnable) { Modify(); bWarnOnMismatchOceanExtent = bEnable; }
+
 	virtual void PostInitProperties() override;
 #if WITH_EDITOR
 	virtual void PostEditChangeProperty(struct FPropertyChangedEvent& PropertyChangedEvent) override;
@@ -84,6 +88,11 @@ private:
 	UPROPERTY(EditAnywhere, Config, Category = Water,  meta = (MetaClass = "/Script/Water.WaterBodyCustomComponent"))
 	TSubclassOf<UWaterBodyCustomComponent> WaterBodyCustomComponentClass;
 
+	// #todo_water: move this
+	/** If enabled, MapCheck will notify users that their ocean does not completely fill the water zone. This can be desirable in some cases but the default should be to fill completely. */
+	UPROPERTY(EditAnywhere, config, Category = Water)
+	bool bWarnOnMismatchOceanExtent = true;
+
 #if WITH_EDITORONLY_DATA
 	UPROPERTY()
 	float WaterBodyIconWorldSize_DEPRECATED;
diff --git a/Engine/Plugins/Experimental/Water/Source/Runtime/Public/WaterUtils.h b/Engine/Plugins/Experimental/Water/Source/Runtime/Public/WaterUtils.h
--- a/Engine/Plugins/Experimental/Water/Source/Runtime/Public/WaterUtils.h
+++ b/Engine/Plugins/Experimental/Water/Source/Runtime/Public/WaterUtils.h
@@ -44,6 +44,8 @@ struct FWaterUtils
 	static WATER_API bool IsWaterMeshEnabled(bool bIsRenderThread);
 	static WATER_API bool IsWaterMeshRenderingEnabled(bool bIsRenderThread);
 	static WATER_API float GetWaterMaxFlowVelocity(bool bIsRenderThread);
+
+	static FVector4f PackFlowData(float VelocityMagnitude, float DirectionAngle);
 };
 
 #if UE_ENABLE_INCLUDE_ORDER_DEPRECATED_IN_5_2
diff --git a/Engine/Plugins/Experimental/Water/Source/Runtime/Public/WaterZoneActor.h b/Engine/Plugins/Experimental/Water/Source/Runtime/Public/WaterZoneActor.h
--- a/Engine/Plugins/Experimental/Water/Source/Runtime/Public/WaterZoneActor.h
+++ b/Engine/Plugins/Experimental/Water/Source/Runtime/Public/WaterZoneActor.h
@@ -20,7 +20,6 @@ enum class EWaterZoneRebuildFlags
 	None = 0,
 	UpdateWaterInfoTexture = (1 << 1),
 	UpdateWaterMesh = (1 << 2),
-	UpdateWaterBodyLODSections = (1 << 3),
 	All = (~0),
 };
 ENUM_CLASS_FLAGS(EWaterZoneRebuildFlags);
@@ -53,6 +52,12 @@ public:
 
 	uint32 GetVelocityBlurRadius() const { return VelocityBlurRadius; }
 
+	FVector GetDynamicWaterMeshCenter() const;
+	FVector GetDynamicWaterMeshExtent() const;
+
+	bool IsLocalOnlyTessellationEnabled() const { return bEnableLocalOnlyTessellation; }
+	void SetLocalTessellationCenter(const FVector& NewCenter) { LocalTessellationCenter = NewCenter;}
+
 	virtual void BeginPlay() override;
 	virtual void EndPlay(const EEndPlayReason::Type EndPlayReason) override;
 	virtual void PostLoadSubobjects(FObjectInstancingGraph* OuterInstanceGraph) override;
@@ -60,7 +65,6 @@ public:
 #if WITH_EDITORONLY_DATA
 	static void DeclareConstructClasses(TArray<FTopLevelAssetPath>& OutConstructClasses, const UClass* SpecificSubclass);
 #endif
-	virtual void PostRegisterAllComponents() override;
 
 	FVector2f GetWaterHeightExtents() const { return WaterHeightExtents; }
 	float GetGroundZMin() const { return GroundZMin; }
@@ -70,25 +74,27 @@ public:
 	UPROPERTY(Transient, DuplicateTransient, VisibleAnywhere, BlueprintReadOnly, Category = Water)
 	TObjectPtr<UTextureRenderTarget2D> WaterInfoTexture;
 
-	FVector GetTessellatedWaterMeshCenter() const;
-	void SetTessellatedWaterMeshCenter(FVector NewCenter) { TessellatedWaterMeshCenter = NewCenter;}
-
-	FVector GetTessellatedWaterMeshExtent() const;
-
-	uint32 GetNonTessellatedLODScale() const { return NonTessellatedLODSectionScale; }
-	float GetNonTessellatedLODSectionSize() const;
-
-	bool IsNonTessellatedLODMeshEnabled() const { return bEnableNonTessellatedLODMesh; }
-
 #if WITH_EDITOR
 	virtual TUniquePtr<class FWorldPartitionActorDesc> CreateClassActorDesc() const override;
 	virtual FBox GetStreamingBounds() const override;
 #endif //WITH_EDITOR
 
+	UE_DEPRECATED(5.3, "Function renamed to GetDynamicWaterMeshCenter")
+	FVector GetTessellatedWaterMeshCenter() const { return GetDynamicWaterMeshCenter(); }
+
+	UE_DEPRECATED(5.3, "Function renamed to SetLocalTessellationCenter")
+	void SetTessellatedWaterMeshCenter(FVector NewCenter) { SetLocalTessellationCenter(NewCenter); }
+
+	UE_DEPRECATED(5.3, "Function renamed to GetDynamicWaterMeshExtent")
+	FVector GetTessellatedWaterMeshExtent() const { return GetDynamicWaterMeshExtent(); }
+
+	UE_DEPRECATED(5.3, "Function renamed to IsLocalTessellationEnabled.")
+	bool IsNonTessellatedLODMeshEnabled() const { return IsLocalOnlyTessellationEnabled(); }
+
 private:
 
 	/**
-	 * Enqueue's a command on the Water Scene View Extension to re-render the water info on the next frame.
+	 * Enqueues a command on the Water Scene View Extension to re-render the water info on the next frame.
 	 * Returns false if the water info cannot be rendered this frame due to one of the dependencies not being ready yet (ie. a material under On-Demand-Shader-Compilation)
 	 */
 	bool UpdateWaterInfoTexture();
@@ -158,28 +164,22 @@ private:
 	UPROPERTY(Category = Water, EditAnywhere, AdvancedDisplay)
 	int32 VelocityBlurRadius = 1;
 
-	/** Area around the camera covered by the tessellated water mesh when LOD is enabled. */
-	UPROPERTY(Category = NonTessellatedLOD, EditAnywhere, meta = (EditCondition = "bEnableNonTessellatedLODMesh"))
-	FVector TessellatedWaterMeshExtent;
-
-	UPROPERTY(Category = NonTessellatedLOD, EditAnywhere, meta = (ClampMin = "1", DisplayName = "Non Tessellated LOD Section Scale", EditCondition = "bEnableNonTessellatedLODMesh"))
-	uint32 NonTessellatedLODSectionScale = 4;
-
 	/** Higher number is higher priority. If Water Zones overlap and a water body does not have a manual water zone override, this priority will be used when automatically assigning the zone. */
 	UPROPERTY(Category = Water, EditAnywhere, AdvancedDisplay)
 	int32 OverlapPriority = 0;
 
-	UPROPERTY(Category = NonTessellatedLOD, EditAnywhere, meta = (DisplayName = "Enable Non-Tessellated LOD Mesh"))
-	bool bEnableNonTessellatedLODMesh = false;
+	UPROPERTY(Category = LocalTessellation, EditAnywhere)
+	bool bEnableLocalOnlyTessellation = false;
 
-	bool bNeedsWaterInfoRebuild = true;
+	UPROPERTY(Category = LocalTessellation, EditAnywhere, meta = (EditCondition = "bEnableLocalOnlyTessellation"))
+	FVector LocalTessellationExtent;
 
-	bool bNeedsNonTessellatedMeshRebuild = true;
+	bool bNeedsWaterInfoRebuild = true;
 
 	FVector2f WaterHeightExtents;
 	float GroundZMin;
 
-	FVector TessellatedWaterMeshCenter;
+	FVector LocalTessellationCenter;
 
 #if WITH_EDITORONLY_DATA
 	/** A manipulatable box for visualizing/editing the water zone bounds */
@@ -192,8 +192,17 @@ private:
 	UPROPERTY(Transient)
 	TObjectPtr<UBillboardComponent> ActorIcon;
 
+	// Deprecated UPROPERTIES
+
 	UPROPERTY(Transient)
 	TObjectPtr<UTexture2D> WaterVelocityTexture_DEPRECATED;
+
+	UPROPERTY()
+	FVector TessellatedWaterMeshExtent_DEPRECATED;
+
+	UPROPERTY()
+	bool bEnableNonTesselatedLODMesh_DEPRECATED;
+
 #endif // WITH_EDITORONLY_DATA
 };
 DEFINE_ACTORDESC_TYPE(AWaterZone, FWaterZoneActorDesc);
diff --git a/Engine/Source/Runtime/Core/Public/UObject/FortniteMainBranchObjectVersion.h b/Engine/Source/Runtime/Core/Public/UObject/FortniteMainBranchObjectVersion.h
--- a/Engine/Source/Runtime/Core/Public/UObject/FortniteMainBranchObjectVersion.h
+++ b/Engine/Source/Runtime/Core/Public/UObject/FortniteMainBranchObjectVersion.h
@@ -298,6 +298,9 @@ struct CORE_API FFortniteMainBranchObjectVersion
 		// Change how connection graphs are stored on Geometry Collections to an edge-array representation
 		ChaosGeometryCollectionConnectionEdgeGroup,
 
+		// Moved the water info mesh data and static water body meshes into new static mesh components for water bodies.
+		WaterBodyStaticMeshComponents,
+
 		// -----<new versions can be added above this line>-------------------------------------------------
 		VersionPlusOne,
 		LatestVersion = VersionPlusOne - 1
