Added back deprecated function with deprecation message.

diff --git a/Engine/Plugins/Experimental/Water/Source/Runtime/Public/WaterBodyComponent.h b/Engine/Plugins/Experimental/Water/Source/Runtime/Public/WaterBodyComponent.h
--- a/Engine/Plugins/Experimental/Water/Source/Runtime/Public/WaterBodyComponent.h
+++ b/Engine/Plugins/Experimental/Water/Source/Runtime/Public/WaterBodyComponent.h
@@ -233,6 +233,11 @@ public:
 	UFUNCTION(BlueprintCallable, Category = Rendering)
 	UMaterialInstanceDynamic* GetWaterStaticMeshMaterialInstance();
 
+	/** Returns water LOD MID */
+	UE_DEPRECATED(5.3, "GetWaterLODMaterialInstance has been renamed to GetWaterStaticMeshMaterialInstance.")
+	UFUNCTION(BlueprintCallable, Category = Rendering, meta=(DeprecationMessage="GetWaterLODMaterialInstance has been renamed to GetWaterStaticMeshMaterialInstance"))
+	UMaterialInstanceDynamic* GetWaterLODMaterialInstance() { return GetWaterStaticMeshMaterialInstance(); };
+
 	/** Returns under water post process MID */
 	UFUNCTION(BlueprintCallable, Category = Rendering)
 	UMaterialInstanceDynamic* GetUnderwaterPostProcessMaterialInstance();
