#UE Set new components to movable as they are not editable by BPs and some BPs make the water spline movable

diff --git a/Engine/Plugins/Experimental/Water/Source/Runtime/Private/WaterBodyActor.cpp b/Engine/Plugins/Experimental/Water/Source/Runtime/Private/WaterBodyActor.cpp
--- a/Engine/Plugins/Experimental/Water/Source/Runtime/Private/WaterBodyActor.cpp
+++ b/Engine/Plugins/Experimental/Water/Source/Runtime/Private/WaterBodyActor.cpp
@@ -57,9 +57,11 @@ AWaterBody::AWaterBody(const FObjectInitializer& ObjectInitializer)
 
 	WaterInfoMeshComponent = CreateDefaultSubobject<UWaterBodyInfoMeshComponent>(TEXT("WaterInfoMeshComponent"));
 	WaterInfoMeshComponent->SetupAttachment(RootComponent);
+	WaterInfoMeshComponent->SetMobility(EComponentMobility::Movable);
 
 	DilatedWaterInfoMeshComponent = CreateDefaultSubobject<UWaterBodyInfoMeshComponent>(TEXT("DilatedWaterInfoMeshComponent"));
 	DilatedWaterInfoMeshComponent->SetupAttachment(RootComponent);
+	DilatedWaterInfoMeshComponent->SetMobility(EComponentMobility::Movable);
 
 #if WITH_EDITORONLY_DATA
 	bAffectsLandscape_DEPRECATED = true;
