Fixed nonunity error in WaterBodyInfoMeshComponent.cpp

diff --git a/Engine/Plugins/Experimental/Water/Source/Runtime/Private/WaterBodyInfoMeshComponent.cpp b/Engine/Plugins/Experimental/Water/Source/Runtime/Private/WaterBodyInfoMeshComponent.cpp
--- a/Engine/Plugins/Experimental/Water/Source/Runtime/Private/WaterBodyInfoMeshComponent.cpp
+++ b/Engine/Plugins/Experimental/Water/Source/Runtime/Private/WaterBodyInfoMeshComponent.cpp
@@ -1,7 +1,5 @@
 // Copyright Epic Games, Inc. All Rights Reserved.
 
-#pragma once
-
 #include "WaterBodyInfoMeshComponent.h"
 #include "StaticMeshSceneProxy.h"
 
