Fixed static analysis warning for WaterEditorModule.cpp

diff --git a/Engine/Plugins/Experimental/Water/Source/Editor/Private/WaterEditorModule.cpp b/Engine/Plugins/Experimental/Water/Source/Editor/Private/WaterEditorModule.cpp
--- a/Engine/Plugins/Experimental/Water/Source/Editor/Private/WaterEditorModule.cpp
+++ b/Engine/Plugins/Experimental/Water/Source/Editor/Private/WaterEditorModule.cpp
@@ -257,78 +257,81 @@ void FWaterEditorModule::OnLevelActorAddedToWorld(AActor* Actor)
 			}
 		}
 
-		// Setup the water zone actor for this water body : 
-
-		const bool bHasZoneActor = !!TActorIterator<AWaterZone>(ActorWorld);
-		if ((WaterBodyActor != nullptr) && !bHasZoneActor)
+		if (WaterBodyActor != nullptr)
 		{
-			TSubclassOf<AWaterZone> WaterZoneClass = WaterEditorSettings->GetWaterZoneClass();
-			if (UClass* WaterZoneClassPtr = WaterZoneClass.Get())
+			const bool bHasZoneActor = !!TActorIterator<AWaterZone>(ActorWorld);
+			// Setup the water zone actor for this water body : 
+			if (!bHasZoneActor)
 			{
-				UActorFactory* WaterZoneActorFactory = GEditor->FindActorFactoryForActorClass(WaterZoneClassPtr);
+				TSubclassOf<AWaterZone> WaterZoneClass = WaterEditorSettings->GetWaterZoneClass();
+				if (UClass* WaterZoneClassPtr = WaterZoneClass.Get())
+				{
+					UActorFactory* WaterZoneActorFactory = GEditor->FindActorFactoryForActorClass(WaterZoneClassPtr);
 
-				FActorSpawnParameters SpawnParams;
-				SpawnParams.OverrideLevel = ActorWorld->PersistentLevel;
-				SpawnParams.bAllowDuringConstructionScript = true; // This can be called by construction script if the actor being added to the world is part of a blueprint, for example : 
+					FActorSpawnParameters SpawnParams;
+					SpawnParams.OverrideLevel = ActorWorld->PersistentLevel;
+					SpawnParams.bAllowDuringConstructionScript = true; // This can be called by construction script if the actor being added to the world is part of a blueprint, for example : 
 
-				AWaterZone* WaterZoneActor = (WaterZoneActorFactory != nullptr)
-					? Cast<AWaterZone>(WaterZoneActorFactory->CreateActor(WaterZoneClassPtr, Actor->GetLevel(), FTransform(WaterZoneBounds.GetCenter()), SpawnParams))
-					: ActorWorld->SpawnActor<AWaterZone>(WaterZoneClassPtr, SpawnParams);
+					AWaterZone* WaterZoneActor = (WaterZoneActorFactory != nullptr)
+						? Cast<AWaterZone>(WaterZoneActorFactory->CreateActor(WaterZoneClassPtr, Actor->GetLevel(), FTransform(WaterZoneBounds.GetCenter()), SpawnParams))
+						: ActorWorld->SpawnActor<AWaterZone>(WaterZoneClassPtr, SpawnParams);
 
-				if (WaterZoneActor)
-				{
-					if (!WaterZoneActorFactory)
+					if (WaterZoneActor)
 					{
-						UE_LOG(LogWaterEditor, Warning, TEXT("WaterZone Actor Factory could not be found! The newly spawned %s may have incorrect defaults!"), *WaterZoneActor->GetActorLabel());
-					}
+						if (!WaterZoneActorFactory)
+						{
+							UE_LOG(LogWaterEditor, Warning, TEXT("WaterZone Actor Factory could not be found! The newly spawned %s may have incorrect defaults!"), *WaterZoneActor->GetActorLabel());
+						}
 
-					// TODO [jonathan.bard] : when we can tag static meshes as "water ground", add these to the bounds
-					// Set a more sensible default location and extent so that the zone fully encapsulates the landscape if one exists.
-					if (WaterZoneBounds.IsValid)
-					{
-						WaterZoneActor->SetActorLocation(WaterZoneBounds.GetCenter());
+						// TODO [jonathan.bard] : when we can tag static meshes as "water ground", add these to the bounds
+						// Set a more sensible default location and extent so that the zone fully encapsulates the landscape if one exists.
+						if (WaterZoneBounds.IsValid)
+						{
+							WaterZoneActor->SetActorLocation(WaterZoneBounds.GetCenter());
 
-						// FBox::GetExtent returns the radius, SetZoneExtent expects diameter.
-						FVector2D NewExtent = 2 * FVector2D(WaterZoneBounds.GetExtent());
+							// FBox::GetExtent returns the radius, SetZoneExtent expects diameter.
+							FVector2D NewExtent = 2 * FVector2D(WaterZoneBounds.GetExtent());
 
-						float ZoneExtentScale = WaterEditorModule::CVarOverrideNewWaterZoneScale.GetValueOnGameThread();
-						const float MinimumMargin = WaterEditorModule::CVarOverrideNewWaterZoneMinimumMargin.GetValueOnGameThread();
+							float ZoneExtentScale = WaterEditorModule::CVarOverrideNewWaterZoneScale.GetValueOnGameThread();
+							const float MinimumMargin = WaterEditorModule::CVarOverrideNewWaterZoneMinimumMargin.GetValueOnGameThread();
 
-						if (ZoneExtentScale == 0)
-						{
-							ZoneExtentScale = GetDefault<UWaterEditorSettings>()->WaterZoneActorDefaults.NewWaterZoneScale;
-						}
+							if (ZoneExtentScale == 0)
+							{
+								ZoneExtentScale = GetDefault<UWaterEditorSettings>()->WaterZoneActorDefaults.NewWaterZoneScale;
+							}
 
-						if (ZoneExtentScale != 0)
-						{
-							NewExtent = FMath::Abs(ZoneExtentScale) * NewExtent;
-						}
+							if (ZoneExtentScale != 0)
+							{
+								NewExtent = FMath::Abs(ZoneExtentScale) * NewExtent;
+							}
 
-						if ((MinimumMargin > 0) && (LandscapeBounds.IsValid))
-						{
-							const FVector2D LandscapeBoundsDiameter = 2.f * FVector2D(LandscapeBounds.GetExtent());
-							const FVector2D MinimumWaterZoneExtent = LandscapeBoundsDiameter + 2.f * MinimumMargin;
+							if ((MinimumMargin > 0) && (LandscapeBounds.IsValid))
+							{
+								const FVector2D LandscapeBoundsDiameter = 2.f * FVector2D(LandscapeBounds.GetExtent());
+								const FVector2D MinimumWaterZoneExtent = LandscapeBoundsDiameter + 2.f * MinimumMargin;
 
-							NewExtent.X = (NewExtent.X < MinimumWaterZoneExtent.X) ? MinimumWaterZoneExtent.X : NewExtent.X;
-							NewExtent.Y = (NewExtent.Y < MinimumWaterZoneExtent.Y) ? MinimumWaterZoneExtent.Y : NewExtent.Y;
-						}
+								NewExtent.X = (NewExtent.X < MinimumWaterZoneExtent.X) ? MinimumWaterZoneExtent.X : NewExtent.X;
+								NewExtent.Y = (NewExtent.Y < MinimumWaterZoneExtent.Y) ? MinimumWaterZoneExtent.Y : NewExtent.Y;
+							}
 
-						WaterZoneActor->SetZoneExtent(NewExtent);
+							WaterZoneActor->SetZoneExtent(NewExtent);
+						}
 					}
 				}
+				else
+				{
+					UE_LOG(LogWaterEditor, Warning, TEXT("Could not find Water Zone class %s to spawn"), *WaterEditorSettings->GetWaterZoneClassPath().GetAssetPathString());
+				}
+
 			}
-			else
+
+			// If the actor is an ocean, we can help the user by initializing the ocean extent to fully fill the zone to which it belongs:
+			check(WaterBodyActor->GetWaterBodyComponent());
+			if (UWaterBodyOceanComponent* OceanComponent = Cast<UWaterBodyOceanComponent>(WaterBodyActor->GetWaterBodyComponent()))
 			{
-				UE_LOG(LogWaterEditor, Warning, TEXT("Could not find Water Zone class %s to spawn"), *WaterEditorSettings->GetWaterZoneClassPath().GetAssetPathString());
+				OceanComponent->FillWaterZoneWithOcean();
 			}
 		}
-
-		// If the actor is an ocean, we can help the user by initializing the ocean extent to fully fill the zone to which it belongs:
-		check(WaterBodyActor->GetWaterBodyComponent());
-		if (UWaterBodyOceanComponent* OceanComponent = Cast<UWaterBodyOceanComponent>(WaterBodyActor->GetWaterBodyComponent()))
-		{
-			OceanComponent->FillWaterZoneWithOcean();
-		}
 	}
 }
 
