Fixed a duplicate loctext key in WaterBodyOceanComponent's mapcheck

diff --git a/Engine/Plugins/Experimental/Water/Source/Runtime/Private/WaterBodyOceanComponent.cpp b/Engine/Plugins/Experimental/Water/Source/Runtime/Private/WaterBodyOceanComponent.cpp
--- a/Engine/Plugins/Experimental/Water/Source/Runtime/Private/WaterBodyOceanComponent.cpp
+++ b/Engine/Plugins/Experimental/Water/Source/Runtime/Private/WaterBodyOceanComponent.cpp
@@ -534,7 +534,7 @@ TArray<TSharedRef<FTokenizedMessage>> UWaterBodyOceanComponent::CheckWaterBodySt
 						FText::FromString(*GetNameSafe(WaterZone)))))
 					->AddToken(FActionToken::Create(LOCTEXT("MapCheck_MessageAction_SaveOcean", "Click here to fill the zone with ocean."), FText(),
 						FOnActionTokenExecuted::CreateUObject(this, &UWaterBodyOceanComponent::FillWaterZoneWithOcean)))
-					->AddToken(FActionToken::Create(LOCTEXT("MapCheck_MessageAction_SaveOcean", "If this is desired behavior, click here to disable this warning."), FText(),
+					->AddToken(FActionToken::Create(LOCTEXT("MapCheck_MessageAction_DisableOceanZoneMismatchWarning", "If this is desired behavior, click here to disable this warning."), FText(),
 						FOnActionTokenExecuted::CreateLambda(DisableWarnOnMismatchExtent)))
 						);
 			}
