Make UWaterMeshComponent not navigation relevant (this fixes an issue with hundred of navmesh tiles being dirtied)

diff --git a/Engine/Plugins/Experimental/Water/Source/Runtime/Public/WaterMeshComponent.h b/Engine/Plugins/Experimental/Water/Source/Runtime/Public/WaterMeshComponent.h
--- a/Engine/Plugins/Experimental/Water/Source/Runtime/Public/WaterMeshComponent.h
+++ b/Engine/Plugins/Experimental/Water/Source/Runtime/Public/WaterMeshComponent.h
@@ -38,6 +38,10 @@ public:
 #endif // WITH_EDITOR
 	//~ End UPrimitiveComponent Interface
 
+	//~ Begin INavRelevantInterface Interface
+	virtual bool IsNavigationRelevant() const override { return false; }
+	//~ End INavRelevantInterface Interface
+
 	virtual void CollectPSOPrecacheData(const FPSOPrecacheParams& BasePrecachePSOParams, FComponentPSOPrecacheParamsList& OutParams) override;
 
 	void Update();
