Water: fixed an issue where water mesh components would not be renamed to the correct package when converting a world containing water bodies to OFPA.

diff --git a/Engine/Plugins/Experimental/Water/Source/Runtime/Private/WaterBodyActor.cpp b/Engine/Plugins/Experimental/Water/Source/Runtime/Private/WaterBodyActor.cpp
--- a/Engine/Plugins/Experimental/Water/Source/Runtime/Private/WaterBodyActor.cpp
+++ b/Engine/Plugins/Experimental/Water/Source/Runtime/Private/WaterBodyActor.cpp
@@ -24,6 +24,7 @@
 #if WITH_EDITOR
 #include "Editor/EditorEngine.h"
 #include "Landscape.h"
+#include "Engine/StaticMesh.h"
 #endif
 
 #define LOCTEXT_NAMESPACE "Water"
@@ -71,6 +72,8 @@ AWaterBody::AWaterBody(const FObjectInitializer& ObjectInitializer)
 	TargetWaveMaskDepth_DEPRECATED = 2048.f;
 	bCanAffectNavigation_DEPRECATED = false;
 	bFillCollisionUnderWaterBodiesForNavmesh_DEPRECATED = false;
+
+	OnPackagingModeChanged.AddUObject(this, &AWaterBody::PackagingModeChanged);
 #endif // WITH_EDITORONLY_DATA
 
 }
@@ -665,6 +668,23 @@ void AWaterBody::GetActorDescProperties(FPropertyPairsMap& PropertyPairsMap) con
 	}
 }
 
+void AWaterBody::PackagingModeChanged(AActor* Actor, bool bIsExternal)
+{
+	// All the UStaticMesh are outered to the package, we need to make sure they remain so after the packaging mode of this actor is changed.
+	// If the actor used to be non-external and is now external, we need to ensure the UStaticMesh is outered to the external package of the actor and not the world.
+	// The inverse is true from external -> non external
+	TArray<TObjectPtr<UWaterBodyMeshComponent>> MeshComponents = { WaterInfoMeshComponent, DilatedWaterInfoMeshComponent };
+	MeshComponents.Append(WaterBodyStaticMeshComponents);
+
+	for (TObjectPtr<UWaterBodyMeshComponent> MeshComponent : MeshComponents)
+	{
+		if (IsValid(MeshComponent) && IsValid(MeshComponent->GetStaticMesh()))
+		{
+			MeshComponent->GetStaticMesh()->Rename(nullptr, GetPackage(), REN_DontCreateRedirectors | REN_ForceNoResetLoaders);
+		}
+	}
+}
+
 #endif // WITH_EDITOR
 
 void AWaterBody::UnregisterAllComponents(bool bForReregister)
diff --git a/Engine/Plugins/Experimental/Water/Source/Runtime/Public/WaterBodyActor.h b/Engine/Plugins/Experimental/Water/Source/Runtime/Public/WaterBodyActor.h
--- a/Engine/Plugins/Experimental/Water/Source/Runtime/Public/WaterBodyActor.h
+++ b/Engine/Plugins/Experimental/Water/Source/Runtime/Public/WaterBodyActor.h
@@ -64,6 +64,8 @@ public:
 	virtual void SetIsTemporarilyHiddenInEditor(bool bIsHidden) override;
 	virtual bool SetIsHiddenEdLayer(bool bIsHiddenEdLayer) override;
 	virtual void GetActorDescProperties(FPropertyPairsMap& PropertyPairsMap) const override;
+
+	void PackagingModeChanged(AActor* Actor, bool bIsExternal);
 #endif // WITH_EDITOR
 	
 	/** Returns the type of body */
