Fixed an issue wherein if nanite is forcibly enabled, a potentially empty water mesh will be created and crash the nanite builder which cannot handle empty meshes.

diff --git a/Engine/Plugins/Experimental/Water/Source/Runtime/Private/WaterBodyMeshBuilder.cpp b/Engine/Plugins/Experimental/Water/Source/Runtime/Private/WaterBodyMeshBuilder.cpp
--- a/Engine/Plugins/Experimental/Water/Source/Runtime/Private/WaterBodyMeshBuilder.cpp
+++ b/Engine/Plugins/Experimental/Water/Source/Runtime/Private/WaterBodyMeshBuilder.cpp
@@ -47,6 +47,11 @@ void FWaterBodyMeshBuilder::BuildWaterInfoMeshes(UWaterBodyComponent* WaterBodyC
 
 	auto BuildWaterInfoMesh = [this, Outer, WaterInfoMID](const FDynamicMesh3& DynamicMesh, UWaterBodyInfoMeshComponent* MeshComponent, FName BaseName) -> UStaticMesh*
 	{
+		if (DynamicMesh.TriangleCount() == 0)
+		{
+			return nullptr;
+		}
+
 		FName WaterInfoMeshName = MakeUniqueObjectName(Outer, UStaticMesh::StaticClass(), BaseName);
 		UStaticMesh* StaticMesh = CreateUStaticMesh(Outer, WaterInfoMeshName);
 
