Fixed an issue preventing water static mesh sectioning from being enabled.

diff --git a/Engine/Plugins/Experimental/Water/Source/Runtime/Private/WaterBodyMeshBuilder.cpp b/Engine/Plugins/Experimental/Water/Source/Runtime/Private/WaterBodyMeshBuilder.cpp
--- a/Engine/Plugins/Experimental/Water/Source/Runtime/Private/WaterBodyMeshBuilder.cpp
+++ b/Engine/Plugins/Experimental/Water/Source/Runtime/Private/WaterBodyMeshBuilder.cpp
@@ -248,13 +248,12 @@ TArray<TObjectPtr<UWaterBodyStaticMeshComponent>> FWaterBodyMeshBuilder::BuildWa
 			TArray<uint32> Indices;
 
 			// Simplification can occur for lakes/rivers since they have flat geometry and no vertex color data.
-			//if (!bHasBoundaryTriangle && (WaterBodyType != EWaterBodyType::River) && Triangles.Num() > 0)
-			//{
-			//	const float MeshZ = ConstantSurfaceZ;
-			//	const uint32 QuadSectionIndex = NumQuadSectionPositions.fetch_add(1);
-			//	QuadSectionPositions[QuadSectionIndex] = FVector(SectionPos.X, SectionPos.Y, ConstantSurfaceZ);
-			//	return;
-			//}
+			if (!bHasBoundaryTriangle && (WaterBodyType != EWaterBodyType::River) && Triangles.Num() > 0)
+			{
+				const uint32 QuadSectionIndex = NumQuadSectionPositions.fetch_add(1);
+				QuadSectionPositions[QuadSectionIndex] = FVector(SectionPos.X, SectionPos.Y, 0);
+				return;
+			}
 
 			TMap<uint32, uint32> VertexMap;
 			for (int32 TriangleID : Triangles)
@@ -370,6 +369,7 @@ TArray<TObjectPtr<UWaterBodyStaticMeshComponent>> FWaterBodyMeshBuilder::BuildWa
 		UWaterBodyStaticMeshComponent* WaterBodyStaticMeshComponent = NewObject<UWaterBodyStaticMeshComponent>(ActorOwner, Name);
 
 		WaterBodyStaticMeshComponent->RegisterComponent();
+		WaterBodyStaticMeshComponent->SetMobility(WaterBodyComponent->Mobility);
 		WaterBodyStaticMeshComponent->AttachToComponent(WaterBodyComponent, FAttachmentTransformRules::SnapToTargetNotIncludingScale);
 
 		StaticMeshComponents.Add(WaterBodyStaticMeshComponent);
@@ -404,6 +404,8 @@ TArray<TObjectPtr<UWaterBodyStaticMeshComponent>> FWaterBodyMeshBuilder::BuildWa
 
 			WaterBodyStaticMeshComponent->SetStaticMesh(StaticMesh);
 
+			// Reset any location/scale offsets which might have been set if this component is being reused.
+			WaterBodyStaticMeshComponent->SetRelativeTransform(FTransform::Identity);
 			StaticMeshes.Add(StaticMesh);
 
 			++StaticMeshComponentIndex;
diff --git a/Engine/Plugins/Experimental/Water/Source/Runtime/Public/WaterBodyStaticMeshSettings.h b/Engine/Plugins/Experimental/Water/Source/Runtime/Public/WaterBodyStaticMeshSettings.h
--- a/Engine/Plugins/Experimental/Water/Source/Runtime/Public/WaterBodyStaticMeshSettings.h
+++ b/Engine/Plugins/Experimental/Water/Source/Runtime/Public/WaterBodyStaticMeshSettings.h
@@ -14,8 +14,7 @@ struct FWaterBodyStaticMeshSettings
 	bool bEnableWaterBodyStaticMesh = false;
 
 	/** If enabled, the water body static mesh will be split into separate meshes to allow them to stream in and out as needed to save memory. */
-	// #todo_water [roey]: Re-enable this when UE-184177 and UE-184211: 
-	//UPROPERTY(Category = StaticMesh, EditAnywhere, meta = (EditCondition = "bEnableWaterBodyStaticMesh"))
+	UPROPERTY(Category = StaticMesh, EditAnywhere, meta = (EditCondition = "bEnableWaterBodyStaticMesh"))
 	bool bSectionWaterBodyStaticMesh = false;
 
 	/** The size of each section tile if sectioning is enabled. */
