Propagate water body component mobility to water info components and ensure they have the correct mobility to allow attachment.

diff --git a/Engine/Plugins/Experimental/Water/Source/Runtime/Private/WaterBodyActor.cpp b/Engine/Plugins/Experimental/Water/Source/Runtime/Private/WaterBodyActor.cpp
--- a/Engine/Plugins/Experimental/Water/Source/Runtime/Private/WaterBodyActor.cpp
+++ b/Engine/Plugins/Experimental/Water/Source/Runtime/Private/WaterBodyActor.cpp
@@ -57,12 +57,12 @@ AWaterBody::AWaterBody(const FObjectInitializer& ObjectInitializer)
 	RootComponent = SplineComp;
 
 	WaterInfoMeshComponent = CreateDefaultSubobject<UWaterBodyInfoMeshComponent>(TEXT("WaterInfoMeshComponent"));
+	WaterInfoMeshComponent->SetMobility(EComponentMobility::Static);
 	WaterInfoMeshComponent->SetupAttachment(RootComponent);
-	WaterInfoMeshComponent->SetMobility(EComponentMobility::Movable);
 
 	DilatedWaterInfoMeshComponent = CreateDefaultSubobject<UWaterBodyInfoMeshComponent>(TEXT("DilatedWaterInfoMeshComponent"));
+	DilatedWaterInfoMeshComponent->SetMobility(EComponentMobility::Static);
 	DilatedWaterInfoMeshComponent->SetupAttachment(RootComponent);
-	DilatedWaterInfoMeshComponent->SetMobility(EComponentMobility::Movable);
 
 #if WITH_EDITORONLY_DATA
 	bAffectsLandscape_DEPRECATED = true;
@@ -84,6 +84,11 @@ void AWaterBody::PreRegisterAllComponents()
 
 	SetRootComponent(WaterBodyComponent);
 	SplineComp->AttachToComponent(WaterBodyComponent, FAttachmentTransformRules::KeepRelativeTransform);
+	
+	WaterInfoMeshComponent->SetMobility(WaterBodyComponent->Mobility);
+	WaterInfoMeshComponent->AttachToComponent(WaterBodyComponent, FAttachmentTransformRules::SnapToTargetNotIncludingScale);
+	DilatedWaterInfoMeshComponent->SetMobility(WaterBodyComponent->Mobility);
+	DilatedWaterInfoMeshComponent->AttachToComponent(WaterBodyComponent, FAttachmentTransformRules::SnapToTargetNotIncludingScale);
 }
 
 void AWaterBody::NotifyActorBeginOverlap(AActor* OtherActor)
diff --git a/Engine/Plugins/Experimental/Water/Source/Runtime/Private/WaterBodyCustomComponent.cpp b/Engine/Plugins/Experimental/Water/Source/Runtime/Private/WaterBodyCustomComponent.cpp
--- a/Engine/Plugins/Experimental/Water/Source/Runtime/Private/WaterBodyCustomComponent.cpp
+++ b/Engine/Plugins/Experimental/Water/Source/Runtime/Private/WaterBodyCustomComponent.cpp
@@ -5,6 +5,7 @@
 #include "Materials/MaterialInstanceDynamic.h"
 #include "Misc/UObjectToken.h"
 #include "WaterBodyActor.h"
+#include "WaterBodyMeshComponent.h"
 
 #include UE_INLINE_GENERATED_CPP_BY_NAME(WaterBodyCustomComponent)
 
@@ -82,7 +83,7 @@ void UWaterBodyCustomComponent::OnUpdateBody(bool bWithExclusionVolumes)
 	TInlineComponentArray<UPrimitiveComponent*> PrimitiveComponents;
 	OwnerActor->GetComponents(PrimitiveComponents);
 
-	// Make no assumptions for custom water bodies: all (non-visualization) primitive components will be included in navigation
+	// Make no assumptions for custom water bodies: all (non-visualization) primitive components will be included in navigation except the water info components
 	for (UPrimitiveComponent* Comp : PrimitiveComponents)
 	{
 #if WITH_EDITORONLY_DATA
@@ -91,6 +92,10 @@ void UWaterBodyCustomComponent::OnUpdateBody(bool bWithExclusionVolumes)
 			continue;
 		}
 #endif // WITH_EDITORONLY_DATA
+		if (Comp->IsA<UWaterBodyMeshComponent>())
+		{
+			continue;
+		}
 
 		CopySharedNavigationSettingsToComponent(Comp);
 
