PSO precaching: Use correct strides for WaterVertexFactory caching

diff --git a/Engine/Plugins/Experimental/Water/Source/Runtime/Public/WaterVertexFactory.inl b/Engine/Plugins/Experimental/Water/Source/Runtime/Public/WaterVertexFactory.inl
--- a/Engine/Plugins/Experimental/Water/Source/Runtime/Public/WaterVertexFactory.inl
+++ b/Engine/Plugins/Experimental/Water/Source/Runtime/Public/WaterVertexFactory.inl
@@ -183,11 +183,11 @@ template <bool bWithWaterSelectionSupport>
 void TWaterVertexFactory<bWithWaterSelectionSupport>::GetPSOPrecacheVertexFetchElements(EVertexInputStreamType VertexInputStreamType, FVertexDeclarationElementList& Elements)
 {
 	// Add position stream
-	Elements.Add(FVertexElement(0, 0, VET_Float4, 0, 0, false));
+	Elements.Add(FVertexElement(0, 0, VET_Float4, 0, sizeof(FVector4f), false));
 
 	// Add all the additional streams
 	for (int32 StreamIdx = 0; StreamIdx < NumAdditionalVertexStreams; ++StreamIdx)
 	{
-		Elements.Add(FVertexElement(1 + StreamIdx, 0, VET_Float4, 8 + StreamIdx, 0, true));
+		Elements.Add(FVertexElement(1 + StreamIdx, 0, VET_Float4, 8 + StreamIdx, sizeof(FVector4f), true));
 	}
 }
