Fix crash when undoing change to water spline

#rnx

diff --git a/Engine/Plugins/Experimental/Water/Source/Runtime/Private/WaterBodyMeshBuilder.cpp b/Engine/Plugins/Experimental/Water/Source/Runtime/Private/WaterBodyMeshBuilder.cpp
--- a/Engine/Plugins/Experimental/Water/Source/Runtime/Private/WaterBodyMeshBuilder.cpp
+++ b/Engine/Plugins/Experimental/Water/Source/Runtime/Private/WaterBodyMeshBuilder.cpp
@@ -562,7 +562,7 @@ void FWaterBodyMeshBuilder::UpdateStaticMesh(UStaticMesh* WaterMesh, const FMesh
 
 UStaticMesh* FWaterBodyMeshBuilder::CreateUStaticMesh(UObject* Outer, FName MeshName) const
 {
-	UStaticMesh* StaticMesh = NewObject<UStaticMesh>(Outer, MeshName);
+	UStaticMesh* StaticMesh = NewObject<UStaticMesh>(Outer, MeshName, RF_Transactional);
 
 	// Disable navigation
 	StaticMesh->MarkAsNotHavingNavigationData();
