Fixed water meshes not being updated after an undo

#rnx

diff --git a/Engine/Plugins/Experimental/Water/Source/Runtime/Private/WaterBodyComponent.cpp b/Engine/Plugins/Experimental/Water/Source/Runtime/Private/WaterBodyComponent.cpp
--- a/Engine/Plugins/Experimental/Water/Source/Runtime/Private/WaterBodyComponent.cpp
+++ b/Engine/Plugins/Experimental/Water/Source/Runtime/Private/WaterBodyComponent.cpp
@@ -1472,7 +1472,7 @@ void UWaterBodyComponent::OnWaterBodyChanged(const FOnWaterBodyChangedParams& In
 	}
 
 #if WITH_EDITOR
-	if (InParams.PropertyChangedEvent.ChangeType == EPropertyChangeType::ValueSet)
+	if (InParams.PropertyChangedEvent.ChangeType != EPropertyChangeType::Interactive)
 	{
 		UpdateWaterBodyRenderData();
 	}
