Disable ray tracing on meshes created by WaterBodyMeshBuilder so they don't affect LumenScene.

diff --git a/Engine/Plugins/Experimental/Water/Source/Runtime/Private/WaterBodyMeshBuilder.cpp b/Engine/Plugins/Experimental/Water/Source/Runtime/Private/WaterBodyMeshBuilder.cpp
--- a/Engine/Plugins/Experimental/Water/Source/Runtime/Private/WaterBodyMeshBuilder.cpp
+++ b/Engine/Plugins/Experimental/Water/Source/Runtime/Private/WaterBodyMeshBuilder.cpp
@@ -567,6 +567,9 @@ UStaticMesh* FWaterBodyMeshBuilder::CreateUStaticMesh(UObject* Outer, FName Mesh
 	// Disable navigation
 	StaticMesh->MarkAsNotHavingNavigationData();
 
+	// Disable ray tracing as we don't want water meshes to show up in LumenScene.
+	StaticMesh->bSupportRayTracing = false;
+
 	if (UBodySetup* BodySetup = StaticMesh->GetBodySetup())
 	{
 		BodySetup->DefaultInstance.SetCollisionProfileName(UCollisionProfile::NoCollision_ProfileName);
