Potential fix for crash packaging Test_Extraction_Loadout_Map.

diff --git a/Engine/Plugins/Experimental/Water/Source/Runtime/Private/WaterBodyActor.cpp b/Engine/Plugins/Experimental/Water/Source/Runtime/Private/WaterBodyActor.cpp
--- a/Engine/Plugins/Experimental/Water/Source/Runtime/Private/WaterBodyActor.cpp
+++ b/Engine/Plugins/Experimental/Water/Source/Runtime/Private/WaterBodyActor.cpp
@@ -475,9 +475,9 @@ void AWaterBody::DeprecateData()
 		FMemory::Memcpy((void*)&WaterHeightmapSettings_DEPRECATED, (void*)&TerrainCarvingSettings_DEPRECATED, sizeof(WaterHeightmapSettings_DEPRECATED));
 	}
 
-	if (GIsEditor && !HasAnyFlags(RF_ClassDefaultObject))
+	if (GIsEditor && !HasAnyFlags(RF_ClassDefaultObject | RF_ArchetypeObject | RF_DefaultSubObject))
 	{
-		if (WaterWaves && (WaterWaves->GetOuter() != this))
+		if (WaterWaves && (!WaterWaves->HasAnyFlags(RF_ClassDefaultObject | RF_ArchetypeObject | RF_DefaultSubObject)) && (WaterWaves->GetOuter() != this))
 		{
 			WaterWaves->ClearFlags(RF_Public);
 			// At one point, WaterWaves's outer was the level. We need them to be outered by the water body : 
