Used friendlier names for water body actors in water-related logs

#rnx

diff --git a/Engine/Plugins/Experimental/Water/Source/Editor/Private/WaterBrushManager.cpp b/Engine/Plugins/Experimental/Water/Source/Editor/Private/WaterBrushManager.cpp
--- a/Engine/Plugins/Experimental/Water/Source/Editor/Private/WaterBrushManager.cpp
+++ b/Engine/Plugins/Experimental/Water/Source/Editor/Private/WaterBrushManager.cpp
@@ -1082,7 +1082,7 @@ bool AWaterBrushManager::DeprecateWaterLandscapeInfo(FVector& OutRTWorldLocation
 void AWaterBrushManager::ShowForceUpdateMapCheckError()
 {
 	FFormatNamedArguments Arguments;
-	Arguments.Add(TEXT("WaterBrush"), FText::FromString(GetName()));
+	Arguments.Add(TEXT("WaterBrush"), FText::FromString(GetActorNameOrLabel()));
 	Arguments.Add(TEXT("Outer"), FText::FromString(GetPackage()->GetPathName()));
 	
 	FMessageLog("MapCheck").Warning()
@@ -1434,7 +1434,7 @@ void AWaterBrushManager::UpdateBrushCacheKeys()
 				SetActorCache(Actor, WaterBrushCacheContainer);
 			}
 			// Make sure there's an appropriate render target in that cache : 
-			WaterBrushCacheContainer->Cache.CacheRenderTarget = FWaterUtils::GetOrCreateTransientRenderTarget2D(WaterBrushCacheContainer->Cache.CacheRenderTarget, FName(*FString::Printf(TEXT("BrushCacheRT_%s"), *Actor->GetName())), LandscapeRTRes, Format);
+			WaterBrushCacheContainer->Cache.CacheRenderTarget = FWaterUtils::GetOrCreateTransientRenderTarget2D(WaterBrushCacheContainer->Cache.CacheRenderTarget, FName(*FString::Printf(TEXT("BrushCacheRT_%s"), *Actor->GetActorNameOrLabel())), LandscapeRTRes, Format);
 			check(WaterBrushCacheContainer->Cache.CacheRenderTarget != nullptr);
 		}
 	}
diff --git a/Engine/Plugins/Experimental/Water/Source/Runtime/Private/BuoyancyComponent.cpp b/Engine/Plugins/Experimental/Water/Source/Runtime/Private/BuoyancyComponent.cpp
--- a/Engine/Plugins/Experimental/Water/Source/Runtime/Private/BuoyancyComponent.cpp
+++ b/Engine/Plugins/Experimental/Water/Source/Runtime/Private/BuoyancyComponent.cpp
@@ -416,7 +416,7 @@ void UBuoyancyComponent::ComputePontoonCoefficients()
 			//Distribute a mass of 1 to each pontoon so that we get a scaling factor based on position relative to CoM
 			FString ErrMsg;
 			bool ComputeSuccess = FSimpleSuspensionHelpers::ComputeSprungMasses(LocalPontoonLocations, LocalCOM, 1.f, PontoonCoefficients, &ErrMsg);
-			ensureMsgf(ComputeSuccess, TEXT("Failed to compute %d sprung masses for: %s\nErrMsg: \"%s\""), LocalPontoonLocations.Num(), *GetOwner()->GetName(), *ErrMsg);
+			ensureMsgf(ComputeSuccess, TEXT("Failed to compute %d sprung masses for: %s\nErrMsg: \"%s\""), LocalPontoonLocations.Num(), *GetOwner()->GetActorNameOrLabel(), *ErrMsg);
 		}
 	}
 
diff --git a/Engine/Plugins/Experimental/Water/Source/Runtime/Private/NiagaraWaterFunctionLibrary.cpp b/Engine/Plugins/Experimental/Water/Source/Runtime/Private/NiagaraWaterFunctionLibrary.cpp
--- a/Engine/Plugins/Experimental/Water/Source/Runtime/Private/NiagaraWaterFunctionLibrary.cpp
+++ b/Engine/Plugins/Experimental/Water/Source/Runtime/Private/NiagaraWaterFunctionLibrary.cpp
@@ -12,13 +12,13 @@ void UNiagaraWaterFunctionLibrary::SetWaterBodyComponent(UNiagaraComponent* Niag
 {
 	if (!NiagaraSystem)
 	{
-		UE_LOG(LogWater, Warning, TEXT("NiagaraSystem in \"Set Water Body Component\" is NULL, OverrideName \"%s\" and WaterBodyComponent \"%s\", skipping."), *OverrideName, WaterBodyComponent ? *WaterBodyComponent->GetName() : TEXT("NULL"));
+		UE_LOG(LogWater, Warning, TEXT("NiagaraSystem in \"Set Water Body Component\" is NULL, OverrideName \"%s\" and WaterBodyComponent \"%s\", skipping."), *OverrideName, WaterBodyComponent ? *WaterBodyComponent->GetOwner()->GetActorNameOrLabel() : TEXT("NULL"));
 		return;
 	}
 
 	if (!WaterBodyComponent)
 	{
-		UE_LOG(LogWater, Warning, TEXT("WaterBodyComponent in \"Set Water Body Component\" is NULL, OverrideName \"%s\" and NiagaraSystem \"%s\", skipping."), *OverrideName, *NiagaraSystem->GetOwner()->GetName());
+		UE_LOG(LogWater, Warning, TEXT("WaterBodyComponent in \"Set Water Body Component\" is NULL, OverrideName \"%s\" and NiagaraSystem \"%s\", skipping."), *OverrideName, *NiagaraSystem->GetOwner()->GetActorNameOrLabel());
 		return;
 	}
 
@@ -30,14 +30,14 @@ void UNiagaraWaterFunctionLibrary::SetWaterBodyComponent(UNiagaraComponent* Niag
 	int32 Index = OverrideParameters.IndexOf(Variable);
 	if (Index == INDEX_NONE)
 	{
-		UE_LOG(LogWater, Warning, TEXT("Could not find index of variable \"%s\" in the OverrideParameters map of NiagaraSystem \"%s\"."), *OverrideName, *NiagaraSystem->GetOwner()->GetName());
+		UE_LOG(LogWater, Warning, TEXT("Could not find index of variable \"%s\" in the OverrideParameters map of NiagaraSystem \"%s\"."), *OverrideName, *NiagaraSystem->GetOwner()->GetActorNameOrLabel());
 		return;
 	}
 
 	UNiagaraDataInterfaceWater* WaterInterface = Cast<UNiagaraDataInterfaceWater>(OverrideParameters.GetDataInterface(Index));
 	if (!WaterInterface)
 	{
-		UE_LOG(LogWater, Warning, TEXT("Did not find a matching Water Data Interface variable named \"%s\" in the User variables of NiagaraSystem \"%s\" ."), *OverrideName, *NiagaraSystem->GetOwner()->GetName());
+		UE_LOG(LogWater, Warning, TEXT("Did not find a matching Water Data Interface variable named \"%s\" in the User variables of NiagaraSystem \"%s\" ."), *OverrideName, *NiagaraSystem->GetOwner()->GetActorNameOrLabel());
 		return;
 	}
 
diff --git a/Engine/Plugins/Experimental/Water/Source/Runtime/Private/WaterBodyActor.cpp b/Engine/Plugins/Experimental/Water/Source/Runtime/Private/WaterBodyActor.cpp
--- a/Engine/Plugins/Experimental/Water/Source/Runtime/Private/WaterBodyActor.cpp
+++ b/Engine/Plugins/Experimental/Water/Source/Runtime/Private/WaterBodyActor.cpp
@@ -308,7 +308,7 @@ void AWaterBody::InitializeBody()
 		WaterBodyComponent->SetMobility(EComponentMobility::Static);
 	}
 
-	checkf(WaterBodyComponent, TEXT("Failed to create a water body component for a water body actor (%s)!"), *GetName());
+	checkf(WaterBodyComponent, TEXT("Failed to create a water body component for a water body actor (%s)!"), *GetActorNameOrLabel());
 
 	SetRootComponent(WaterBodyComponent);
 
diff --git a/Engine/Plugins/Experimental/Water/Source/Runtime/Private/WaterSubsystem.cpp b/Engine/Plugins/Experimental/Water/Source/Runtime/Private/WaterSubsystem.cpp
--- a/Engine/Plugins/Experimental/Water/Source/Runtime/Private/WaterSubsystem.cpp
+++ b/Engine/Plugins/Experimental/Water/Source/Runtime/Private/WaterSubsystem.cpp
@@ -834,7 +834,7 @@ void UWaterSubsystem::ShowOnScreenDebugInfo(const FVector& InViewLocation, const
 			MaterialDescription = FString::Format(TEXT("{0} (parent: {1})"), { MID->Parent->GetName(), MID->GetMaterial()->GetName() });
 		}
 		OutputStrings.Add(FText::Format(LOCTEXT("VisualizeActiveUnderwaterPostProcess_ActivePostprocess", "Active underwater post process water body {0} (material: {1})"),
-			FText::FromString(InDebugInfo.ActiveWaterBodyComponent->GetOwner()->GetName()),
+			FText::FromString(InDebugInfo.ActiveWaterBodyComponent->GetOwner()->GetActorNameOrLabel()),
 			FText::FromString(MaterialDescription)));
 	}
 	else
@@ -874,7 +874,7 @@ void UWaterSubsystem::ShowOnScreenDebugInfo(const FVector& InViewLocation, const
 				if (WaterBody.IsValid() && WaterBody->GetOwner())
 				{
 					OutputStrings.Add(FText::Format(LOCTEXT("VisualizeActiveUnderwaterPostProcess_OverlappedWaterBodyDetails", "- {0} (overlap material priority: {1})"),
-						FText::FromString(WaterBody->GetOwner()->GetName()),
+						FText::FromString(WaterBody->GetOwner()->GetActorNameOrLabel()),
 						FText::AsNumber(WaterBody->GetOverlapMaterialPriority())));
 				}
 			}
