Integrated 25389054
[Water] Fixed name of compute stats

#rnx

diff --git a/Engine/Plugins/Experimental/Water/Source/Runtime/Private/WaterBodyComponent.cpp b/Engine/Plugins/Experimental/Water/Source/Runtime/Private/WaterBodyComponent.cpp
--- a/Engine/Plugins/Experimental/Water/Source/Runtime/Private/WaterBodyComponent.cpp
+++ b/Engine/Plugins/Experimental/Water/Source/Runtime/Private/WaterBodyComponent.cpp
@@ -51,8 +51,8 @@
 
 DECLARE_CYCLE_STAT(TEXT("WaterBody_ComputeWaterInfo"), STAT_WaterBody_ComputeWaterInfo, STATGROUP_Water);
 DECLARE_CYCLE_STAT(TEXT("WaterBody_ComputeWaterDepth"), STAT_WaterBody_ComputeWaterDepth, STATGROUP_Water);
-DECLARE_CYCLE_STAT(TEXT("WaterBody_ComputeWaterDepth"), STAT_WaterBody_ComputeLocation, STATGROUP_Water);
-DECLARE_CYCLE_STAT(TEXT("WaterBody_ComputeWaterDepth"), STAT_WaterBody_ComputeNormal, STATGROUP_Water);
+DECLARE_CYCLE_STAT(TEXT("WaterBody_ComputeLocation"), STAT_WaterBody_ComputeLocation, STATGROUP_Water);
+DECLARE_CYCLE_STAT(TEXT("WaterBody_ComputeNormal"), STAT_WaterBody_ComputeNormal, STATGROUP_Water);
 DECLARE_CYCLE_STAT(TEXT("WaterBody_ComputeLandscapeDepth"), STAT_WaterBody_ComputeLandscapeDepth, STATGROUP_Water);
 DECLARE_CYCLE_STAT(TEXT("WaterBody_ComputeWaveHeight"), STAT_WaterBody_ComputeWaveHeight, STATGROUP_Water);
 
