Disabled additional show flags for the SceneCaptureComponent2D of AWaterBrushManager. This is necessary for a clean capture of just the output of the water info material.

#rnx

diff --git a/Engine/Plugins/Experimental/Water/Source/Editor/Private/WaterBrushManager.cpp b/Engine/Plugins/Experimental/Water/Source/Editor/Private/WaterBrushManager.cpp
--- a/Engine/Plugins/Experimental/Water/Source/Editor/Private/WaterBrushManager.cpp
+++ b/Engine/Plugins/Experimental/Water/Source/Editor/Private/WaterBrushManager.cpp
@@ -53,6 +53,10 @@ AWaterBrushManager::AWaterBrushManager(const FObjectInitializer& ObjectInitializ
 	SceneCaptureComponent2D->SetRelativeScale3D(FVector(0.01f, 0.01f, 0.01f));
 	// HACK [jonathan.bard] : Nanite doesn't support USceneCaptureComponent's ShowOnlyComponents ATM so just disable Nanite during captures : 
 	SceneCaptureComponent2D->ShowFlagSettings.Add(FEngineShowFlagsSetting { TEXT("NaniteMeshes"), false } );
+	// These also need to be disabled to get a clean capture of just the water info material output
+	SceneCaptureComponent2D->ShowFlagSettings.Add(FEngineShowFlagsSetting { TEXT("Atmosphere"), false } );
+	SceneCaptureComponent2D->ShowFlagSettings.Add(FEngineShowFlagsSetting { TEXT("Bloom"), false } );
+	SceneCaptureComponent2D->ShowFlagSettings.Add(FEngineShowFlagsSetting { TEXT("Lighting"), false } );
 
 	PrimaryActorTick.TickGroup = ETickingGroup::TG_PrePhysics;
 	bIsEditorOnlyActor = false;
