Fixes incorrect water velocities with some scalability settings.

diff --git a/Engine/Plugins/Experimental/Water/Source/Runtime/Private/WaterInfoRendering.cpp b/Engine/Plugins/Experimental/Water/Source/Runtime/Private/WaterInfoRendering.cpp
--- a/Engine/Plugins/Experimental/Water/Source/Runtime/Private/WaterInfoRendering.cpp
+++ b/Engine/Plugins/Experimental/Water/Source/Runtime/Private/WaterInfoRendering.cpp
@@ -688,8 +688,11 @@ static FSceneRenderer* CreateWaterInfoSceneRenderer(const FCreateWaterInfoSceneR
 	ViewInitOptions.ProjectionMatrix = Params.ProjectionMatrix;
 	ViewInitOptions.LODDistanceFactor = 0.001f;
 	ViewInitOptions.OverlayColor = FLinearColor::Black;
-	// Must be set to false to prevent the renders from using different VSM page pool sizes leading to unnecessary reallocations.
-	ViewInitOptions.bIsSceneCapture = false;
+	/* Must be set to false to prevent the renders from using different VSM page pool sizes leading to unnecessary reallocations.
+	** #todo_water: Temporarily set back to true since bIsSceneCapture forces bRequiresAlphaChannel in ViewFamily, leading to the proper SceneColor pixel format (PF_FloatRGBA).
+	** Without it, the Renderer would initialize SceneColor to PF_FloatRGB, resulting in incomplete velocities in the texture
+	*/
+	ViewInitOptions.bIsSceneCapture = true;
 
 	if (ViewFamily.Scene->GetWorld() != nullptr && ViewFamily.Scene->GetWorld()->GetWorldSettings() != nullptr)
 	{
