Fixes WaterBody with scale different than 1. Resulting mesh was not following the intended spline.

diff --git a/Engine/Plugins/Experimental/Water/Source/Runtime/Private/WaterBodyActor.cpp b/Engine/Plugins/Experimental/Water/Source/Runtime/Private/WaterBodyActor.cpp
--- a/Engine/Plugins/Experimental/Water/Source/Runtime/Private/WaterBodyActor.cpp
+++ b/Engine/Plugins/Experimental/Water/Source/Runtime/Private/WaterBodyActor.cpp
@@ -88,13 +88,13 @@ void AWaterBody::PreRegisterAllComponents()
 	if (IsValid(WaterInfoMeshComponent))
 	{
 		WaterInfoMeshComponent->SetMobility(WaterBodyComponent->Mobility);
-		WaterInfoMeshComponent->AttachToComponent(WaterBodyComponent, FAttachmentTransformRules::SnapToTargetNotIncludingScale);
+		WaterInfoMeshComponent->AttachToComponent(WaterBodyComponent, FAttachmentTransformRules::KeepRelativeTransform);
 	}
 
 	if (IsValid(DilatedWaterInfoMeshComponent))
 	{
 		DilatedWaterInfoMeshComponent->SetMobility(WaterBodyComponent->Mobility);
-		DilatedWaterInfoMeshComponent->AttachToComponent(WaterBodyComponent, FAttachmentTransformRules::SnapToTargetNotIncludingScale);
+		DilatedWaterInfoMeshComponent->AttachToComponent(WaterBodyComponent, FAttachmentTransformRules::KeepRelativeTransform);
 	}
 }
 
