Make sure materials are 'PostLoaded' on PSO precaching requests

diff --git a/Engine/Plugins/Experimental/Water/Source/Runtime/Private/WaterBodyComponent.cpp b/Engine/Plugins/Experimental/Water/Source/Runtime/Private/WaterBodyComponent.cpp
--- a/Engine/Plugins/Experimental/Water/Source/Runtime/Private/WaterBodyComponent.cpp
+++ b/Engine/Plugins/Experimental/Water/Source/Runtime/Private/WaterBodyComponent.cpp
@@ -1541,17 +1541,14 @@ void UWaterBodyComponent::PostLoad()
 		SetupPrecachePSOParams(PrecachePSOParams);
 		if (WaterMaterial)
 		{
-			WaterMaterial->ConditionalPostLoad();
 			WaterMaterial->PrecachePSOs(&FLocalVertexFactory::StaticType, PrecachePSOParams);
 		}
 		if (UnderwaterPostProcessMaterial)
 		{
-			UnderwaterPostProcessMaterial->ConditionalPostLoad();
 			UnderwaterPostProcessMaterial->PrecachePSOs(&FLocalVertexFactory::StaticType, PrecachePSOParams);
 		}
 		if (WaterInfoMaterial)
 		{
-			WaterInfoMaterial->ConditionalPostLoad();
 			WaterInfoMaterial->PrecachePSOs(&FLocalVertexFactory::StaticType, PrecachePSOParams);
 		}
 	}
diff --git a/Engine/Source/Runtime/Landscape/Private/Landscape.cpp b/Engine/Source/Runtime/Landscape/Private/Landscape.cpp
--- a/Engine/Source/Runtime/Landscape/Private/Landscape.cpp
+++ b/Engine/Source/Runtime/Landscape/Private/Landscape.cpp
@@ -1056,7 +1056,6 @@ void ULandscapeComponent::PostLoad()
 		{
 			if (MaterialInterface)
 			{
-				MaterialInterface->ConditionalPostLoad();
 				MaterialInterface->PrecachePSOs(VertexFactoryDataList, PrecachePSOParams, EPSOPrecachePriority::High, MaterialPrecacheRequestIDs);
 			}
 		}
