Fixup and nullify template WaterBodyComponent reference on OwningWaterZone.
Adds mapcheck to dirty the WaterBodyComponent and resave it.

diff --git a/Engine/Plugins/Experimental/Water/Source/Runtime/Private/WaterBodyComponent.cpp b/Engine/Plugins/Experimental/Water/Source/Runtime/Private/WaterBodyComponent.cpp
--- a/Engine/Plugins/Experimental/Water/Source/Runtime/Private/WaterBodyComponent.cpp
+++ b/Engine/Plugins/Experimental/Water/Source/Runtime/Private/WaterBodyComponent.cpp
@@ -1559,6 +1559,29 @@ void UWaterBodyComponent::PostLoad()
 		UpdateWaterBodyRenderData();
 	}
 	RegisterOnUpdateWavesData(GetWaterWaves(), /* bRegister = */true);
+
+	if (IsTemplate() && ((OwningWaterZone != nullptr) || (!OwningWaterZone.GetAssetName().IsEmpty())))
+	{
+		TWeakObjectPtr<UWaterBodyComponent> WaterBodyComponent = this;
+
+		FMessageLog("MapCheck").Warning()
+			->AddToken(FUObjectToken::Create(this))
+			->AddToken(FTextToken::Create(FText::Format(
+				LOCTEXT("MapCheck_Warning_WaterZoneReferenceInCDO", "WaterBodyComponent contains a reference to {0} but is a template object. Please use the following action to mark the package dirty then save it."),
+				FText::FromString(OwningWaterZone.GetLongPackageName()))))
+			->AddToken(FActionToken::Create(LOCTEXT("MapCheck_MarkWaterBodyComponentPackageDirty", "Mark WaterBodyComponent package Dirty"), LOCTEXT("MapCheck_MarkWaterBodyComponentPackageDirty_Desc", "Marks the WaterBodyComponent's package dirty."),
+				FOnActionTokenExecuted::CreateLambda([WaterBodyComponent]()
+			{
+				if (WaterBodyComponent.IsValid())
+				{
+					WaterBodyComponent->MarkPackageDirty();
+				}
+			}), /*bInSingleUse = */true));
+
+		FMessageLog("MapCheck").Open(EMessageSeverity::Warning);
+		
+		OwningWaterZone.Reset();
+	}
 #endif // WITH_EDITOR
 }
 
