Overrides PostActorCreated to build WaterInfoMesh properly when spawning an actor using the UWorld::SpawnActor method (actors usually spawned with UActorFactory::CreateActor were creating the mesh using PostEditMove).
This fixes initial WaterBody being invalid when converting an actor to BP.

diff --git a/Engine/Plugins/Experimental/Water/Source/Runtime/Private/WaterBodyActor.cpp b/Engine/Plugins/Experimental/Water/Source/Runtime/Private/WaterBodyActor.cpp
--- a/Engine/Plugins/Experimental/Water/Source/Runtime/Private/WaterBodyActor.cpp
+++ b/Engine/Plugins/Experimental/Water/Source/Runtime/Private/WaterBodyActor.cpp
@@ -697,6 +697,17 @@ void AWaterBody::PackagingModeChanged(AActor* Actor, bool bIsExternal)
 	}
 }
 
+void AWaterBody::PostActorCreated()
+{
+	FOnWaterBodyChangedParams Params;
+
+	Params.PropertyChangedEvent.ChangeType = EPropertyChangeType::ValueSet;
+	Params.bShapeOrPositionChanged = true;
+	Params.bUserTriggered = false;
+	
+	WaterBodyComponent->OnWaterBodyChanged(Params);
+}
+
 #endif // WITH_EDITOR
 
 void AWaterBody::UnregisterAllComponents(bool bForReregister)
diff --git a/Engine/Plugins/Experimental/Water/Source/Runtime/Public/WaterBodyActor.h b/Engine/Plugins/Experimental/Water/Source/Runtime/Public/WaterBodyActor.h
--- a/Engine/Plugins/Experimental/Water/Source/Runtime/Public/WaterBodyActor.h
+++ b/Engine/Plugins/Experimental/Water/Source/Runtime/Public/WaterBodyActor.h
@@ -66,6 +66,8 @@ public:
 	virtual void GetActorDescProperties(FPropertyPairsMap& PropertyPairsMap) const override;
 
 	void PackagingModeChanged(AActor* Actor, bool bIsExternal);
+
+	virtual void PostActorCreated() override;
 #endif // WITH_EDITOR
 	
 	/** Returns the type of body */
