Water: fixed a crash when undoing the deletion of a water body by restoring the transactional property on water info component properties.

diff --git a/Engine/Plugins/Experimental/Water/Source/Runtime/Public/WaterBodyActor.h b/Engine/Plugins/Experimental/Water/Source/Runtime/Public/WaterBodyActor.h
--- a/Engine/Plugins/Experimental/Water/Source/Runtime/Public/WaterBodyActor.h
+++ b/Engine/Plugins/Experimental/Water/Source/Runtime/Public/WaterBodyActor.h
@@ -117,10 +117,10 @@ protected:
 	UPROPERTY(EditAnywhere, BlueprintReadOnly, Instanced, Category = Wave, DisplayName = "Waves Source", meta = (Tooltip = ""))
 	TObjectPtr<UWaterWavesBase> WaterWaves = nullptr;
 
-	UPROPERTY(NonTransactional, TextExportTransient, NonPIEDuplicateTransient)
+	UPROPERTY(TextExportTransient, NonPIEDuplicateTransient)
 	TObjectPtr<UWaterBodyInfoMeshComponent> WaterInfoMeshComponent;
 
-	UPROPERTY(NonTransactional, TextExportTransient, NonPIEDuplicateTransient)
+	UPROPERTY(TextExportTransient, NonPIEDuplicateTransient)
 	TObjectPtr<UWaterBodyInfoMeshComponent> DilatedWaterInfoMeshComponent;
 
 	UPROPERTY(NonTransactional, TextExportTransient, NonPIEDuplicateTransient)
