Water: removed some code which was added to try to prevent VSM rebuilds after water info renders. This code is now useless due to having to re-enable the bIsSceneCapture and also contains a memory leak.

diff --git a/Engine/Plugins/Experimental/Water/Source/Runtime/Private/WaterInfoRendering.cpp b/Engine/Plugins/Experimental/Water/Source/Runtime/Private/WaterInfoRendering.cpp
--- a/Engine/Plugins/Experimental/Water/Source/Runtime/Private/WaterInfoRendering.cpp
+++ b/Engine/Plugins/Experimental/Water/Source/Runtime/Private/WaterInfoRendering.cpp
@@ -721,21 +721,7 @@ static FSceneRenderer* CreateWaterInfoSceneRenderer(const FCreateWaterInfoSceneR
 		Extension->SetupView(ViewFamily, *View);
 	}
 
-	FSceneRenderer* SceneRenderer = FSceneRenderer::CreateSceneRenderer(&ViewFamily, nullptr);
-	FSceneViewStateReference* ViewStateReference = new FSceneViewStateReference();
-	ViewStateReference->Allocate(ViewFamily.Scene->GetFeatureLevel());
-	FSceneViewStateInterface* ViewStateInterface = ViewStateReference->GetReference();
-
-	if (UseVirtualShadowMaps(SceneRenderer->ShaderPlatform, ViewFamily.Scene->GetFeatureLevel()))
-	{
-		// This is needed to prevent the default cache being invalidated on every frame.
-		if (ViewStateInterface)
-		{
-			ViewStateInterface->AddVirtualShadowMapCache(ViewFamily.Scene);
-		}
-	}
-
-	return SceneRenderer;
+	return FSceneRenderer::CreateSceneRenderer(&ViewFamily, nullptr);
 }
 
 void UpdateWaterInfoRendering(
