Water: Set IsHLODRelevant for water info mesh components to ensure they never influence HLOD.

diff --git a/Engine/Plugins/Experimental/Water/Source/Runtime/Private/WaterBodyInfoMeshComponent.h b/Engine/Plugins/Experimental/Water/Source/Runtime/Private/WaterBodyInfoMeshComponent.h
--- a/Engine/Plugins/Experimental/Water/Source/Runtime/Private/WaterBodyInfoMeshComponent.h
+++ b/Engine/Plugins/Experimental/Water/Source/Runtime/Private/WaterBodyInfoMeshComponent.h
@@ -21,6 +21,8 @@ public:
 	UWaterBodyInfoMeshComponent(const FObjectInitializer& ObjectInitializer);
 
 	virtual FPrimitiveSceneProxy* CreateSceneProxy() override;
+	
+	virtual bool IsHLODRelevant() const { return false; }
 };
 
 struct FWaterBodyInfoMeshSceneProxy : public FStaticMeshSceneProxy
