Expose a callback after actors get attached/detached from their external package. This allows the actor to rename some objects it might have outered to its external package. This is intended as a temporary fix.

diff --git a/Engine/Plugins/Experimental/Water/Source/Runtime/Private/WaterBodyActor.cpp b/Engine/Plugins/Experimental/Water/Source/Runtime/Private/WaterBodyActor.cpp
--- a/Engine/Plugins/Experimental/Water/Source/Runtime/Private/WaterBodyActor.cpp
+++ b/Engine/Plugins/Experimental/Water/Source/Runtime/Private/WaterBodyActor.cpp
@@ -708,6 +708,40 @@ void AWaterBody::PostActorCreated()
 	WaterBodyComponent->OnWaterBodyChanged(Params);
 }
 
+void AWaterBody::OnDetachExternalPackage()
+{
+	// All the UStaticMesh are outered to the package, we need to make sure they remain so after the packaging mode of this actor is changed.
+	// If the actor used to be non-external and is now external, we need to ensure the UStaticMesh is outered to the external package of the actor and not the world.
+	// The inverse is true from external -> non external
+	TArray<TObjectPtr<UWaterBodyMeshComponent>> MeshComponents = { WaterInfoMeshComponent, DilatedWaterInfoMeshComponent };
+	MeshComponents.Append(WaterBodyStaticMeshComponents);
+
+	for (TObjectPtr<UWaterBodyMeshComponent> MeshComponent : MeshComponents)
+	{
+		if (IsValid(MeshComponent) && IsValid(MeshComponent->GetStaticMesh()))
+		{
+			MeshComponent->GetStaticMesh()->Rename(nullptr, GetPackage(), REN_DontCreateRedirectors | REN_ForceNoResetLoaders);
+		}
+	}
+}
+
+void AWaterBody::OnReattachExternalPackage()
+{
+	// All the UStaticMesh are outered to the package, we need to make sure they remain so after the packaging mode of this actor is changed.
+	// If the actor used to be non-external and is now external, we need to ensure the UStaticMesh is outered to the external package of the actor and not the world.
+	// The inverse is true from external -> non external
+	TArray<TObjectPtr<UWaterBodyMeshComponent>> MeshComponents = { WaterInfoMeshComponent, DilatedWaterInfoMeshComponent };
+	MeshComponents.Append(WaterBodyStaticMeshComponents);
+
+	for (TObjectPtr<UWaterBodyMeshComponent> MeshComponent : MeshComponents)
+	{
+		if (IsValid(MeshComponent) && IsValid(MeshComponent->GetStaticMesh()))
+		{
+			MeshComponent->GetStaticMesh()->Rename(nullptr, GetPackage(), REN_DontCreateRedirectors | REN_ForceNoResetLoaders);
+		}
+	}
+}
+
 #endif // WITH_EDITOR
 
 void AWaterBody::UnregisterAllComponents(bool bForReregister)
diff --git a/Engine/Plugins/Experimental/Water/Source/Runtime/Public/WaterBodyActor.h b/Engine/Plugins/Experimental/Water/Source/Runtime/Public/WaterBodyActor.h
--- a/Engine/Plugins/Experimental/Water/Source/Runtime/Public/WaterBodyActor.h
+++ b/Engine/Plugins/Experimental/Water/Source/Runtime/Public/WaterBodyActor.h
@@ -67,6 +67,9 @@ public:
 
 	void PackagingModeChanged(AActor* Actor, bool bIsExternal);
 
+	virtual void OnReattachExternalPackage() override;
+	virtual void OnDetachExternalPackage() override;
+
 	virtual void PostActorCreated() override;
 #endif // WITH_EDITOR
 	
@@ -117,10 +120,10 @@ protected:
 	UPROPERTY(EditAnywhere, BlueprintReadOnly, Instanced, Category = Wave, DisplayName = "Waves Source", meta = (Tooltip = ""))
 	TObjectPtr<UWaterWavesBase> WaterWaves = nullptr;
 
-	UPROPERTY(TextExportTransient, NonPIEDuplicateTransient)
+	UPROPERTY(NonTransactional, TextExportTransient, NonPIEDuplicateTransient)
 	TObjectPtr<UWaterBodyInfoMeshComponent> WaterInfoMeshComponent;
 
-	UPROPERTY(TextExportTransient, NonPIEDuplicateTransient)
+	UPROPERTY(NonTransactional, TextExportTransient, NonPIEDuplicateTransient)
 	TObjectPtr<UWaterBodyInfoMeshComponent> DilatedWaterInfoMeshComponent;
 
 	UPROPERTY(NonTransactional, TextExportTransient, NonPIEDuplicateTransient)
diff --git a/Engine/Source/Runtime/Engine/Classes/GameFramework/Actor.h b/Engine/Source/Runtime/Engine/Classes/GameFramework/Actor.h
--- a/Engine/Source/Runtime/Engine/Classes/GameFramework/Actor.h
+++ b/Engine/Source/Runtime/Engine/Classes/GameFramework/Actor.h
@@ -2183,6 +2183,11 @@ public:
 
 	/** Does this actor supports external packaging? */
 	virtual bool SupportsExternalPackaging() const;
+
+	/** Called when the we detach the actor from the external package. Temporary fix to allow the actor to fixup and rename any objects outered to the external package. */
+	virtual void OnDetachExternalPackage() {}
+	/** Called when the we reattach the actor to the external package. Temporary fix to allow the actor to fixup and rename any objects outered to the external package. */
+	virtual void OnReattachExternalPackage() {}
 #endif
 
 #if WITH_EDITOR
diff --git a/Engine/Source/Runtime/Engine/Private/Level.cpp b/Engine/Source/Runtime/Engine/Private/Level.cpp
--- a/Engine/Source/Runtime/Engine/Private/Level.cpp
+++ b/Engine/Source/Runtime/Engine/Private/Level.cpp
@@ -4093,6 +4093,7 @@ void ULevel::DetachAttachAllActorsPackages(bool bReattach)
 			if (Actor)
 			{
 				Actor->ReattachExternalPackage();
+				Actor->OnReattachExternalPackage();
 			}
 		}
 	}
@@ -4102,6 +4103,7 @@ void ULevel::DetachAttachAllActorsPackages(bool bReattach)
 		{
 			if (Actor)
 			{
+				Actor->OnDetachExternalPackage();
 				Actor->DetachExternalPackage();
 			}
 		}
