Water: Changed the encoding for velocity to improve precision and prevent having negative numbers. Certain texture pixel formats don't support a sign bit (R11G11B10_Float for ex). Water Info now stores normalized flow x/flow y in the RG channels.

diff --git a/Engine/Plugins/Experimental/Water/Source/Runtime/Private/WaterInfoRendering.cpp b/Engine/Plugins/Experimental/Water/Source/Runtime/Private/WaterInfoRendering.cpp
--- a/Engine/Plugins/Experimental/Water/Source/Runtime/Private/WaterInfoRendering.cpp
+++ b/Engine/Plugins/Experimental/Water/Source/Runtime/Private/WaterInfoRendering.cpp
@@ -688,11 +688,8 @@ static FSceneRenderer* CreateWaterInfoSceneRenderer(const FCreateWaterInfoSceneR
 	ViewInitOptions.ProjectionMatrix = Params.ProjectionMatrix;
 	ViewInitOptions.LODDistanceFactor = 0.001f;
 	ViewInitOptions.OverlayColor = FLinearColor::Black;
-	/* Must be set to false to prevent the renders from using different VSM page pool sizes leading to unnecessary reallocations.
-	** #todo_water: Temporarily set back to true since bIsSceneCapture forces bRequiresAlphaChannel in ViewFamily, leading to the proper SceneColor pixel format (PF_FloatRGBA).
-	** Without it, the Renderer would initialize SceneColor to PF_FloatRGB, resulting in incomplete velocities in the texture
-	*/
-	ViewInitOptions.bIsSceneCapture = true;
+	// Must be set to false to prevent the renders from using different VSM page pool sizes leading to unnecessary reallocations.
+	ViewInitOptions.bIsSceneCapture = false;
 
 	if (ViewFamily.Scene->GetWorld() != nullptr && ViewFamily.Scene->GetWorld()->GetWorldSettings() != nullptr)
 	{
diff --git a/Engine/Plugins/Experimental/Water/Source/Runtime/Private/WaterUtils.cpp b/Engine/Plugins/Experimental/Water/Source/Runtime/Private/WaterUtils.cpp
--- a/Engine/Plugins/Experimental/Water/Source/Runtime/Private/WaterUtils.cpp
+++ b/Engine/Plugins/Experimental/Water/Source/Runtime/Private/WaterUtils.cpp
@@ -128,6 +128,7 @@ float FWaterUtils::GetWaterMaxFlowVelocity(bool bIsRenderThread)
 FVector4f FWaterUtils::PackFlowData(float VelocityMagnitude, float DirectionAngle)
 {
 	check((DirectionAngle >= 0.f) && (DirectionAngle <= TWO_PI));
+	checkf(VelocityMagnitude >= 0., TEXT("We expect a positive magnitude"));
 
 	const float MaxVelocity = GetWaterMaxFlowVelocity(false);
 
