Water: Added the ERDGResourceExtractionFlags::AllowTransient to the transient intermediates of the water info passes to ensure they allocate from the transient pool and don't create expensive allocations.

diff --git a/Engine/Plugins/Experimental/Water/Source/Runtime/Private/WaterInfoRendering.cpp b/Engine/Plugins/Experimental/Water/Source/Runtime/Private/WaterInfoRendering.cpp
--- a/Engine/Plugins/Experimental/Water/Source/Runtime/Private/WaterInfoRendering.cpp
+++ b/Engine/Plugins/Experimental/Water/Source/Runtime/Private/WaterInfoRendering.cpp
@@ -494,7 +494,7 @@ static void UpdateWaterInfoRendering_RenderThread(
 		
 		// We currently can't have multiple scene renderers run within the same RDGBuilder. Therefore, we must
 		// extract the texture to allow it to survive until the water info pass which runs in a later RDG graph (however, still within the same frame).
-		GraphBuilder.QueueTextureExtraction(DepthTexture, &ExtractedDepthTexture);
+		GraphBuilder.QueueTextureExtraction(DepthTexture, &ExtractedDepthTexture, ERDGResourceExtractionFlags::AllowTransient);
 
 		GraphBuilder.Execute();
 
@@ -549,7 +549,7 @@ static void UpdateWaterInfoRendering_RenderThread(
 
 		// We currently can't have multiple scene renderers run within the same RDGBuilder. Therefore, we must
 		// extract the texture to allow it to survive until the water info pass which runs in a later RDG graph (however, still within the same frame).
-		GraphBuilder.QueueTextureExtraction(ColorTexture, &ExtractedColorTexture);
+		GraphBuilder.QueueTextureExtraction(ColorTexture, &ExtractedColorTexture, ERDGResourceExtractionFlags::AllowTransient);
 
 		GraphBuilder.Execute();
 		
