Water: Fixed an issue where all primitive components attached to landscape actors were collected and submitted to the water info ground pass.

diff --git a/Engine/Plugins/Experimental/Water/Source/Runtime/Private/WaterInfoRendering.cpp b/Engine/Plugins/Experimental/Water/Source/Runtime/Private/WaterInfoRendering.cpp
--- a/Engine/Plugins/Experimental/Water/Source/Runtime/Private/WaterInfoRendering.cpp
+++ b/Engine/Plugins/Experimental/Water/Source/Runtime/Private/WaterInfoRendering.cpp
@@ -758,21 +758,14 @@ void UpdateWaterInfoRendering(
 	CreateSceneRendererParams.ViewRotationMatrix.RemoveScaling();
 
 	TSet<FPrimitiveComponentId> ComponentsToRenderInDepthPass;
-	if (Context.GroundActors.Num() > 0)
+	if (Context.GroundPrimitiveComponents.Num() > 0)
 	{
-		ComponentsToRenderInDepthPass.Reserve(Context.GroundActors.Num());
-		for (TWeakObjectPtr<AActor> GroundActor : Context.GroundActors)
+		ComponentsToRenderInDepthPass.Reserve(Context.GroundPrimitiveComponents.Num());
+		for (TWeakObjectPtr<UPrimitiveComponent> GroundPrimComp : Context.GroundPrimitiveComponents)
 		{
-			if (GroundActor.IsValid())
+			if (GroundPrimComp.IsValid())
 			{
-				TInlineComponentArray<UPrimitiveComponent*> PrimComps(GroundActor.Get());
-				for (UPrimitiveComponent* PrimComp : PrimComps)
-				{
-					if (PrimComp)
-					{
-						ComponentsToRenderInDepthPass.Add(PrimComp->ComponentId);
-					}
-				}
+				ComponentsToRenderInDepthPass.Add(GroundPrimComp.Get()->ComponentId);
 			}
 		}
 	}
diff --git a/Engine/Plugins/Experimental/Water/Source/Runtime/Private/WaterInfoRendering.h b/Engine/Plugins/Experimental/Water/Source/Runtime/Private/WaterInfoRendering.h
--- a/Engine/Plugins/Experimental/Water/Source/Runtime/Private/WaterInfoRendering.h
+++ b/Engine/Plugins/Experimental/Water/Source/Runtime/Private/WaterInfoRendering.h
@@ -8,7 +8,7 @@ class AWaterZone;
 class UWaterBodyComponent;
 class FSceneInterface;
 class UTextureRenderTarget2D;
-class AActor;
+class UPrimitiveComponent;
 
 namespace UE::WaterInfo
 {
@@ -17,7 +17,7 @@ namespace UE::WaterInfo
 		AWaterZone* ZoneToRender = nullptr;
 		UTextureRenderTarget2D* TextureRenderTarget;
 		TArray<UWaterBodyComponent*> WaterBodies;
-		TArray<TWeakObjectPtr<AActor>> GroundActors;
+		TArray<TWeakObjectPtr<UPrimitiveComponent>> GroundPrimitiveComponents;
 		float CaptureZ;
 	};
 	
diff --git a/Engine/Plugins/Experimental/Water/Source/Runtime/Private/WaterZoneActor.cpp b/Engine/Plugins/Experimental/Water/Source/Runtime/Private/WaterZoneActor.cpp
--- a/Engine/Plugins/Experimental/Water/Source/Runtime/Private/WaterZoneActor.cpp
+++ b/Engine/Plugins/Experimental/Water/Source/Runtime/Private/WaterZoneActor.cpp
@@ -433,7 +433,8 @@ bool AWaterZone::UpdateWaterInfoTexture()
 		GroundZMin = TNumericLimits<float>::Max();
 		float GroundZMax = TNumericLimits<float>::Lowest();
 
-		TArray<TWeakObjectPtr<AActor>> GroundActors;
+		TArray<TWeakObjectPtr<UPrimitiveComponent>> GroundPrimitiveComponents;
+
 		const FBox WaterZoneBounds = GetZoneBounds();
 		for (ALandscapeProxy* LandscapeProxy : TActorRange<ALandscapeProxy>(World))
 		{
@@ -443,31 +444,27 @@ bool AWaterZone::UpdateWaterInfoTexture()
 			{
 				GroundZMin = FMath::Min(GroundZMin, LandscapeBox.Min.Z);
 				GroundZMax = FMath::Max(GroundZMax, LandscapeBox.Max.Z);
-				GroundActors.Add(LandscapeProxy);
+				TInlineComponentArray<ULandscapeComponent*> LandscapeComponents(LandscapeProxy);
+				GroundPrimitiveComponents.Append(LandscapeComponents);
 			}
 		}
 
-		// If we have no ground actors we need to set GroundZMin to be something sensible because it won't have been set above.
-		if (GroundActors.Num() == 0)
+		// If we have no ground components we need to set GroundZMin to be something sensible because it won't have been set above.
+		if (GroundPrimitiveComponents.Num() == 0)
 		{
 			GroundZMax = WaterZMax;
 			GroundZMin = WaterZMin - CVarWaterFallbackDepth.GetValueOnGameThread();
 		}
 
 #if WITH_EDITOR
-		// Check all the ground actors have complete shader maps before we try to render them into the water info texture
-		for (TWeakObjectPtr<AActor> GroundActorPtr : GroundActors)
+		// Check all the ground components have complete shader maps before we try to render them into the water info texture
+		for (TWeakObjectPtr<UPrimitiveComponent> GroundPrimCompPtr : GroundPrimitiveComponents)
 		{
-			if (AActor* GroundActor = GroundActorPtr.Get())
+			if (UPrimitiveComponent* GroundPrimComp = GroundPrimCompPtr.Get())
 			{
-				TInlineComponentArray<UPrimitiveComponent*> PrimitiveComponents(GroundActor);
-
-				for (UPrimitiveComponent* PrimitiveComponent : PrimitiveComponents)
-				{
 					TArray<UMaterialInterface*> TmpUsedMaterials;
-					PrimitiveComponent->GetUsedMaterials(TmpUsedMaterials, false);
+					GroundPrimComp->GetUsedMaterials(TmpUsedMaterials, false);
 					UsedMaterials.Append(TmpUsedMaterials);
-				}
 			}
 		}
 
@@ -503,7 +500,7 @@ bool AWaterZone::UpdateWaterInfoTexture()
 		UE::WaterInfo::FRenderingContext Context;
 		Context.ZoneToRender = this;
 		Context.WaterBodies = WaterBodiesToRender;
-		Context.GroundActors = MoveTemp(GroundActors);
+		Context.GroundPrimitiveComponents = MoveTemp(GroundPrimitiveComponents);
 		Context.CaptureZ = FMath::Max(WaterZMax, GroundZMax) + CaptureZOffset;
 		Context.TextureRenderTarget = WaterInfoTexture;
 
