Water: fixed accidental removal of NonTransactional flag on the water mesh properties

diff --git a/Engine/Plugins/Experimental/Water/Source/Runtime/Public/WaterBodyActor.h b/Engine/Plugins/Experimental/Water/Source/Runtime/Public/WaterBodyActor.h
--- a/Engine/Plugins/Experimental/Water/Source/Runtime/Public/WaterBodyActor.h
+++ b/Engine/Plugins/Experimental/Water/Source/Runtime/Public/WaterBodyActor.h
@@ -120,13 +120,13 @@ protected:
 	UPROPERTY(EditAnywhere, BlueprintReadOnly, Instanced, Category = Wave, DisplayName = "Waves Source", meta = (Tooltip = ""))
 	TObjectPtr<UWaterWavesBase> WaterWaves = nullptr;
 
-	UPROPERTY(NonTransactional, TextExportTransient, NonPIEDuplicateTransient)
+	UPROPERTY(TextExportTransient, NonPIEDuplicateTransient)
 	TObjectPtr<UWaterBodyInfoMeshComponent> WaterInfoMeshComponent;
 
-	UPROPERTY(NonTransactional, TextExportTransient, NonPIEDuplicateTransient)
+	UPROPERTY(TextExportTransient, NonPIEDuplicateTransient)
 	TObjectPtr<UWaterBodyInfoMeshComponent> DilatedWaterInfoMeshComponent;
 
-	UPROPERTY(NonTransactional, TextExportTransient, NonPIEDuplicateTransient)
+	UPROPERTY(TextExportTransient, NonPIEDuplicateTransient)
 	TArray<TObjectPtr<UWaterBodyStaticMeshComponent>> WaterBodyStaticMeshComponents;
 
 #if WITH_EDITOR
