Water: prevent running fixup code on new maps that haven't yet been saved.

GetLinkerCustomVersion returns -1 for objects in unsaved maps as they have no Linker to query a custom version number from.

diff --git a/Engine/Plugins/Experimental/Water/Source/Runtime/Private/WaterBodyOceanComponent.cpp b/Engine/Plugins/Experimental/Water/Source/Runtime/Private/WaterBodyOceanComponent.cpp
--- a/Engine/Plugins/Experimental/Water/Source/Runtime/Private/WaterBodyOceanComponent.cpp
+++ b/Engine/Plugins/Experimental/Water/Source/Runtime/Private/WaterBodyOceanComponent.cpp
@@ -470,12 +470,16 @@ void UWaterBodyOceanComponent::OnPostRegisterAllComponents()
 	Super::OnPostRegisterAllComponents();
 
 #if WITH_EDITOR
-	// In this version the ocean was changed to not be strongly coupled to the water zone and instead rely on the user to keep the mesh extent in sync with the zone extent
-	if (GetLinkerCustomVersion(FFortniteMainBranchObjectVersion::GUID) < FFortniteMainBranchObjectVersion::WaterBodyStaticMeshComponents)
-	{
-		// Trigger a rebuild of the render data manually since this doesn't occur on load anymore. 
-		// The mesh extent will be out of sync for old oceans so ensure they rebuild now:
-		FillWaterZoneWithOcean();
+	// Only run the fixup code when the object was loaded. This is required as any newly created and not-yet-saved objects will have no linker and therefore return an invalid custom version.
+	if (HasAnyFlags(RF_WasLoaded))
+	{	
+		// In this version the ocean was changed to not be strongly coupled to the water zone and instead rely on the user to keep the mesh extent in sync with the zone extent
+		if (GetLinkerCustomVersion(FFortniteMainBranchObjectVersion::GUID) < FFortniteMainBranchObjectVersion::WaterBodyStaticMeshComponents)
+		{
+			// Trigger a rebuild of the render data manually since this doesn't occur on load anymore. 
+			// The mesh extent will be out of sync for old oceans so ensure they rebuild now:
+			FillWaterZoneWithOcean();
+		}
 	}
 #endif // WITH_EDITOR
 }
