Water: Fixed a bug in the water quad tree where it would not account for the new OceanExtent parameter and instead always fill the quadtree with ocean, even if the ocean has a different extent parameter.

diff --git a/Engine/Plugins/Experimental/Water/Source/Runtime/Private/WaterMeshComponent.cpp b/Engine/Plugins/Experimental/Water/Source/Runtime/Private/WaterMeshComponent.cpp
--- a/Engine/Plugins/Experimental/Water/Source/Runtime/Private/WaterMeshComponent.cpp
+++ b/Engine/Plugins/Experimental/Water/Source/Runtime/Private/WaterMeshComponent.cpp
@@ -519,6 +519,9 @@ void UWaterMeshComponent::RebuildWaterMesh(float InTileSize, const FIntPoint& In
 				Polygon.Reserve(PolyLineVertices.Num());
 				Algo::Transform(PolyLineVertices, Polygon, [](const FVector& Vertex) { return FVector2D(Vertex); });
 
+				FBox OceanBounds = WaterBodyComponent->Bounds.GetBox();
+				OceanBounds.Max.Z += WaterBodyComponent->GetMaxWaveHeight();
+
 #if !(UE_BUILD_SHIPPING || UE_BUILD_TEST)
 				if (!!CVarWaterMeshShowTileGenerationGeometry.GetValueOnGameThread())
 				{
@@ -530,11 +533,12 @@ void UWaterMeshComponent::RebuildWaterMesh(float InTileSize, const FIntPoint& In
 						const FVector2D& Point1 = Polygon[(i + 1) % NumVertices];
 						DrawDebugLine(GetWorld(), FVector(Point0.X, Point0.Y, Z), FVector(Point1.X, Point1.Y, Z), FColor::Blue);
 					}
+
+					DrawDebugBox(GetWorld(), OceanBounds.GetCenter(), OceanBounds.GetExtent(), FColor::Blue);
 				}
 #endif //!(UE_BUILD_SHIPPING || UE_BUILD_TEST)
 
-				const FBox OceanBounds = Actor->GetComponentsBoundingBox();
-				WaterQuadTree.AddOcean(Polygon, FVector2D(OceanBounds.Min.Z, OceanBounds.Max.Z + WaterBodyComponent->GetMaxWaveHeight()), WaterBodyRenderDataIndex);
+				WaterQuadTree.AddOcean(Polygon, OceanBounds, WaterBodyRenderDataIndex);
 			}
 
 			// Place far mesh height just below the ocean level
diff --git a/Engine/Plugins/Experimental/Water/Source/Runtime/Private/WaterQuadTree.cpp b/Engine/Plugins/Experimental/Water/Source/Runtime/Private/WaterQuadTree.cpp
--- a/Engine/Plugins/Experimental/Water/Source/Runtime/Private/WaterQuadTree.cpp
+++ b/Engine/Plugins/Experimental/Water/Source/Runtime/Private/WaterQuadTree.cpp
@@ -611,11 +611,11 @@ void FWaterQuadTree::AddWaterTilesInsideBounds(const FBox& InBounds, uint32 InWa
 	NodeData.Nodes[0].AddNodes(NodeData, FBox(FVector(TileRegion.Min, 0.0f), FVector(TileRegion.Max, 0.0f)),  InBounds, InWaterBodyIndex, TreeDepth, 0);
 }
 
-void FWaterQuadTree::AddOcean(const TArray<FVector2D>& InPoly, const FVector2D& InZBounds, uint32 InWaterBodyIndex)
+void FWaterQuadTree::AddOcean(const TArray<FVector2D>& InPoly, const FBox& InOceanBounds, uint32 InWaterBodyIndex)
 {
 	check(!bIsReadOnly);
-	const FBox2D OceanBounds(FVector2D(GetBounds().Min), FVector2D(GetBounds().Max));
-	AddOceanRecursive(InPoly, OceanBounds, InZBounds, true, TreeDepth * 2, InWaterBodyIndex);
+	const FBox2D OceanBounds(FVector2D(InOceanBounds.Min), FVector2D(InOceanBounds.Max));
+	AddOceanRecursive(InPoly, OceanBounds, FVector2D(InOceanBounds.Min.Z, InOceanBounds.Max.Z), true, TreeDepth * 2, InWaterBodyIndex);
 }
 
 void FWaterQuadTree::AddLake(const TArray<FVector2D>& InPoly, const FBox& InLakeBounds, uint32 InWaterBodyIndex)
diff --git a/Engine/Plugins/Experimental/Water/Source/Runtime/Public/WaterQuadTree.h b/Engine/Plugins/Experimental/Water/Source/Runtime/Public/WaterQuadTree.h
--- a/Engine/Plugins/Experimental/Water/Source/Runtime/Public/WaterQuadTree.h
+++ b/Engine/Plugins/Experimental/Water/Source/Runtime/Public/WaterQuadTree.h
@@ -131,7 +131,7 @@ struct FWaterQuadTree
 	void AddWaterTilesInsideBounds(const FBox& InBounds, uint32 InWaterBodyIndex);
 
 	/** Add Ocean by giving a closed spline that represents the land mass for this ocean. */
-	void AddOcean(const TArray<FVector2D>& InPoly, const FVector2D& InZBounds, uint32 InWaterBodyIndex);
+	void AddOcean(const TArray<FVector2D>& InPoly, const FBox& InOceanBounds, uint32 InWaterBodyIndex);
 
 	/** Add Lake by giving a closed spline that represents the lake */
 	void AddLake(const TArray<FVector2D>& InPoly, const FBox& InLakeBounds, uint32 InWaterBodyIndex);
