Water: disable editing mesh section size when static mesh is disabled.

diff --git a/Engine/Plugins/Experimental/Water/Source/Runtime/Public/WaterBodyStaticMeshSettings.h b/Engine/Plugins/Experimental/Water/Source/Runtime/Public/WaterBodyStaticMeshSettings.h
--- a/Engine/Plugins/Experimental/Water/Source/Runtime/Public/WaterBodyStaticMeshSettings.h
+++ b/Engine/Plugins/Experimental/Water/Source/Runtime/Public/WaterBodyStaticMeshSettings.h
@@ -18,6 +18,6 @@ struct FWaterBodyStaticMeshSettings
 	bool bSectionWaterBodyStaticMesh = false;
 
 	/** The size of each section tile if sectioning is enabled. */
-	UPROPERTY(Category = StaticMesh, EditAnywhere, AdvancedDisplay, meta = (EditCondition = "bSectionWaterBodyStaticMesh"))
+	UPROPERTY(Category = StaticMesh, EditAnywhere, AdvancedDisplay, meta = (EditCondition = "bEnableWaterBodyStaticMesh && bSectionWaterBodyStaticMesh"))
 	double SectionSize = 8192.0;
 };
