PSO precaching fixes and improvements:

Cherry pick 25619283: Fix for decal being rendered a second time in the translucent pass.
Cherry pick 25581812: PSO precaching fixes for nanite materials
Cherry pick 25507656:
- Make sure water material received post load before precaching them
- re-enable PSO precaching for base pass after auto resolve issues which removed the PSO precaching

Extra fixes:
- Use correct stream index for PSO precaching vertex elements in FHairCardsVertexFactory

- Make sure water material received post load before precaching them
- re-enable PSO precaching for base pass after auto resolve issues which removed the PSO precaching

diff --git a/Engine/Plugins/Experimental/Water/Source/Runtime/Private/WaterBodyComponent.cpp b/Engine/Plugins/Experimental/Water/Source/Runtime/Private/WaterBodyComponent.cpp
--- a/Engine/Plugins/Experimental/Water/Source/Runtime/Private/WaterBodyComponent.cpp
+++ b/Engine/Plugins/Experimental/Water/Source/Runtime/Private/WaterBodyComponent.cpp
@@ -1541,14 +1541,17 @@ void UWaterBodyComponent::PostLoad()
 		SetupPrecachePSOParams(PrecachePSOParams);
 		if (WaterMaterial)
 		{
+			WaterMaterial->ConditionalPostLoad();
 			WaterMaterial->PrecachePSOs(&FLocalVertexFactory::StaticType, PrecachePSOParams);
 		}
 		if (UnderwaterPostProcessMaterial)
 		{
+			UnderwaterPostProcessMaterial->ConditionalPostLoad();
 			UnderwaterPostProcessMaterial->PrecachePSOs(&FLocalVertexFactory::StaticType, PrecachePSOParams);
 		}
 		if (WaterInfoMaterial)
 		{
+			WaterInfoMaterial->ConditionalPostLoad();
 			WaterInfoMaterial->PrecachePSOs(&FLocalVertexFactory::StaticType, PrecachePSOParams);
 		}
 	}
diff --git a/Engine/Plugins/Runtime/HairStrands/Source/HairStrandsCore/Private/HairCardsVertexFactory.cpp b/Engine/Plugins/Runtime/HairStrands/Source/HairStrandsCore/Private/HairCardsVertexFactory.cpp
--- a/Engine/Plugins/Runtime/HairStrands/Source/HairStrandsCore/Private/HairCardsVertexFactory.cpp
+++ b/Engine/Plugins/Runtime/HairStrands/Source/HairStrandsCore/Private/HairCardsVertexFactory.cpp
@@ -257,7 +257,7 @@ void FHairCardsVertexFactory::ValidateCompiledResult(const FVertexFactoryType* T
 void FHairCardsVertexFactory::GetPSOPrecacheVertexFetchElements(EVertexInputStreamType VertexInputStreamType, FVertexDeclarationElementList& Elements)
 {
 	// Manual vertex fetch is available for this factory so only primitive ID stream is used
-	Elements.Add(FVertexElement(1, 0, VET_UInt, 13, 0, true));
+	Elements.Add(FVertexElement(0, 0, VET_UInt, 13, 0, true));
 }
 
 void FHairCardsVertexFactory::SetData(const FDataType& InData)
