Fixed OnDetachExternalPackage not actually renaming the meshes into the levels package.

diff --git a/Engine/Plugins/Experimental/Water/Source/Runtime/Private/WaterBodyActor.cpp b/Engine/Plugins/Experimental/Water/Source/Runtime/Private/WaterBodyActor.cpp
--- a/Engine/Plugins/Experimental/Water/Source/Runtime/Private/WaterBodyActor.cpp
+++ b/Engine/Plugins/Experimental/Water/Source/Runtime/Private/WaterBodyActor.cpp
@@ -720,7 +720,7 @@ void AWaterBody::OnDetachExternalPackage()
 	{
 		if (IsValid(MeshComponent) && IsValid(MeshComponent->GetStaticMesh()))
 		{
-			MeshComponent->GetStaticMesh()->Rename(nullptr, GetPackage(), REN_DontCreateRedirectors | REN_ForceNoResetLoaders);
+			MeshComponent->GetStaticMesh()->Rename(nullptr, GetLevel()->GetPackage(), REN_DontCreateRedirectors | REN_ForceNoResetLoaders);
 		}
 	}
 }
