Missing Level.h include.

diff --git a/Engine/Plugins/Experimental/Water/Source/Runtime/Private/WaterBodyActor.cpp b/Engine/Plugins/Experimental/Water/Source/Runtime/Private/WaterBodyActor.cpp
--- a/Engine/Plugins/Experimental/Water/Source/Runtime/Private/WaterBodyActor.cpp
+++ b/Engine/Plugins/Experimental/Water/Source/Runtime/Private/WaterBodyActor.cpp
@@ -25,6 +25,7 @@
 #include "Editor/EditorEngine.h"
 #include "Landscape.h"
 #include "Engine/StaticMesh.h"
+#include "Engine/Level.h"
 #endif
 
 #define LOCTEXT_NAMESPACE "Water"
