Water: Temporary fix for cooking issue where packaging mode changed moves the static mesh before the actor moves levels and gets renamed into the wrong level package.

diff --git a/Engine/Plugins/Experimental/Water/Source/Runtime/Private/WaterBodyActor.cpp b/Engine/Plugins/Experimental/Water/Source/Runtime/Private/WaterBodyActor.cpp
--- a/Engine/Plugins/Experimental/Water/Source/Runtime/Private/WaterBodyActor.cpp
+++ b/Engine/Plugins/Experimental/Water/Source/Runtime/Private/WaterBodyActor.cpp
@@ -683,17 +683,21 @@ void AWaterBody::GetActorDescProperties(FPropertyPairsMap& PropertyPairsMap) con
 
 void AWaterBody::PackagingModeChanged(AActor* Actor, bool bIsExternal)
 {
-	// All the UStaticMesh are outered to the package, we need to make sure they remain so after the packaging mode of this actor is changed.
-	// If the actor used to be non-external and is now external, we need to ensure the UStaticMesh is outered to the external package of the actor and not the world.
-	// The inverse is true from external -> non external
-	TArray<TObjectPtr<UWaterBodyMeshComponent>> MeshComponents = { WaterInfoMeshComponent, DilatedWaterInfoMeshComponent };
-	MeshComponents.Append(WaterBodyStaticMeshComponents);
-
-	for (TObjectPtr<UWaterBodyMeshComponent> MeshComponent : MeshComponents)
+	// This callback is needed in the editor to ensure the static meshes follow the package of the actor. This is handled by the world partition system during cook for WP maps and DetachExternalPackage for non-WP case.
+	if (!IsRunningCookCommandlet())
 	{
-		if (IsValid(MeshComponent) && IsValid(MeshComponent->GetStaticMesh()))
+		// All the UStaticMesh are outered to the package, we need to make sure they remain so after the packaging mode of this actor is changed.
+		// If the actor used to be non-external and is now external, we need to ensure the UStaticMesh is outered to the external package of the actor and not the world.
+		// The inverse is true from external -> non external
+		TArray<TObjectPtr<UWaterBodyMeshComponent>> MeshComponents = { WaterInfoMeshComponent, DilatedWaterInfoMeshComponent };
+		MeshComponents.Append(WaterBodyStaticMeshComponents);
+
+		for (TObjectPtr<UWaterBodyMeshComponent> MeshComponent : MeshComponents)
 		{
-			MeshComponent->GetStaticMesh()->Rename(nullptr, GetPackage(), REN_DontCreateRedirectors | REN_ForceNoResetLoaders);
+			if (IsValid(MeshComponent) && IsValid(MeshComponent->GetStaticMesh()))
+			{
+				MeshComponent->GetStaticMesh()->Rename(nullptr, GetPackage(), REN_DontCreateRedirectors | REN_ForceNoResetLoaders);
+			}
 		}
 	}
 }
