Water: Fixed the snapping of water info mesh components to also include scale.
Fixes issues with water bodies with non-1 scale.

diff --git a/Engine/Plugins/Experimental/Water/Source/Runtime/Private/WaterBodyActor.cpp b/Engine/Plugins/Experimental/Water/Source/Runtime/Private/WaterBodyActor.cpp
--- a/Engine/Plugins/Experimental/Water/Source/Runtime/Private/WaterBodyActor.cpp
+++ b/Engine/Plugins/Experimental/Water/Source/Runtime/Private/WaterBodyActor.cpp
@@ -89,13 +89,13 @@ void AWaterBody::PreRegisterAllComponents()
 	if (IsValid(WaterInfoMeshComponent))
 	{
 		WaterInfoMeshComponent->SetMobility(WaterBodyComponent->Mobility);
-		WaterInfoMeshComponent->AttachToComponent(WaterBodyComponent, FAttachmentTransformRules::KeepRelativeTransform);
+		WaterInfoMeshComponent->AttachToComponent(WaterBodyComponent, FAttachmentTransformRules::SnapToTargetIncludingScale);
 	}
 
 	if (IsValid(DilatedWaterInfoMeshComponent))
 	{
 		DilatedWaterInfoMeshComponent->SetMobility(WaterBodyComponent->Mobility);
-		DilatedWaterInfoMeshComponent->AttachToComponent(WaterBodyComponent, FAttachmentTransformRules::KeepRelativeTransform);
+		DilatedWaterInfoMeshComponent->AttachToComponent(WaterBodyComponent, FAttachmentTransformRules::SnapToTargetIncludingScale);
 	}
 }
 
