Water: re-implemented r.water.waterinfo.showsceneproxies cvar to enable debug view of water info scene proxies in the viewport.

diff --git a/Engine/Plugins/Experimental/Water/Source/Runtime/Private/WaterBodyInfoMeshComponent.cpp b/Engine/Plugins/Experimental/Water/Source/Runtime/Private/WaterBodyInfoMeshComponent.cpp
--- a/Engine/Plugins/Experimental/Water/Source/Runtime/Private/WaterBodyInfoMeshComponent.cpp
+++ b/Engine/Plugins/Experimental/Water/Source/Runtime/Private/WaterBodyInfoMeshComponent.cpp
@@ -2,9 +2,25 @@
 
 #include "WaterBodyInfoMeshComponent.h"
 #include "StaticMeshSceneProxy.h"
+#include  "UObject/UObjectIterator.h"
 
 #include UE_INLINE_GENERATED_CPP_BY_NAME(WaterBodyInfoMeshComponent)
 
+static void OnCVarWaterInfoSceneProxiesValueChanged(IConsoleVariable*)
+{
+	for (UWaterBodyInfoMeshComponent* WaterBodyInfoMeshComponent : TObjectRange<UWaterBodyInfoMeshComponent>(RF_ClassDefaultObject | RF_ArchetypeObject, true, EInternalObjectFlags::Garbage))
+	{
+		WaterBodyInfoMeshComponent->MarkRenderStateDirty();
+	}
+}
+
+static TAutoConsoleVariable<bool> CVarShowWaterInfoSceneProxies(
+	TEXT("r.Water.WaterInfo.ShowSceneProxies"),
+	false,
+	TEXT("When enabled, always shows the water scene proxies in the main viewport. Useful for debugging only"),
+	FConsoleVariableDelegate::CreateStatic(OnCVarWaterInfoSceneProxiesValueChanged),
+	ECVF_RenderThreadSafe
+);
 
 UWaterBodyInfoMeshComponent::UWaterBodyInfoMeshComponent(const FObjectInitializer& ObjectInitializer)
 	: Super(ObjectInitializer)
@@ -27,3 +43,8 @@ FWaterBodyInfoMeshSceneProxy::FWaterBodyInfoMeshSceneProxy(UWaterBodyInfoMeshCom
 {
 	SetEnabled(false);
 }
+
+void FWaterBodyInfoMeshSceneProxy::SetEnabled(bool bInEnabled)
+{
+	SetForceHidden(!(bInEnabled || CVarShowWaterInfoSceneProxies.GetValueOnAnyThread()));
+}
diff --git a/Engine/Plugins/Experimental/Water/Source/Runtime/Private/WaterBodyInfoMeshComponent.h b/Engine/Plugins/Experimental/Water/Source/Runtime/Private/WaterBodyInfoMeshComponent.h
--- a/Engine/Plugins/Experimental/Water/Source/Runtime/Private/WaterBodyInfoMeshComponent.h
+++ b/Engine/Plugins/Experimental/Water/Source/Runtime/Private/WaterBodyInfoMeshComponent.h
@@ -36,6 +36,6 @@ public:
 		return reinterpret_cast<size_t>(&UniquePointer);
 	}
 
-	void SetEnabled(bool bInEnabled) { SetForceHidden(!bInEnabled); }
+	void SetEnabled(bool bInEnabled);
 };
 
