Water: only update the water zone pointers when water bodies are updated via a user-triggered change.

diff --git a/Engine/Plugins/Experimental/Water/Source/Runtime/Private/WaterBodyComponent.cpp b/Engine/Plugins/Experimental/Water/Source/Runtime/Private/WaterBodyComponent.cpp
--- a/Engine/Plugins/Experimental/Water/Source/Runtime/Private/WaterBodyComponent.cpp
+++ b/Engine/Plugins/Experimental/Water/Source/Runtime/Private/WaterBodyComponent.cpp
@@ -1381,6 +1381,7 @@ void UWaterBodyComponent::UpdateAll(const FOnWaterBodyChangedParams& InParams)
 		return;
 	}
 
+	const bool bUserTriggered = InParams.bUserTriggered;
 	bool bShapeOrPositionChanged = InParams.bShapeOrPositionChanged;
 	
 	if (GIsEditor || IsBodyDynamic())
@@ -1405,7 +1406,11 @@ void UWaterBodyComponent::UpdateAll(const FOnWaterBodyChangedParams& InParams)
 
 			UpdateExclusionVolumes();
 
-			UpdateWaterZones();
+			// Only update the water zones if this was a user triggered change. Rely on the serialized water zone pointer normally.
+			if (bUserTriggered)
+			{
+				UpdateWaterZones();
+			}
 		}
 
 		// Finally, generate the body once again, this time with the updated list of exclusion volumes
@@ -1426,12 +1431,6 @@ void UWaterBodyComponent::UpdateAll(const FOnWaterBodyChangedParams& InParams)
 		UpdateWaterSpriteComponent();
 #endif
 	}
-
-	// #todo_water: this should be reconsidered along with the entire concept of non-dynamic water bodies now that evening is runtime-generated.
-	if (bShapeOrPositionChanged)
-	{
-		UpdateWaterZones();
-	}
 }
 
 void UWaterBodyComponent::UpdateSplineComponent()
