Implemented a method for actors to populate the duplication seed map when duplicating for PIE.

diff --git a/Engine/Plugins/Experimental/Water/Source/Runtime/Private/WaterBodyActor.cpp b/Engine/Plugins/Experimental/Water/Source/Runtime/Private/WaterBodyActor.cpp
--- a/Engine/Plugins/Experimental/Water/Source/Runtime/Private/WaterBodyActor.cpp
+++ b/Engine/Plugins/Experimental/Water/Source/Runtime/Private/WaterBodyActor.cpp
@@ -73,8 +73,6 @@ AWaterBody::AWaterBody(const FObjectInitializer& ObjectInitializer)
 	TargetWaveMaskDepth_DEPRECATED = 2048.f;
 	bCanAffectNavigation_DEPRECATED = false;
 	bFillCollisionUnderWaterBodiesForNavmesh_DEPRECATED = false;
-
-	OnPackagingModeChanged.AddUObject(this, &AWaterBody::PackagingModeChanged);
 #endif // WITH_EDITORONLY_DATA
 
 }
@@ -518,6 +516,22 @@ void AWaterBody::DeprecateData()
 		// Some deprecated data on the actor were also later deprecated on the component (e.g. bFillCollisionUnderWaterBodiesForNavmesh_DEPRECATED), make sure it runs its own deprecation operations on it : 
 		WaterBodyComponent->DeprecateData();
 	}
+
+	// With this version bump, we introduced a fix that allows water static meshes to be outered to the actor and still avoid the cost of duplicating them for PIE. It is now safe to rename them back to have a proper outer chain with the actor as the outer.
+	if (GetLinkerCustomVersion(FFortniteMainBranchObjectVersion::GUID) < FFortniteMainBranchObjectVersion::WaterBodyStaticMeshRename)
+	{
+		TArray<TObjectPtr<UWaterBodyMeshComponent>> MeshComponents = { WaterInfoMeshComponent, DilatedWaterInfoMeshComponent };
+		MeshComponents.Append(WaterBodyStaticMeshComponents);
+
+		for (TObjectPtr<UWaterBodyMeshComponent> MeshComponent : MeshComponents)
+		{
+			if (IsValid(MeshComponent) && IsValid(MeshComponent->GetStaticMesh()))
+			{
+				UStaticMesh* StaticMesh = MeshComponent->GetStaticMesh();
+				StaticMesh->Rename(nullptr, this, REN_DontCreateRedirectors | REN_ForceNoResetLoaders);
+			}
+		}
+	}
 #endif // WITH_EDITORONLY_DATA
 }
 
@@ -627,6 +641,24 @@ bool AWaterBody::IsHLODRelevant() const
 }
 
 #if WITH_EDITOR
+void AWaterBody::PopulatePIEDuplicationSeed(AActor::FDuplicationSeedInterface& DuplicationSeed)
+{
+	Super::PopulatePIEDuplicationSeed(DuplicationSeed);
+
+	// Avoid copying expensive UStaticMesh when entering pie, instead use the duplication seed to have the PIE instance share the same mesh pointer.
+	TArray<TObjectPtr<UWaterBodyMeshComponent>> MeshComponents = { WaterInfoMeshComponent, DilatedWaterInfoMeshComponent };
+	MeshComponents.Append(WaterBodyStaticMeshComponents);
+
+	for (TObjectPtr<UWaterBodyMeshComponent> MeshComponent : MeshComponents)
+	{
+		if (IsValid(MeshComponent) && IsValid(MeshComponent->GetStaticMesh()))
+		{
+			UStaticMesh* StaticMesh = MeshComponent->GetStaticMesh();
+			DuplicationSeed.AddEntry(StaticMesh, StaticMesh);
+		}
+	}
+}
+
 void AWaterBody::SetActorHiddenInGame(bool bNewHidden)
 {
 	// It's kinda sad that being hidden in game for the actor doesn't end up calling OnHiddenInGameChanged on the component but intercepting it on the actor, we can inform the component that it needs
@@ -681,27 +713,6 @@ void AWaterBody::GetActorDescProperties(FPropertyPairsMap& PropertyPairsMap) con
 	}
 }
 
-void AWaterBody::PackagingModeChanged(AActor* Actor, bool bIsExternal)
-{
-	// This callback is needed in the editor to ensure the static meshes follow the package of the actor. This is handled by the world partition system during cook for WP maps and DetachExternalPackage for non-WP case.
-	if (!IsRunningCookCommandlet())
-	{
-		// All the UStaticMesh are outered to the package, we need to make sure they remain so after the packaging mode of this actor is changed.
-		// If the actor used to be non-external and is now external, we need to ensure the UStaticMesh is outered to the external package of the actor and not the world.
-		// The inverse is true from external -> non external
-		TArray<TObjectPtr<UWaterBodyMeshComponent>> MeshComponents = { WaterInfoMeshComponent, DilatedWaterInfoMeshComponent };
-		MeshComponents.Append(WaterBodyStaticMeshComponents);
-
-		for (TObjectPtr<UWaterBodyMeshComponent> MeshComponent : MeshComponents)
-		{
-			if (IsValid(MeshComponent) && IsValid(MeshComponent->GetStaticMesh()))
-			{
-				MeshComponent->GetStaticMesh()->Rename(nullptr, GetPackage(), REN_DontCreateRedirectors | REN_ForceNoResetLoaders);
-			}
-		}
-	}
-}
-
 void AWaterBody::PostActorCreated()
 {
 	FOnWaterBodyChangedParams Params;
@@ -713,40 +724,6 @@ void AWaterBody::PostActorCreated()
 	WaterBodyComponent->OnWaterBodyChanged(Params);
 }
 
-void AWaterBody::OnDetachExternalPackage()
-{
-	// All the UStaticMesh are outered to the package, we need to make sure they remain so after the packaging mode of this actor is changed.
-	// If the actor used to be non-external and is now external, we need to ensure the UStaticMesh is outered to the external package of the actor and not the world.
-	// The inverse is true from external -> non external
-	TArray<TObjectPtr<UWaterBodyMeshComponent>> MeshComponents = { WaterInfoMeshComponent, DilatedWaterInfoMeshComponent };
-	MeshComponents.Append(WaterBodyStaticMeshComponents);
-
-	for (TObjectPtr<UWaterBodyMeshComponent> MeshComponent : MeshComponents)
-	{
-		if (IsValid(MeshComponent) && IsValid(MeshComponent->GetStaticMesh()))
-		{
-			MeshComponent->GetStaticMesh()->Rename(nullptr, GetLevel()->GetPackage(), REN_DontCreateRedirectors | REN_ForceNoResetLoaders);
-		}
-	}
-}
-
-void AWaterBody::OnReattachExternalPackage()
-{
-	// All the UStaticMesh are outered to the package, we need to make sure they remain so after the packaging mode of this actor is changed.
-	// If the actor used to be non-external and is now external, we need to ensure the UStaticMesh is outered to the external package of the actor and not the world.
-	// The inverse is true from external -> non external
-	TArray<TObjectPtr<UWaterBodyMeshComponent>> MeshComponents = { WaterInfoMeshComponent, DilatedWaterInfoMeshComponent };
-	MeshComponents.Append(WaterBodyStaticMeshComponents);
-
-	for (TObjectPtr<UWaterBodyMeshComponent> MeshComponent : MeshComponents)
-	{
-		if (IsValid(MeshComponent) && IsValid(MeshComponent->GetStaticMesh()))
-		{
-			MeshComponent->GetStaticMesh()->Rename(nullptr, GetPackage(), REN_DontCreateRedirectors | REN_ForceNoResetLoaders);
-		}
-	}
-}
-
 #endif // WITH_EDITOR
 
 void AWaterBody::UnregisterAllComponents(bool bForReregister)
diff --git a/Engine/Plugins/Experimental/Water/Source/Runtime/Private/WaterBodyMeshBuilder.cpp b/Engine/Plugins/Experimental/Water/Source/Runtime/Private/WaterBodyMeshBuilder.cpp
--- a/Engine/Plugins/Experimental/Water/Source/Runtime/Private/WaterBodyMeshBuilder.cpp
+++ b/Engine/Plugins/Experimental/Water/Source/Runtime/Private/WaterBodyMeshBuilder.cpp
@@ -43,7 +43,7 @@ void FWaterBodyMeshBuilder::BuildWaterInfoMeshes(UWaterBodyComponent* WaterBodyC
 	GetDynamicMesh(WaterBodyComponent, WaterBodyMesh, &WaterBodyDilatedMesh);
 
 	UMaterialInterface* WaterInfoMID = WaterBodyComponent->GetWaterInfoMaterialInstance();
-	UObject* Outer = WaterBodyComponent->GetPackage();
+	UObject* Outer = WaterBodyComponent->GetOwner();
 
 	auto BuildWaterInfoMesh = [this, Outer, WaterInfoMID](const FDynamicMesh3& DynamicMesh, UWaterBodyInfoMeshComponent* MeshComponent, FName BaseName) -> UStaticMesh*
 	{
@@ -382,7 +382,7 @@ TArray<TObjectPtr<UWaterBodyStaticMeshComponent>> FWaterBodyMeshBuilder::BuildWa
 
 	UMaterialInterface* WaterStaticMeshMID = WaterBodyComponent->GetWaterStaticMeshMaterialInstance();
 
-	UObject* MeshComponentOuter = WaterBodyComponent->GetPackage();
+	UObject* MeshComponentOuter = WaterBodyComponent->GetOwner();
 
 	int32 StaticMeshComponentIndex = 0;
 	{
diff --git a/Engine/Plugins/Experimental/Water/Source/Runtime/Public/WaterBodyActor.h b/Engine/Plugins/Experimental/Water/Source/Runtime/Public/WaterBodyActor.h
--- a/Engine/Plugins/Experimental/Water/Source/Runtime/Public/WaterBodyActor.h
+++ b/Engine/Plugins/Experimental/Water/Source/Runtime/Public/WaterBodyActor.h
@@ -65,12 +65,9 @@ public:
 	virtual bool SetIsHiddenEdLayer(bool bIsHiddenEdLayer) override;
 	virtual void GetActorDescProperties(FPropertyPairsMap& PropertyPairsMap) const override;
 
-	void PackagingModeChanged(AActor* Actor, bool bIsExternal);
-
-	virtual void OnReattachExternalPackage() override;
-	virtual void OnDetachExternalPackage() override;
-
 	virtual void PostActorCreated() override;
+
+	virtual void PopulatePIEDuplicationSeed(AActor::FDuplicationSeedInterface& DuplicationSeed) override;
 #endif // WITH_EDITOR
 	
 	/** Returns the type of body */
diff --git a/Engine/Source/Runtime/Core/Public/UObject/FortniteMainBranchObjectVersion.h b/Engine/Source/Runtime/Core/Public/UObject/FortniteMainBranchObjectVersion.h
--- a/Engine/Source/Runtime/Core/Public/UObject/FortniteMainBranchObjectVersion.h
+++ b/Engine/Source/Runtime/Core/Public/UObject/FortniteMainBranchObjectVersion.h
@@ -331,6 +331,8 @@ struct CORE_API FFortniteMainBranchObjectVersion
 		// WorldPartition HLOD now contains a source actors object
 		WorldPartitionHLODSourceActorsRefactor,
 
+		WaterBodyStaticMeshRename,
+
 		// -----<new versions can be added above this line>-------------------------------------------------
 		VersionPlusOne,
 		LatestVersion = VersionPlusOne - 1
diff --git a/Engine/Source/Runtime/Engine/Classes/GameFramework/Actor.h b/Engine/Source/Runtime/Engine/Classes/GameFramework/Actor.h
--- a/Engine/Source/Runtime/Engine/Classes/GameFramework/Actor.h
+++ b/Engine/Source/Runtime/Engine/Classes/GameFramework/Actor.h
@@ -2175,6 +2175,22 @@ public:
 	ENGINE_API virtual void PostTransacted(const FTransactionObjectEvent& TransactionEvent) override;
 	ENGINE_API virtual bool IsSelectedInEditor() const override;
 
+	/** DuplicationSeedInterface prevents exposing the entire DupSeed map to every actor and provides a method to do so instead. */
+	struct FDuplicationSeedInterface
+	{
+	public:
+		FDuplicationSeedInterface(TMap<UObject*, UObject*>& InDuplicationSeed);
+
+		/** Add a new entry to the duplication seed. */
+		ENGINE_API void AddEntry(UObject* Source, UObject* Dest);
+
+	private:
+		TMap<UObject*, UObject*>& DuplicationSeed;
+	};
+
+	/** Populate the duplication seed when duplicating the actor for PIE. This allows objects to be remapped to existing objects rather than duplicating. */
+	virtual void PopulatePIEDuplicationSeed(FDuplicationSeedInterface& DupSeed) {}
+
 	/** Defines if preview should be shown when you select an actor, but none of its children */
 	virtual bool IsDefaultPreviewEnabled() const
 	{
diff --git a/Engine/Source/Runtime/Engine/Private/Actor.cpp b/Engine/Source/Runtime/Engine/Private/Actor.cpp
--- a/Engine/Source/Runtime/Engine/Private/Actor.cpp
+++ b/Engine/Source/Runtime/Engine/Private/Actor.cpp
@@ -124,6 +124,18 @@ FGuid FActorInstanceGuidMapper::MapGuid(FName InPackageName, const FGuid& InGuid
 	}
 	return InGuid;
 }
+
+AActor::FDuplicationSeedInterface::FDuplicationSeedInterface(TMap<UObject*, UObject*>& InDuplicationSeed)
+	: DuplicationSeed(InDuplicationSeed)
+{
+}
+
+void AActor::FDuplicationSeedInterface::AddEntry(UObject* Source, UObject* Destination)
+{	
+	DuplicationSeed.Emplace(Source, Destination);
+}
+
+
 #endif
 
 uint32 AActor::BeginPlayCallDepth = 0;
@@ -6414,4 +6426,3 @@ void AActor::ForEachComponentOfActorClassDefault(const TSubclassOf<AActor>& Acto
 
 
 #undef LOCTEXT_NAMESPACE
-
diff --git a/Engine/Source/Runtime/Engine/Private/Level.cpp b/Engine/Source/Runtime/Engine/Private/Level.cpp
--- a/Engine/Source/Runtime/Engine/Private/Level.cpp
+++ b/Engine/Source/Runtime/Engine/Private/Level.cpp
@@ -1308,6 +1308,18 @@ void ULevel::PreDuplicate(FObjectDuplicationParameters& DupParams)
 	Super::PreDuplicate(DupParams);
 
 #if WITH_EDITOR
+	if (DupParams.DuplicateMode == EDuplicateMode::PIE)
+	{
+		AActor::FDuplicationSeedInterface DuplicationSeedInterface(DupParams.DuplicationSeed);
+		for (AActor* Actor : Actors)
+		{
+			if (Actor != nullptr)
+			{
+				Actor->PopulatePIEDuplicationSeed(DuplicationSeedInterface);
+			}
+		}
+	}
+
 	if (DupParams.DuplicateMode != EDuplicateMode::PIE && DupParams.bAssignExternalPackages)
 	{
 		UPackage* SrcPackage = GetPackage();
