Water: use attach parent bounds for water body info mesh components. Fixes a bug when meshes are in a sublevel and bounds aren't recomputed at runtime.

diff --git a/Engine/Plugins/Experimental/Water/Source/Runtime/Private/WaterBodyInfoMeshComponent.cpp b/Engine/Plugins/Experimental/Water/Source/Runtime/Private/WaterBodyInfoMeshComponent.cpp
--- a/Engine/Plugins/Experimental/Water/Source/Runtime/Private/WaterBodyInfoMeshComponent.cpp
+++ b/Engine/Plugins/Experimental/Water/Source/Runtime/Private/WaterBodyInfoMeshComponent.cpp
@@ -27,6 +27,9 @@ UWaterBodyInfoMeshComponent::UWaterBodyInfoMeshComponent(const FObjectInitialize
 {
 	bAffectDistanceFieldLighting = false;
 	bSelectable = false;
+	
+	// Skip computing the bounds for this component since it should always be attached to the water body component and those bounds are the proper bounds for the info meshes.
+	bUseAttachParentBound = true;
 }
 
 FPrimitiveSceneProxy* UWaterBodyInfoMeshComponent::CreateSceneProxy()
