Water: expose SetWaterZoneOverride to blueprints

diff --git a/Engine/Plugins/Experimental/Water/Source/Runtime/Private/WaterBodyComponent.cpp b/Engine/Plugins/Experimental/Water/Source/Runtime/Private/WaterBodyComponent.cpp
--- a/Engine/Plugins/Experimental/Water/Source/Runtime/Private/WaterBodyComponent.cpp
+++ b/Engine/Plugins/Experimental/Water/Source/Runtime/Private/WaterBodyComponent.cpp
@@ -417,6 +417,12 @@ AWaterZone* UWaterBodyComponent::GetWaterZone() const
 	return OwningWaterZone.Get();
 }
 
+void UWaterBodyComponent::SetWaterZoneOverride(const TSoftObjectPtr<AWaterZone>& InWaterZoneOverride)
+{
+	WaterZoneOverride = InWaterZoneOverride;
+	UpdateWaterZones();
+}
+
 FPostProcessVolumeProperties UWaterBodyComponent::GetPostProcessProperties() const
 {
 	FPostProcessVolumeProperties Ret;
@@ -2061,3 +2067,4 @@ void UWaterBodyComponent::SetHLODMaterial(UMaterialInterface* InMaterial)
 #endif // WITH_EDITOR
 
 #undef LOCTEXT_NAMESPACE
+
diff --git a/Engine/Plugins/Experimental/Water/Source/Runtime/Public/WaterBodyComponent.h b/Engine/Plugins/Experimental/Water/Source/Runtime/Public/WaterBodyComponent.h
--- a/Engine/Plugins/Experimental/Water/Source/Runtime/Public/WaterBodyComponent.h
+++ b/Engine/Plugins/Experimental/Water/Source/Runtime/Public/WaterBodyComponent.h
@@ -384,6 +384,10 @@ public:
 	/** Gets the water zone to which this component belongs */
 	AWaterZone* GetWaterZone() const;
 
+	/** Override the default behavior of water bodies finding their water zone based on bounds and set a specific water zone to which this water body should register. */
+	UFUNCTION(BlueprintCallable, Category=Water)
+	void SetWaterZoneOverride(const TSoftObjectPtr<AWaterZone>& InWaterZoneOverride);
+
 	/** Registers or this water body with corresponding overlapping water zones and unregisters it from any old zones if they are no longer overlapping. */
 	void UpdateWaterZones();
 
