Water: fixed issue where water bodies no longer update zones but some older saved maps might not have serialized the pointer yet.

diff --git a/Engine/Plugins/Experimental/Water/Source/Runtime/Private/WaterBodyComponent.cpp b/Engine/Plugins/Experimental/Water/Source/Runtime/Private/WaterBodyComponent.cpp
--- a/Engine/Plugins/Experimental/Water/Source/Runtime/Private/WaterBodyComponent.cpp
+++ b/Engine/Plugins/Experimental/Water/Source/Runtime/Private/WaterBodyComponent.cpp
@@ -1439,6 +1439,20 @@ void UWaterBodyComponent::UpdateAll(const FOnWaterBodyChangedParams& InParams)
 	}
 }
 
+void UWaterBodyComponent::OnPostRegisterAllComponents()
+{
+#if WITH_EDITOR
+	// Before this version, we always updated the water zone pointer on load. This no longer occurs so any water body which doesn't have a serialized pointer needs to set it now
+	if (GetLinkerCustomVersion(FFortniteMainBranchObjectVersion::GUID) < FFortniteMainBranchObjectVersion::WaterOwningZonePointerFixup)
+	{
+		if (OwningWaterZone.IsNull())
+		{
+			UpdateWaterZones();
+		}
+	}
+#endif // WITH_EDITOR
+}
+
 void UWaterBodyComponent::UpdateSplineComponent()
 {
 	if (UWaterSplineComponent* WaterSpline = GetWaterSpline())
diff --git a/Engine/Plugins/Experimental/Water/Source/Runtime/Public/WaterBodyComponent.h b/Engine/Plugins/Experimental/Water/Source/Runtime/Public/WaterBodyComponent.h
--- a/Engine/Plugins/Experimental/Water/Source/Runtime/Public/WaterBodyComponent.h
+++ b/Engine/Plugins/Experimental/Water/Source/Runtime/Public/WaterBodyComponent.h
@@ -449,7 +449,7 @@ protected:
 	virtual void OnUpdateBody(bool bWithExclusionVolumes) {}
 
 	/** Called when the WaterBodyActor has had all its components registered. */
-	virtual void OnPostRegisterAllComponents() {}
+	virtual void OnPostRegisterAllComponents();
 
 	/** Returns navigation area class */
 	TSubclassOf<UNavAreaBase> GetNavAreaClass() const { return WaterNavAreaClass; }
diff --git a/Engine/Source/Runtime/Core/Public/UObject/FortniteMainBranchObjectVersion.h b/Engine/Source/Runtime/Core/Public/UObject/FortniteMainBranchObjectVersion.h
--- a/Engine/Source/Runtime/Core/Public/UObject/FortniteMainBranchObjectVersion.h
+++ b/Engine/Source/Runtime/Core/Public/UObject/FortniteMainBranchObjectVersion.h
@@ -336,6 +336,9 @@ struct CORE_API FFortniteMainBranchObjectVersion
 		// Geometry Collection now by-default converts vertex colors to sRGB when creating render data
 		GeometryCollectionConvertVertexColorToSRGB,
 
+		// Water bodies before this version need to update their water zone on load since they won't have been serialized yet.
+		WaterOwningZonePointerFixup,
+
 		// -----<new versions can be added above this line>-------------------------------------------------
 		VersionPlusOne,
 		LatestVersion = VersionPlusOne - 1
