Fixed non-unity build break.

diff --git a/Engine/Plugins/Experimental/Water/Source/Runtime/Private/WaterMeshSceneProxy.cpp b/Engine/Plugins/Experimental/Water/Source/Runtime/Private/WaterMeshSceneProxy.cpp
--- a/Engine/Plugins/Experimental/Water/Source/Runtime/Private/WaterMeshSceneProxy.cpp
+++ b/Engine/Plugins/Experimental/Water/Source/Runtime/Private/WaterMeshSceneProxy.cpp
@@ -12,6 +12,7 @@
 #include "Math/ColorList.h"
 #include "RayTracingInstance.h"
 #include "SceneInterface.h"
+#include "RenderGraphBuilder.h"
 
 DECLARE_STATS_GROUP(TEXT("Water Mesh"), STATGROUP_WaterMesh, STATCAT_Advanced);
 
diff --git a/Engine/Source/Runtime/Landscape/Private/LandscapeCulling.h b/Engine/Source/Runtime/Landscape/Private/LandscapeCulling.h
--- a/Engine/Source/Runtime/Landscape/Private/LandscapeCulling.h
+++ b/Engine/Source/Runtime/Landscape/Private/LandscapeCulling.h
@@ -14,6 +14,7 @@ class FSceneViewFamily;
 class FSceneView;
 class FRDGBuilder;
 struct FViewMatrices;
+class FRHICommandListBase;
 
 /**
 * Landscape GPU culling implementation
