Water: add new call to UpdateWaterSpriteComponent to ensure the sprite represents the state of the water body after PostRegisiter where some data might have been fixed up.

diff --git a/Engine/Plugins/Experimental/Water/Source/Runtime/Private/WaterBodyComponent.cpp b/Engine/Plugins/Experimental/Water/Source/Runtime/Private/WaterBodyComponent.cpp
--- a/Engine/Plugins/Experimental/Water/Source/Runtime/Private/WaterBodyComponent.cpp
+++ b/Engine/Plugins/Experimental/Water/Source/Runtime/Private/WaterBodyComponent.cpp
@@ -1450,6 +1450,9 @@ void UWaterBodyComponent::OnPostRegisterAllComponents()
 			UpdateWaterZones();
 		}
 	}
+
+	// Ensure that the sprite component is updated once the water body is fully setup after PostRegister.
+	UpdateWaterSpriteComponent();
 #endif // WITH_EDITOR
 }
 
