Water: made the water static meshes duplicate/textexport transient.

diff --git a/Engine/Plugins/Experimental/Water/Source/Runtime/Private/WaterBodyActor.cpp b/Engine/Plugins/Experimental/Water/Source/Runtime/Private/WaterBodyActor.cpp
--- a/Engine/Plugins/Experimental/Water/Source/Runtime/Private/WaterBodyActor.cpp
+++ b/Engine/Plugins/Experimental/Water/Source/Runtime/Private/WaterBodyActor.cpp
@@ -532,6 +532,22 @@ void AWaterBody::DeprecateData()
 			}
 		}
 	}
+
+	// Static meshes should be duplicate transient so we don't run into issues with async builders running tasks. Instead we can just rely on them being recreated after the water body is duplicated.
+	if (GetLinkerCustomVersion(FFortniteMainBranchObjectVersion::GUID) < FFortniteMainBranchObjectVersion::WaterBodyStaticMeshRename)
+	{
+		TArray<TObjectPtr<UWaterBodyMeshComponent>> MeshComponents = { WaterInfoMeshComponent, DilatedWaterInfoMeshComponent };
+		MeshComponents.Append(WaterBodyStaticMeshComponents);
+
+		for (TObjectPtr<UWaterBodyMeshComponent> MeshComponent : MeshComponents)
+		{
+			if (IsValid(MeshComponent) && IsValid(MeshComponent->GetStaticMesh()))
+			{
+				UStaticMesh* StaticMesh = MeshComponent->GetStaticMesh();
+				StaticMesh->SetFlags(StaticMesh->GetFlags() | RF_DuplicateTransient | RF_TextExportTransient);
+			}
+		}
+	}
 #endif // WITH_EDITORONLY_DATA
 }
 
diff --git a/Engine/Plugins/Experimental/Water/Source/Runtime/Private/WaterBodyMeshBuilder.cpp b/Engine/Plugins/Experimental/Water/Source/Runtime/Private/WaterBodyMeshBuilder.cpp
--- a/Engine/Plugins/Experimental/Water/Source/Runtime/Private/WaterBodyMeshBuilder.cpp
+++ b/Engine/Plugins/Experimental/Water/Source/Runtime/Private/WaterBodyMeshBuilder.cpp
@@ -562,7 +562,7 @@ void FWaterBodyMeshBuilder::UpdateStaticMesh(UStaticMesh* WaterMesh, const FMesh
 
 UStaticMesh* FWaterBodyMeshBuilder::CreateUStaticMesh(UObject* Outer, FName MeshName) const
 {
-	UStaticMesh* StaticMesh = NewObject<UStaticMesh>(Outer, MeshName, RF_Transactional);
+	UStaticMesh* StaticMesh = NewObject<UStaticMesh>(Outer, MeshName, RF_Transactional | RF_TextExportTransient | RF_DuplicateTransient);
 
 	// Disable navigation
 	StaticMesh->MarkAsNotHavingNavigationData();
diff --git a/Engine/Source/Runtime/Core/Public/UObject/FortniteMainBranchObjectVersion.h b/Engine/Source/Runtime/Core/Public/UObject/FortniteMainBranchObjectVersion.h
--- a/Engine/Source/Runtime/Core/Public/UObject/FortniteMainBranchObjectVersion.h
+++ b/Engine/Source/Runtime/Core/Public/UObject/FortniteMainBranchObjectVersion.h
@@ -339,6 +339,9 @@ struct CORE_API FFortniteMainBranchObjectVersion
 		// Water bodies before this version need to update their water zone on load since they won't have been serialized yet.
 		WaterOwningZonePointerFixup,
 
+		// Set flags on water static meshes to duplicate transient to avoid underlying static mesh duplication issue
+		WaterBodyStaticMeshDuplicateTransient,
+
 		// -----<new versions can be added above this line>-------------------------------------------------
 		VersionPlusOne,
 		LatestVersion = VersionPlusOne - 1
