Water: static meshes should be NonPIEDuplicateTransient since they still need to be marked to duplicate for pie otherwise the duplication seed won't do anything.

diff --git a/Engine/Plugins/Experimental/Water/Source/Runtime/Private/WaterBodyActor.cpp b/Engine/Plugins/Experimental/Water/Source/Runtime/Private/WaterBodyActor.cpp
--- a/Engine/Plugins/Experimental/Water/Source/Runtime/Private/WaterBodyActor.cpp
+++ b/Engine/Plugins/Experimental/Water/Source/Runtime/Private/WaterBodyActor.cpp
@@ -544,7 +544,7 @@ void AWaterBody::DeprecateData()
 			if (IsValid(MeshComponent) && IsValid(MeshComponent->GetStaticMesh()))
 			{
 				UStaticMesh* StaticMesh = MeshComponent->GetStaticMesh();
-				StaticMesh->SetFlags(StaticMesh->GetFlags() | RF_DuplicateTransient | RF_TextExportTransient);
+				StaticMesh->SetFlags(StaticMesh->GetFlags() | RF_NonPIEDuplicateTransient | RF_TextExportTransient);
 			}
 		}
 	}
diff --git a/Engine/Plugins/Experimental/Water/Source/Runtime/Private/WaterBodyMeshBuilder.cpp b/Engine/Plugins/Experimental/Water/Source/Runtime/Private/WaterBodyMeshBuilder.cpp
--- a/Engine/Plugins/Experimental/Water/Source/Runtime/Private/WaterBodyMeshBuilder.cpp
+++ b/Engine/Plugins/Experimental/Water/Source/Runtime/Private/WaterBodyMeshBuilder.cpp
@@ -562,7 +562,7 @@ void FWaterBodyMeshBuilder::UpdateStaticMesh(UStaticMesh* WaterMesh, const FMesh
 
 UStaticMesh* FWaterBodyMeshBuilder::CreateUStaticMesh(UObject* Outer, FName MeshName) const
 {
-	UStaticMesh* StaticMesh = NewObject<UStaticMesh>(Outer, MeshName, RF_Transactional | RF_TextExportTransient | RF_DuplicateTransient);
+	UStaticMesh* StaticMesh = NewObject<UStaticMesh>(Outer, MeshName, RF_Transactional | RF_TextExportTransient | RF_NonPIEDuplicateTransient);
 
 	// Disable navigation
 	StaticMesh->MarkAsNotHavingNavigationData();
