Water: Exclude water info meshes from HLOD. Water info meshes were being included in hlod and appeared in the main scene.
We have a custom HLOD builder for water and shouldn't use any of the mesh components for this.

diff --git a/Engine/Plugins/Experimental/Water/Source/Runtime/Public/WaterBodyMeshComponent.h b/Engine/Plugins/Experimental/Water/Source/Runtime/Public/WaterBodyMeshComponent.h
--- a/Engine/Plugins/Experimental/Water/Source/Runtime/Public/WaterBodyMeshComponent.h
+++ b/Engine/Plugins/Experimental/Water/Source/Runtime/Public/WaterBodyMeshComponent.h
@@ -23,7 +23,8 @@ public:
 
 protected:
 	virtual bool CanCreateSceneProxy() const;
-};
 
+	// Exclude water mesh components from HLOD since we have a custom HLOD builder for water bodies.
+	virtual bool IsHLODRelevant() const override { return false; }
+};
 
-// #todo_water: UWaterBodyInstancedStaticMeshComponent for the quad meshes to be instanced
