Water: disabled bShouldNotifyOnWorldAddRemove for water info scene proxies since they should never have the ForceHidden flag disabled outside of water info rendering.

diff --git a/Engine/Plugins/Experimental/Water/Source/Runtime/Private/WaterBodyInfoMeshComponent.cpp b/Engine/Plugins/Experimental/Water/Source/Runtime/Private/WaterBodyInfoMeshComponent.cpp
--- a/Engine/Plugins/Experimental/Water/Source/Runtime/Private/WaterBodyInfoMeshComponent.cpp
+++ b/Engine/Plugins/Experimental/Water/Source/Runtime/Private/WaterBodyInfoMeshComponent.cpp
@@ -44,6 +44,10 @@ FPrimitiveSceneProxy* UWaterBodyInfoMeshComponent::CreateSceneProxy()
 FWaterBodyInfoMeshSceneProxy::FWaterBodyInfoMeshSceneProxy(UWaterBodyInfoMeshComponent* Component)
 	: FStaticMeshSceneProxy(Component, true)
 {
+	// Disable Notify on WorldAddRemove. This prevents the component from being unhidden in FPrimitiveSceneProxy::OnLevelAddedToWorld_RenderThread if it was part of a streamed level.
+	// WaterInfo proxies should only be unhidden during WaterInfo passes.
+	bShouldNotifyOnWorldAddRemove = false;
+
 	SetEnabled(false);
 }
 
