Merging CL#26114624

#UE Fix math error skipping the last element in the histogram in GenerateSubtractSweepConvexDecomposition

diff --git a/Engine/Plugins/Experimental/Water/Source/Runtime/Private/WaterBooleanUtils.cpp b/Engine/Plugins/Experimental/Water/Source/Runtime/Private/WaterBooleanUtils.cpp
--- a/Engine/Plugins/Experimental/Water/Source/Runtime/Private/WaterBooleanUtils.cpp
+++ b/Engine/Plugins/Experimental/Water/Source/Runtime/Private/WaterBooleanUtils.cpp
@@ -1017,7 +1017,7 @@ void FWaterBooleanUtils::GenerateSubtractSweepConvexDecomposition(const FDynamic
 
 	// find largest "middle" bin - use this as the cut plane height
 	int32 MaxBinIndex = 1;
-	for (int32 k = 2; k < IntervalSteps - 1; ++k)
+	for (int32 k = 2; k < IntervalSteps; ++k)
 	{
 		if (ZHistogram[k] > ZHistogram[MaxBinIndex])
 		{
