Backing out change to skip last element of histogram in GenerateSubtractSweepConvexDecomposition since it is inentionally looking for mid bins. Lowering verbosity of the message when there are no mid bins to choose from for now since it is happening in real data

diff --git a/Engine/Plugins/Experimental/Water/Source/Runtime/Private/WaterBooleanUtils.cpp b/Engine/Plugins/Experimental/Water/Source/Runtime/Private/WaterBooleanUtils.cpp
--- a/Engine/Plugins/Experimental/Water/Source/Runtime/Private/WaterBooleanUtils.cpp
+++ b/Engine/Plugins/Experimental/Water/Source/Runtime/Private/WaterBooleanUtils.cpp
@@ -1017,7 +1017,7 @@ void FWaterBooleanUtils::GenerateSubtractSweepConvexDecomposition(const FDynamic
 
 	// find largest "middle" bin - use this as the cut plane height
 	int32 MaxBinIndex = 1;
-	for (int32 k = 2; k < IntervalSteps; ++k)
+	for (int32 k = 2; k < IntervalSteps - 1; ++k)
 	{
 		if (ZHistogram[k] > ZHistogram[MaxBinIndex])
 		{
@@ -1029,7 +1029,7 @@ void FWaterBooleanUtils::GenerateSubtractSweepConvexDecomposition(const FDynamic
 	if (MaxBinIndex == 1 && ZHistogram[1] == 0)
 	{
 		BaseZ = Bounds.Center().Z;		// if we did not find a bin to cut at, this might be a through-hole, so cut at the middle ? failure case really..
-		UE_LOG(LogWater, Warning, TEXT("UE::Water::GenerateSubtractSweepConvexDecomposition : Invalid Configuration - either no cut, or through-hole"));
+		UE_LOG(LogWater, Display, TEXT("UE::Water::GenerateSubtractSweepConvexDecomposition : Invalid Configuration - either no cut, or through-hole"));
 	}
 
 	// set up the "bottom" box
