Fixed bad deprecation code in water static mesh flag changes.

diff --git a/Engine/Plugins/Experimental/Water/Source/Runtime/Private/WaterBodyActor.cpp b/Engine/Plugins/Experimental/Water/Source/Runtime/Private/WaterBodyActor.cpp
--- a/Engine/Plugins/Experimental/Water/Source/Runtime/Private/WaterBodyActor.cpp
+++ b/Engine/Plugins/Experimental/Water/Source/Runtime/Private/WaterBodyActor.cpp
@@ -516,38 +516,6 @@ void AWaterBody::DeprecateData()
 		// Some deprecated data on the actor were also later deprecated on the component (e.g. bFillCollisionUnderWaterBodiesForNavmesh_DEPRECATED), make sure it runs its own deprecation operations on it : 
 		WaterBodyComponent->DeprecateData();
 	}
-
-	// With this version bump, we introduced a fix that allows water static meshes to be outered to the actor and still avoid the cost of duplicating them for PIE. It is now safe to rename them back to have a proper outer chain with the actor as the outer.
-	if (GetLinkerCustomVersion(FFortniteMainBranchObjectVersion::GUID) < FFortniteMainBranchObjectVersion::WaterBodyStaticMeshRename)
-	{
-		TArray<TObjectPtr<UWaterBodyMeshComponent>> MeshComponents = { WaterInfoMeshComponent, DilatedWaterInfoMeshComponent };
-		MeshComponents.Append(WaterBodyStaticMeshComponents);
-
-		for (TObjectPtr<UWaterBodyMeshComponent> MeshComponent : MeshComponents)
-		{
-			if (IsValid(MeshComponent) && IsValid(MeshComponent->GetStaticMesh()))
-			{
-				UStaticMesh* StaticMesh = MeshComponent->GetStaticMesh();
-				StaticMesh->Rename(nullptr, this, REN_DontCreateRedirectors | REN_ForceNoResetLoaders);
-			}
-		}
-	}
-
-	// Static meshes should be duplicate transient so we don't run into issues with async builders running tasks. Instead we can just rely on them being recreated after the water body is duplicated.
-	if (GetLinkerCustomVersion(FFortniteMainBranchObjectVersion::GUID) < FFortniteMainBranchObjectVersion::WaterBodyStaticMeshRename)
-	{
-		TArray<TObjectPtr<UWaterBodyMeshComponent>> MeshComponents = { WaterInfoMeshComponent, DilatedWaterInfoMeshComponent };
-		MeshComponents.Append(WaterBodyStaticMeshComponents);
-
-		for (TObjectPtr<UWaterBodyMeshComponent> MeshComponent : MeshComponents)
-		{
-			if (IsValid(MeshComponent) && IsValid(MeshComponent->GetStaticMesh()))
-			{
-				UStaticMesh* StaticMesh = MeshComponent->GetStaticMesh();
-				StaticMesh->SetFlags(StaticMesh->GetFlags() | RF_NonPIEDuplicateTransient | RF_TextExportTransient);
-			}
-		}
-	}
 #endif // WITH_EDITORONLY_DATA
 }
 
diff --git a/Engine/Plugins/Experimental/Water/Source/Runtime/Private/WaterBodyComponent.cpp b/Engine/Plugins/Experimental/Water/Source/Runtime/Private/WaterBodyComponent.cpp
--- a/Engine/Plugins/Experimental/Water/Source/Runtime/Private/WaterBodyComponent.cpp
+++ b/Engine/Plugins/Experimental/Water/Source/Runtime/Private/WaterBodyComponent.cpp
@@ -1451,6 +1451,11 @@ void UWaterBodyComponent::OnPostRegisterAllComponents()
 		}
 	}
 
+	if (GetLinkerCustomVersion(FFortniteMainBranchObjectVersion::GUID) < FFortniteMainBranchObjectVersion::WaterBodyStaticMeshFixup)
+	{
+		UpdateWaterBodyRenderData();
+	}
+
 	// Ensure that the sprite component is updated once the water body is fully setup after PostRegister.
 	UpdateWaterSpriteComponent();
 #endif // WITH_EDITOR
diff --git a/Engine/Source/Runtime/Core/Public/UObject/FortniteMainBranchObjectVersion.h b/Engine/Source/Runtime/Core/Public/UObject/FortniteMainBranchObjectVersion.h
--- a/Engine/Source/Runtime/Core/Public/UObject/FortniteMainBranchObjectVersion.h
+++ b/Engine/Source/Runtime/Core/Public/UObject/FortniteMainBranchObjectVersion.h
@@ -345,6 +345,9 @@ struct FFortniteMainBranchObjectVersion
 		// Update paths to use the SkeletalClass
 		MVVMConvertPropertyPathToSkeletalClass,
 
+		// Fixup all flags/outering on static meshes on water bodies by rebuilding them completely
+		WaterBodyStaticMeshFixup,
+
 		// -----<new versions can be added above this line>-------------------------------------------------
 		VersionPlusOne,
 		LatestVersion = VersionPlusOne - 1
