Water: Fixed missing translation in the ocean bounds so they are always centered on the saved zone location. This ensures water mesh tiles are properly generated for oceans which aren't at 0,0.

This bug was introduced by the fix to the quadtree considering the ocean bounds in 25678023.

diff --git a/Engine/Plugins/Experimental/Water/Source/Runtime/Private/WaterBodyOceanComponent.cpp b/Engine/Plugins/Experimental/Water/Source/Runtime/Private/WaterBodyOceanComponent.cpp
--- a/Engine/Plugins/Experimental/Water/Source/Runtime/Private/WaterBodyOceanComponent.cpp
+++ b/Engine/Plugins/Experimental/Water/Source/Runtime/Private/WaterBodyOceanComponent.cpp
@@ -115,7 +115,13 @@ void UWaterBodyOceanComponent::OnPostEditChangeProperty(FOnWaterBodyChangedParam
 		// Affects the physics shape
 		InOutOnWaterBodyChangedParams.bShapeOrPositionChanged = true;
 	}
+	else if (PropertyChangedEvent.GetPropertyName() == GET_MEMBER_NAME_CHECKED(UWaterBodyOceanComponent, OceanExtents)
+		|| PropertyChangedEvent.GetPropertyName() == GET_MEMBER_NAME_CHECKED(UWaterBodyOceanComponent, SavedZoneLocation))
+	{
+		InOutOnWaterBodyChangedParams.bShapeOrPositionChanged = true;
+	}
 }
+
 const TCHAR* UWaterBodyOceanComponent::GetWaterSpriteTextureName() const
 {
 	return TEXT("/Water/Icons/WaterBodyOceanSprite");
@@ -490,7 +496,9 @@ FBoxSphereBounds UWaterBodyOceanComponent::CalcBounds(const FTransform& LocalToW
 	{
 		const FVector Min(FVector(-OceanExtents / 2.0, -1.0 * GetChannelDepth()));
 		const FVector Max(FVector(OceanExtents / 2.0, 0.0));
-		return FBoxSphereBounds(FBox(Min, Max)).TransformBy(LocalToWorld);
+
+		// translate the bounds so they are centered on the SavedZoneLocation.
+		return FBoxSphereBounds(FBox(Min, Max)).TransformBy(LocalToWorld).TransformBy(FTransform(FVector(SavedZoneLocation, 0.)));
 	}
 	return FBoxSphereBounds(EForceInit::ForceInit).TransformBy(LocalToWorld);
 }
diff --git a/Engine/Plugins/Experimental/Water/Source/Runtime/Public/WaterBodyOceanComponent.h b/Engine/Plugins/Experimental/Water/Source/Runtime/Public/WaterBodyOceanComponent.h
--- a/Engine/Plugins/Experimental/Water/Source/Runtime/Public/WaterBodyOceanComponent.h
+++ b/Engine/Plugins/Experimental/Water/Source/Runtime/Public/WaterBodyOceanComponent.h
@@ -69,19 +69,18 @@ protected:
 	UPROPERTY(Category = Collision, EditAnywhere, BlueprintReadOnly)
 	FVector CollisionExtents;
 
-	// #todo_water: should we make this editor-only? It's not needed at runtime since the mesh will never regenerate at runtime
 	UPROPERTY(Category = Water, EditAnywhere, BlueprintReadOnly)
 	FVector2D OceanExtents;
 
+	/** Save the water zone location so that the ocean mesh can be regenerated relative to it and match it perfectly without being loaded. */
+	UPROPERTY(Category = Water, EditAnywhere, BlueprintReadOnly, meta=(DisplayName="Ocean Center"))
+	FVector2D SavedZoneLocation;
+
 	UPROPERTY(Transient)
 	float HeightOffset = 0.0f;
 
 #if WITH_EDITORONLY_DATA
 
-	/** Serialize the offset relative to the water zone to check if the serialized mesh and the current water zone have a mismatch. */
-	UPROPERTY()
-	FVector2D SavedZoneLocation;
-
 	UPROPERTY()
 	FVector2D VisualExtents_DEPRECATED;
 #endif // WITH_EDITORONLY_DATA
