Water: prevent trying to reuse components marked as PendingKill.

diff --git a/Engine/Plugins/Experimental/Water/Source/Runtime/Private/WaterBodyComponent.cpp b/Engine/Plugins/Experimental/Water/Source/Runtime/Private/WaterBodyComponent.cpp
--- a/Engine/Plugins/Experimental/Water/Source/Runtime/Private/WaterBodyComponent.cpp
+++ b/Engine/Plugins/Experimental/Water/Source/Runtime/Private/WaterBodyComponent.cpp
@@ -43,6 +43,7 @@
 #include "MeshMergeModule.h"
 #include "WaterBodyMeshBuilder.h"
 #include "MeshDescription.h"
+#include "Algo/RemoveIf.h"
 #endif
 
 #define LOCTEXT_NAMESPACE "Water"
@@ -1961,12 +1962,18 @@ void UWaterBodyComponent::UpdateWaterBodyStaticMeshComponents()
 		return;
 	}
 
-	const TArray<TObjectPtr<UWaterBodyStaticMeshComponent>>& WaterBodyStaticMeshComponents = WaterBodyActor->GetWaterBodyStaticMeshComponents();
+	TArray<TObjectPtr<UWaterBodyStaticMeshComponent>>& WaterBodyStaticMeshComponents = WaterBodyActor->GetWaterBodyStaticMeshComponents();
 
 	const bool bShouldUseStaticMesh = StaticMeshSettings.bEnableWaterBodyStaticMesh && (GetWaterMeshOverride() == nullptr) && (GetWaterBodyType() != EWaterBodyType::Transition);
 
 	if (bShouldUseStaticMesh)
 	{
+		// Can't reuse components marked PendingKill or null. This can occur when undoing after disabling static meshes.
+		WaterBodyStaticMeshComponents.SetNum(Algo::RemoveIf(WaterBodyStaticMeshComponents, [](const TObjectPtr<UWaterBodyStaticMeshComponent>& StaticMeshComponent)
+		{
+			return !IsValid(StaticMeshComponent);
+		}));
+
 		FWaterBodyMeshBuilder LODMeshBuilder;
 		TArray<TObjectPtr<UWaterBodyStaticMeshComponent>> NewStaticMeshComponents = LODMeshBuilder.BuildWaterBodyStaticMesh(this, StaticMeshSettings, WaterBodyStaticMeshComponents);
 		
diff --git a/Engine/Plugins/Experimental/Water/Source/Runtime/Private/WaterBodyMeshBuilder.cpp b/Engine/Plugins/Experimental/Water/Source/Runtime/Private/WaterBodyMeshBuilder.cpp
--- a/Engine/Plugins/Experimental/Water/Source/Runtime/Private/WaterBodyMeshBuilder.cpp
+++ b/Engine/Plugins/Experimental/Water/Source/Runtime/Private/WaterBodyMeshBuilder.cpp
@@ -163,7 +163,7 @@ TArray<TObjectPtr<UWaterBodyStaticMeshComponent>> FWaterBodyMeshBuilder::BuildWa
 		{
 			TRACE_CPUPROFILER_EVENT_SCOPE(BuildWaterBodyStaticMesh::BuildMeshSection);
 			const int32 SectionIndexX = SectionIndex % NumberOfMeshSectionsX;
-			const int32 SectionIndexY = SectionIndex / NumberOfMeshSectionsY;
+			const int32 SectionIndexY = SectionIndex / NumberOfMeshSectionsX;
 
 			const FVector2d SectionPos = MeshBounds2d.Min + FVector2d(SectionIndexX, SectionIndexY) * WaterMeshSectionSize;
 
diff --git a/Engine/Plugins/Experimental/Water/Source/Runtime/Public/WaterBodyActor.h b/Engine/Plugins/Experimental/Water/Source/Runtime/Public/WaterBodyActor.h
--- a/Engine/Plugins/Experimental/Water/Source/Runtime/Public/WaterBodyActor.h
+++ b/Engine/Plugins/Experimental/Water/Source/Runtime/Public/WaterBodyActor.h
@@ -135,7 +135,7 @@ protected:
 
 	void SetWaterWavesInternal(UWaterWavesBase* InWaterWaves);
 
-	const TArray<TObjectPtr<UWaterBodyStaticMeshComponent>>& GetWaterBodyStaticMeshComponents() const { return WaterBodyStaticMeshComponents; }
+	TArray<TObjectPtr<UWaterBodyStaticMeshComponent>>& GetWaterBodyStaticMeshComponents() { return WaterBodyStaticMeshComponents; }
 
 	/** Sets up a new list of water body static mesh components. */
 	void SetWaterBodyStaticMeshComponents(TArrayView<TObjectPtr<UWaterBodyStaticMeshComponent>> NewComponentList, TConstArrayView<TObjectPtr<UWaterBodyStaticMeshComponent>> ComponentsToUnregister = {});
