Water: Cleanup invalid static mesh components in PostEditUndo which may have been marked dirty by the operation.

diff --git a/Engine/Plugins/Experimental/Water/Source/Runtime/Private/WaterBodyActor.cpp b/Engine/Plugins/Experimental/Water/Source/Runtime/Private/WaterBodyActor.cpp
--- a/Engine/Plugins/Experimental/Water/Source/Runtime/Private/WaterBodyActor.cpp
+++ b/Engine/Plugins/Experimental/Water/Source/Runtime/Private/WaterBodyActor.cpp
@@ -18,6 +18,7 @@
 #include "WaterBodyInfoMeshComponent.h"
 #include "WaterBodyStaticMeshComponent.h"
 #include "WaterVersion.h"
+#include "Algo/RemoveIf.h"
 
 #include UE_INLINE_GENERATED_CPP_BY_NAME(WaterBodyActor)
 
@@ -218,6 +219,14 @@ void AWaterBody::SetWaterWavesInternal(UWaterWavesBase* InWaterWaves)
 	}
 }
 
+void AWaterBody::CleanupInvalidStaticMeshComponents()
+{
+	WaterBodyStaticMeshComponents.SetNum(Algo::RemoveIf(WaterBodyStaticMeshComponents, [](const TObjectPtr<UWaterBodyStaticMeshComponent>& StaticMeshComponent)
+	{
+		return !IsValid(StaticMeshComponent);
+	}));
+}
+
 void AWaterBody::SetWaterBodyStaticMeshComponents(TArrayView<TObjectPtr<UWaterBodyStaticMeshComponent>> NewComponentList, TConstArrayView<TObjectPtr<UWaterBodyStaticMeshComponent>> ComponentsToUnregister)
 {
 	for (const TObjectPtr<UWaterBodyStaticMeshComponent>& StaticMeshComponent : ComponentsToUnregister)
diff --git a/Engine/Plugins/Experimental/Water/Source/Runtime/Private/WaterBodyComponent.cpp b/Engine/Plugins/Experimental/Water/Source/Runtime/Private/WaterBodyComponent.cpp
--- a/Engine/Plugins/Experimental/Water/Source/Runtime/Private/WaterBodyComponent.cpp
+++ b/Engine/Plugins/Experimental/Water/Source/Runtime/Private/WaterBodyComponent.cpp
@@ -1039,6 +1039,11 @@ void UWaterBodyComponent::PreEditUndo()
 
 void UWaterBodyComponent::PostEditUndo()
 {
+	if (AWaterBody* WaterBodyActor = GetWaterBodyActor())
+	{
+		WaterBodyActor->CleanupInvalidStaticMeshComponents();
+	}
+
 	Super::PostEditUndo();
 
 	// Since this component may become unregistered/deleted if we are undoing the creation of a water body.
@@ -1962,18 +1967,12 @@ void UWaterBodyComponent::UpdateWaterBodyStaticMeshComponents()
 		return;
 	}
 
-	TArray<TObjectPtr<UWaterBodyStaticMeshComponent>>& WaterBodyStaticMeshComponents = WaterBodyActor->GetWaterBodyStaticMeshComponents();
+	const TArray<TObjectPtr<UWaterBodyStaticMeshComponent>>& WaterBodyStaticMeshComponents = WaterBodyActor->GetWaterBodyStaticMeshComponents();
 
 	const bool bShouldUseStaticMesh = StaticMeshSettings.bEnableWaterBodyStaticMesh && (GetWaterMeshOverride() == nullptr) && (GetWaterBodyType() != EWaterBodyType::Transition);
 
 	if (bShouldUseStaticMesh)
 	{
-		// Can't reuse components marked PendingKill or null. This can occur when undoing after disabling static meshes.
-		WaterBodyStaticMeshComponents.SetNum(Algo::RemoveIf(WaterBodyStaticMeshComponents, [](const TObjectPtr<UWaterBodyStaticMeshComponent>& StaticMeshComponent)
-		{
-			return !IsValid(StaticMeshComponent);
-		}));
-
 		FWaterBodyMeshBuilder LODMeshBuilder;
 		TArray<TObjectPtr<UWaterBodyStaticMeshComponent>> NewStaticMeshComponents = LODMeshBuilder.BuildWaterBodyStaticMesh(this, StaticMeshSettings, WaterBodyStaticMeshComponents);
 		
diff --git a/Engine/Plugins/Experimental/Water/Source/Runtime/Public/WaterBodyActor.h b/Engine/Plugins/Experimental/Water/Source/Runtime/Public/WaterBodyActor.h
--- a/Engine/Plugins/Experimental/Water/Source/Runtime/Public/WaterBodyActor.h
+++ b/Engine/Plugins/Experimental/Water/Source/Runtime/Public/WaterBodyActor.h
@@ -135,7 +135,10 @@ protected:
 
 	void SetWaterWavesInternal(UWaterWavesBase* InWaterWaves);
 
-	TArray<TObjectPtr<UWaterBodyStaticMeshComponent>>& GetWaterBodyStaticMeshComponents() { return WaterBodyStaticMeshComponents; }
+	const TArray<TObjectPtr<UWaterBodyStaticMeshComponent>>& GetWaterBodyStaticMeshComponents() const { return WaterBodyStaticMeshComponents; }
+
+	/** Removes all invalid references in the WaterBodyStaticMeshComponents list. */
+	void CleanupInvalidStaticMeshComponents();
 
 	/** Sets up a new list of water body static mesh components. */
 	void SetWaterBodyStaticMeshComponents(TArrayView<TObjectPtr<UWaterBodyStaticMeshComponent>> NewComponentList, TConstArrayView<TObjectPtr<UWaterBodyStaticMeshComponent>> ComponentsToUnregister = {});
