Water: implemented ShouldExport/Import on ABuoyancyManager to prevent copy/paste from creating additional buoyancy managers besides the single transient one.

diff --git a/Engine/Plugins/Experimental/Water/Source/Runtime/Public/BuoyancyManager.h b/Engine/Plugins/Experimental/Water/Source/Runtime/Public/BuoyancyManager.h
--- a/Engine/Plugins/Experimental/Water/Source/Runtime/Public/BuoyancyManager.h
+++ b/Engine/Plugins/Experimental/Water/Source/Runtime/Public/BuoyancyManager.h
@@ -51,6 +51,10 @@ public:
 
 	void Update(FPhysScene* PhysScene, float DeltaTime);
 
+#if WITH_EDITOR
+	// Prevent import/export of buoyancy manager actors since they should be transient and are always spawned when the map is loaded.
+	virtual bool ShouldImport(FStringView ActorPropString, bool IsMovingLevel) override { return false; }
+#endif // WITH_EDITOR
 private:
 	void InitializeAsyncAux(UBuoyancyComponent* Component);
 	virtual void BeginPlay() override;
