Water: removed #todo in ocean header which was showing up in the tooltip.

#rnx

diff --git a/Engine/Plugins/Experimental/Water/Source/Runtime/Public/WaterBodyOceanComponent.h b/Engine/Plugins/Experimental/Water/Source/Runtime/Public/WaterBodyOceanComponent.h
--- a/Engine/Plugins/Experimental/Water/Source/Runtime/Public/WaterBodyOceanComponent.h
+++ b/Engine/Plugins/Experimental/Water/Source/Runtime/Public/WaterBodyOceanComponent.h
@@ -69,11 +69,12 @@ protected:
 	UPROPERTY(Category = Collision, EditAnywhere, BlueprintReadOnly)
 	FVector CollisionExtents;
 
+	/** The extent of the ocean, centered around water zone to which the ocean belongs. */
 	UPROPERTY(Category = Water, EditAnywhere, BlueprintReadOnly)
 	FVector2D OceanExtents;
 
-	/** Save the water zone location so that the ocean mesh can be regenerated relative to it and match it perfectly without being loaded. */
-	UPROPERTY(Category = Water, EditAnywhere, BlueprintReadOnly, meta=(DisplayName="Ocean Center"))
+	/** Saved water zone location so that the ocean mesh can be regenerated relative to it and match it perfectly without being loaded. */
+	UPROPERTY()
 	FVector2D SavedZoneLocation;
 
 	UPROPERTY(Transient)
