Verify that Water Body River is exposed in the Create menu

diff --git a/Engine/Plugins/Experimental/Water/Source/Editor/Private/Tests/Display/WaterExposedInContentMenu.cpp b/Engine/Plugins/Experimental/Water/Source/Editor/Private/Tests/Display/WaterExposedInContentMenu.cpp
new file mode 100644
--- /dev/null
+++ b/Engine/Plugins/Experimental/Water/Source/Editor/Private/Tests/Display/WaterExposedInContentMenu.cpp
@@ -0,0 +1,33 @@
+// Copyright Epic Games, Inc. All Rights Reserved.
+
+#include "Misc/AutomationTest.h"
+#include "IPlacementModeModule.h"
+
+#if WITH_AUTOMATION_TESTS
+
+// TestRail: C27425803
+IMPLEMENT_SIMPLE_AUTOMATION_TEST(FWaterRiverBodyExposedInContentMenu, "Editor.Plugins.Tools.Water.RiverBodyExposedInContentMenu", EAutomationTestFlags::EditorContext | EAutomationTestFlags::EngineFilter)
+
+bool FWaterRiverBodyExposedInContentMenu::RunTest(const FString& Parameters)
+{
+	// Getting the module
+	const IPlacementModeModule& PlacementModeModule = IPlacementModeModule::Get();
+	// Water Body River is only accessible through All Classes
+	const FName PlacementModeCategoryHandle = TEXT("AllClasses");
+	FString WaterBodyRiverName = TEXT("Water Body River");
+	
+	TArray<TSharedPtr<FPlaceableItem>> OutItems;
+	// We need to manually load All Classes to search for Water Body River
+	PlacementModeModule.Get().RegenerateItemsForCategory(PlacementModeCategoryHandle);
+
+	PlacementModeModule.GetItemsForCategory(PlacementModeCategoryHandle, OutItems);
+	bool RiverBodyFound = OutItems.ContainsByPredicate([&WaterBodyRiverName](const TSharedPtr<FPlaceableItem>& Item) {
+		return Item->DisplayName.ToString().Contains(WaterBodyRiverName);
+		});
+	
+	TestTrue(TEXT("Water Body River is not found in Create Menu"), RiverBodyFound);
+
+	return true;
+}
+
+#endif
\ No newline at end of file
diff --git a/Engine/Plugins/Experimental/Water/Source/Editor/WaterEditor.Build.cs b/Engine/Plugins/Experimental/Water/Source/Editor/WaterEditor.Build.cs
--- a/Engine/Plugins/Experimental/Water/Source/Editor/WaterEditor.Build.cs
+++ b/Engine/Plugins/Experimental/Water/Source/Editor/WaterEditor.Build.cs
@@ -30,6 +30,7 @@ public class WaterEditor : ModuleRules
 				"ComponentVisualizers",
 				"DeveloperSettings",
 				"AdvancedPreviewScene",
+				"PlacementMode",
 			});
 
 		PublicDependencyModuleNames.AddRange(
