Water: Gave AWaterBody a more useful tooltip and removed the outdated todo to remove blueprintable (this is not planned).

diff --git a/Engine/Plugins/Experimental/Water/Source/Runtime/Public/WaterBodyActor.h b/Engine/Plugins/Experimental/Water/Source/Runtime/Public/WaterBodyActor.h
--- a/Engine/Plugins/Experimental/Water/Source/Runtime/Public/WaterBodyActor.h
+++ b/Engine/Plugins/Experimental/Water/Source/Runtime/Public/WaterBodyActor.h
@@ -25,7 +25,16 @@ class WATER_API UDEPRECATED_WaterBodyGenerator : public UObject
 
 // ----------------------------------------------------------------------------------
 
-//@todo_water: Remove Blueprintable
+/**
+ * Base class for all water body actors.
+ *
+ * WaterBodyActors provide a spline-based workflow to create lakes, rivers, and oceans which automatically create meshes,
+ * carve landscapes, and support physics interactions.
+ *
+ * To create a new water body this class must be derived (native or blueprint) and have the `WaterBodyType` property changed to the specific water type.
+ * The new class will automatically have a corresponding UWaterBodyComponent specific to that water body type.
+ * The component class for each water body type can be defined in the Editor Settings.
+ */
 UCLASS(Blueprintable, Abstract, HideCategories = (Tags, Activation, Cooking, Replication, Input, Actor, AssetUserData))
 class WATER_API AWaterBody : public AActor, public IWaterBrushActorInterface
 {
