Water: Fixed an issue relying on cached WaterSplineMetadata when creating the river mesh. This code needs to be able to execute on PostLoad, before OnRegister is called and thus we have to directly grab the metadata from the actor rather than waiting for the pointer to be assigned in OnRegister.

diff --git a/Engine/Plugins/Experimental/Water/Source/Runtime/Private/WaterBodyRiverComponent.cpp b/Engine/Plugins/Experimental/Water/Source/Runtime/Private/WaterBodyRiverComponent.cpp
--- a/Engine/Plugins/Experimental/Water/Source/Runtime/Private/WaterBodyRiverComponent.cpp
+++ b/Engine/Plugins/Experimental/Water/Source/Runtime/Private/WaterBodyRiverComponent.cpp
@@ -14,6 +14,7 @@
 #include "Engine/World.h"
 #include "DynamicMesh/DynamicMesh3.h"
 #include "DynamicMesh/DynamicMeshAttributeSet.h"
+#include "WaterBodyActor.h"
 
 // for working around Chaos issue
 #include "Chaos/Convex.h"
@@ -290,7 +291,11 @@ bool UWaterBodyRiverComponent::GenerateWaterBodyMesh(UE::Geometry::FDynamicMesh3
 	TRACE_CPUPROFILER_EVENT_SCOPE(GenerateRiverMesh);
 
 	const UWaterSplineComponent* SplineComp = GetWaterSpline();
-	if ((SplineComp == nullptr) || (WaterSplineMetadata == nullptr) || (SplineComp->GetNumberOfSplinePoints() < 2))
+
+	check(GetWaterBodyActor());
+
+	const UWaterSplineMetadata* SplineMetadata = GetWaterBodyActor()->GetWaterSplineMetadata();
+	if ((SplineComp == nullptr) || (SplineMetadata == nullptr) || (SplineComp->GetNumberOfSplinePoints() < 2))
 	{
 		return false;
 	}
@@ -305,17 +310,17 @@ bool UWaterBodyRiverComponent::GenerateWaterBodyMesh(UE::Geometry::FDynamicMesh3
 	
 	if (ShapeDilation > 0.f && OutDilatedMesh)
 	{
-		AddTerminalVerticesForRiverSpline(ERiverBoundaryEdge::Start, this, SplineComp, WaterSplineMetadata, *OutDilatedMesh);
+		AddTerminalVerticesForRiverSpline(ERiverBoundaryEdge::Start, this, SplineComp, SplineMetadata, *OutDilatedMesh);
 	}
 
 	for (double DistanceAlongSpline : Distances)
 	{
-		AddVerticesForRiverSplineStep(DistanceAlongSpline, this, SplineComp, WaterSplineMetadata, OutMesh, OutDilatedMesh);
+		AddVerticesForRiverSplineStep(DistanceAlongSpline, this, SplineComp, SplineMetadata, OutMesh, OutDilatedMesh);
 	}
 	
 	if (ShapeDilation > 0.f && OutDilatedMesh)
 	{
-		AddTerminalVerticesForRiverSpline(ERiverBoundaryEdge::End, this, SplineComp, WaterSplineMetadata, *OutDilatedMesh);
+		AddTerminalVerticesForRiverSpline(ERiverBoundaryEdge::End, this, SplineComp, SplineMetadata, *OutDilatedMesh);
 	}
 	return true;
 }
