Water: implemented support for a localized water quadtree when Local-only tessellation is enabled on the parent water zone.
- Local quad tree extents are defined by independent extent parameter on the mesh component
- Quad tree is rebuilt whenever the new local water info region would no longer be covered by the quad tree

diff --git a/Engine/Plugins/Experimental/Water/Source/Editor/Private/WaterLandscapeBrush.cpp b/Engine/Plugins/Experimental/Water/Source/Editor/Private/WaterLandscapeBrush.cpp
--- a/Engine/Plugins/Experimental/Water/Source/Editor/Private/WaterLandscapeBrush.cpp
+++ b/Engine/Plugins/Experimental/Water/Source/Editor/Private/WaterLandscapeBrush.cpp
@@ -202,15 +202,6 @@ void AWaterLandscapeBrush::OnWaterBrushActorChanged(const IWaterBrushActorInterf
 	{
 		RequestLandscapeUpdate(/* bInUserTriggered = */ InParams.bUserTriggered);
 	}
-
-	bool bAllowWaterMeshUpdate = (InParams.PropertyChangedEvent.ChangeType != EPropertyChangeType::Interactive) || WaterEditorSettings->GetUpdateWaterMeshDuringInteractiveChanges();
-	if (bForceUpdateWaterMesh || (bAffectsWaterMesh && bAllowWaterMeshUpdate))
-	{
-		if (UWaterSubsystem* WaterSubsystem = UWaterSubsystem::GetWaterSubsystem(GetWorld()))
-		{
-			WaterSubsystem->MarkAllWaterZonesForRebuild(EWaterZoneRebuildFlags::UpdateWaterMesh);
-		}
-	}
 }
 
 void AWaterLandscapeBrush::OnActorsAffectingLandscapeChanged()
diff --git a/Engine/Plugins/Experimental/Water/Source/Runtime/Private/WaterBodyComponent.cpp b/Engine/Plugins/Experimental/Water/Source/Runtime/Private/WaterBodyComponent.cpp
--- a/Engine/Plugins/Experimental/Water/Source/Runtime/Private/WaterBodyComponent.cpp
+++ b/Engine/Plugins/Experimental/Water/Source/Runtime/Private/WaterBodyComponent.cpp
@@ -1025,7 +1025,8 @@ void UWaterBodyComponent::UpdateComponentVisibility(bool bAllowWaterZoneRebuild)
 					RebuildFlags |= EWaterZoneRebuildFlags::UpdateWaterInfoTexture;
 				}
 
-				WaterZone->MarkForRebuild(RebuildFlags);
+				const FBox WaterBodyBounds = Bounds.GetBox();
+				WaterZone->MarkForRebuild(RebuildFlags, FBox2D(FVector2D(WaterBodyBounds.Min), FVector2D(WaterBodyBounds.Max)));
 			}
 		}
 	}
@@ -1728,8 +1729,8 @@ bool UWaterBodyComponent::SetDynamicParametersOnMID(UMaterialInstanceDynamic* In
 		check(WaterMeshComponent);
 
 		// Location is the bottom left of the zone
-		const FVector2D WaterInfoExtent = FVector2D(WaterZone->GetDynamicWaterMeshExtent());
-		const FVector2D WaterInfoLocation = FVector2D(WaterZone->GetDynamicWaterMeshCenter()) - (WaterInfoExtent / 2.f);
+		const FVector2D WaterInfoExtent = FVector2D(WaterZone->GetDynamicWaterInfoExtent());
+		const FVector2D WaterInfoLocation = FVector2D(WaterZone->GetDynamicWaterInfoCenter()) - (WaterInfoExtent / 2.f);
 
 		FVector4 WaterArea;
 		WaterArea.X = WaterInfoLocation.X;
diff --git a/Engine/Plugins/Experimental/Water/Source/Runtime/Private/WaterInfoRendering.cpp b/Engine/Plugins/Experimental/Water/Source/Runtime/Private/WaterInfoRendering.cpp
--- a/Engine/Plugins/Experimental/Water/Source/Runtime/Private/WaterInfoRendering.cpp
+++ b/Engine/Plugins/Experimental/Water/Source/Runtime/Private/WaterInfoRendering.cpp
@@ -434,6 +434,11 @@ static void UpdateWaterInfoRendering_RenderThread(
 {
 	FScopedLandscapeLODOverride ScopedLandscapeLODOverride(Params);
 
+	static bool bIsUpdatingWaterInfo = false;
+	check(bIsUpdatingWaterInfo == false);
+	bIsUpdatingWaterInfo = true;
+	ON_SCOPE_EXIT {bIsUpdatingWaterInfo = false;};
+
 	FMaterialRenderProxy::UpdateDeferredCachedUniformExpressions();
 	FDeferredUpdateResource::UpdateResources(RHICmdList);
 
@@ -734,9 +739,9 @@ void UpdateWaterInfoRendering(
 	{
 		return;
 	}
-	const FVector ZoneExtent = Context.ZoneToRender->GetDynamicWaterMeshExtent();
+	const FVector ZoneExtent = Context.ZoneToRender->GetDynamicWaterInfoExtent();
 
-	FVector ViewLocation = Context.ZoneToRender->GetDynamicWaterMeshCenter();
+	FVector ViewLocation = Context.ZoneToRender->GetDynamicWaterInfoCenter();
 	ViewLocation.Z = Context.CaptureZ;
 
 	const FBox2D CaptureBounds(FVector2D(ViewLocation - ZoneExtent), FVector2D(ViewLocation + ZoneExtent));
diff --git a/Engine/Plugins/Experimental/Water/Source/Runtime/Private/WaterMeshComponent.cpp b/Engine/Plugins/Experimental/Water/Source/Runtime/Private/WaterMeshComponent.cpp
--- a/Engine/Plugins/Experimental/Water/Source/Runtime/Private/WaterMeshComponent.cpp
+++ b/Engine/Plugins/Experimental/Water/Source/Runtime/Private/WaterMeshComponent.cpp
@@ -171,14 +171,6 @@ FMaterialRelevance UWaterMeshComponent::GetWaterMaterialRelevance(ERHIFeatureLev
 	return Result;
 }
 
-void UWaterMeshComponent::PushTessellatedWaterMeshBoundsToPoxy(const FBox2D& TessellatedWaterMeshBounds)
-{
-	if (SceneProxy)
-	{
-		static_cast<FWaterMeshSceneProxy*>(SceneProxy)->OnTessellatedWaterMeshBoundsChanged_GameThread(TessellatedWaterMeshBounds);
-	}
-}
-
 void UWaterMeshComponent::SetExtentInTiles(FIntPoint NewExtentInTiles)
 {
 	ExtentInTiles = NewExtentInTiles;
@@ -186,6 +178,15 @@ void UWaterMeshComponent::SetExtentInTiles(FIntPoint NewExtentInTiles)
 	MarkRenderStateDirty();
 }
 
+void UWaterMeshComponent::SetDynamicWaterMeshCenter(const FVector2D& NewCenter)
+{
+	if (!DynamicWaterMeshCenter.Equals(NewCenter))
+	{
+		DynamicWaterMeshCenter = NewCenter;
+		MarkWaterMeshGridDirty();
+	}
+}
+
 void UWaterMeshComponent::SetTileSize(float NewTileSize)
 {
 	TileSize = NewTileSize;
@@ -224,8 +225,26 @@ void UWaterMeshComponent::RebuildWaterMesh(float InTileSize, const FIntPoint& In
 	const FVector2D GridPosition = FVector2D(FMath::GridSnap<FVector::FReal>(GetComponentLocation().X, InTileSize), FMath::GridSnap<FVector::FReal>(GetComponentLocation().Y, InTileSize));
 	const FVector2D WorldExtent = FVector2D(InTileSize * InExtentInTiles.X, InTileSize * InExtentInTiles.Y);
 
-	const FBox2D WaterWorldBox = FBox2D(-WorldExtent + GridPosition, WorldExtent + GridPosition);
+	FBox2D WaterWorldBox = FBox2D(-WorldExtent + GridPosition, WorldExtent + GridPosition);
+
+	AWaterZone* WaterZone = Cast<AWaterZone>(GetOwner());
+
+	// when local tessellation is enabled, only use the overlap of the dynamic mesh region with the total water mesh extent to avoid updating the entire water zone.
+	if (WaterZone->IsLocalOnlyTessellationEnabled())
+	{
+		const FVector2D WorldspaceExtent = FVector2D(LocalTessellationExtentInTiles * InTileSize);
+		const FVector2D MeshPosition = GetDynamicWaterMeshCenter();
+		FBox2D DynamicWaterMeshBounds(MeshPosition - WorldspaceExtent, MeshPosition + WorldspaceExtent);
+
+		WaterWorldBox = DynamicWaterMeshBounds.Overlap(WaterWorldBox);
+	}
 	
+	// If the dynamic bounds is outside the full bounds of the water mesh, we shouldn't regenerate the quadtree
+	if (!(WaterWorldBox.GetArea() > 0.f))
+	{
+		return;
+	}
+
 	// This resets the tree to an initial state, ready for node insertion
 	WaterQuadTree.InitTree(WaterWorldBox, InTileSize, InExtentInTiles);
 
@@ -251,15 +270,10 @@ void UWaterMeshComponent::RebuildWaterMesh(float InTileSize, const FIntPoint& In
 		AActor* Actor = WaterBodyComponent->GetOwner();
 		check(Actor);
 
-		if (WaterBodyComponent->GetWaterSpline() == nullptr)
-		{
-			return true;
-		}
-
-		const FBox SplineCompBounds = WaterBodyComponent->GetWaterSpline()->Bounds.GetBox();
+		const FBox WaterBodyBounds = WaterBodyComponent->Bounds.GetBox();
 
 		// Don't process water bodies that has their spline outside of this water mesh
-		if (!SplineCompBounds.IntersectXY(FBox(FVector(WaterWorldBox.Min, 0.0f), FVector(WaterWorldBox.Max, 0.0f))))
+		if (!WaterBodyBounds.IntersectXY(FBox(FVector(WaterWorldBox.Min, 0.0f), FVector(WaterWorldBox.Max, 0.0f))))
 		{
 			return true;
 		}
@@ -567,7 +581,8 @@ void UWaterMeshComponent::RebuildWaterMesh(float InTileSize, const FIntPoint& In
 	{
 		UsedMaterials.Add(FarDistanceMaterial);
 
-		WaterQuadTree.AddFarMesh(FarDistanceMaterial, FarDistanceMeshExtent, FarMeshHeight);
+		const FBox2D FarMeshInnerRegion = FBox2D(-WorldExtent + GridPosition, WorldExtent + GridPosition);
+		WaterQuadTree.AddFarMesh(FarDistanceMaterial, FarMeshInnerRegion, FarDistanceMeshExtent, FarMeshHeight);
 	}
 
 	WaterQuadTree.Unlock(true);
@@ -620,6 +635,7 @@ void UWaterMeshComponent::PostEditChangeProperty(FPropertyChangedEvent& Property
 			|| PropertyName == GET_MEMBER_NAME_CHECKED(UWaterMeshComponent, TessellationFactor)
 			|| PropertyName == GET_MEMBER_NAME_CHECKED(UWaterMeshComponent, TileSize)
 			|| PropertyName == GET_MEMBER_NAME_CHECKED(UWaterMeshComponent, ExtentInTiles)
+			|| PropertyName == GET_MEMBER_NAME_CHECKED(UWaterMeshComponent, LocalTessellationExtentInTiles)
 			|| PropertyName == GET_MEMBER_NAME_CHECKED(UWaterMeshComponent, ForceCollapseDensityLevel)
 			|| PropertyName == GET_MEMBER_NAME_CHECKED(UWaterMeshComponent, FarDistanceMaterial)
 			|| PropertyName == GET_MEMBER_NAME_CHECKED(UWaterMeshComponent, FarDistanceMeshExtent)
diff --git a/Engine/Plugins/Experimental/Water/Source/Runtime/Private/WaterMeshSceneProxy.cpp b/Engine/Plugins/Experimental/Water/Source/Runtime/Private/WaterMeshSceneProxy.cpp
--- a/Engine/Plugins/Experimental/Water/Source/Runtime/Private/WaterMeshSceneProxy.cpp
+++ b/Engine/Plugins/Experimental/Water/Source/Runtime/Private/WaterMeshSceneProxy.cpp
@@ -268,7 +268,6 @@ void FWaterMeshSceneProxy::GetDynamicMeshElements(const TArray<const FSceneView*
 			TraversalDesc.PreViewTranslation = View->ViewMatrices.GetPreViewTranslation();
 			TraversalDesc.LODScale = LODScale;
 			TraversalDesc.bLODMorphingEnabled = !!CVarWaterMeshLODMorphEnabled.GetValueOnRenderThread();
-			TraversalDesc.TessellatedWaterMeshBounds = TessellatedWaterMeshBounds;
 
 #if !(UE_BUILD_SHIPPING || UE_BUILD_TEST)
 			//Debug
@@ -508,7 +507,6 @@ void FWaterMeshSceneProxy::GetDynamicRayTracingInstances(FRayTracingMaterialGath
 	TraversalDesc.Frustum = FConvexVolume(); // Default volume to disable frustum culling
 	TraversalDesc.LODScale = LODScale;
 	TraversalDesc.bLODMorphingEnabled = !!CVarWaterMeshLODMorphEnabled.GetValueOnRenderThread();
-	TraversalDesc.TessellatedWaterMeshBounds = TessellatedWaterMeshBounds;
 
 	WaterQuadTree.BuildWaterTileInstanceData(TraversalDesc, WaterInstanceData);
 
diff --git a/Engine/Plugins/Experimental/Water/Source/Runtime/Private/WaterQuadTree.cpp b/Engine/Plugins/Experimental/Water/Source/Runtime/Private/WaterQuadTree.cpp
--- a/Engine/Plugins/Experimental/Water/Source/Runtime/Private/WaterQuadTree.cpp
+++ b/Engine/Plugins/Experimental/Water/Source/Runtime/Private/WaterQuadTree.cpp
@@ -279,11 +279,9 @@ void FWaterQuadTree::FNode::SelectLODWithinBounds(const FNodeData& InNodeData, i
 		return;
 	}
 
-	check(InTraversalDesc.TessellatedWaterMeshBounds.bIsValid);
 	if (InLODLevel == 0)
 	{
-		if ((InTraversalDesc.TessellatedWaterMeshBounds.IsInsideOrOn(FVector2D(Bounds.Min)) && InTraversalDesc.TessellatedWaterMeshBounds.IsInsideOrOn(FVector2D(Bounds.Max))) &&
-			CanRender(0, InTraversalDesc.ForceCollapseDensityLevel, WaterBodyRenderData))
+		if (CanRender(0, InTraversalDesc.ForceCollapseDensityLevel, WaterBodyRenderData))
 		{
 			AddNodeForRender(InNodeData, WaterBodyRenderData, 0, InLODLevel, InTraversalDesc, Output);
 		}
@@ -625,7 +623,7 @@ void FWaterQuadTree::AddLake(const TArray<FVector2D>& InPoly, const FBox& InLake
 	AddLakeRecursive(InPoly, LakeBounds, FVector2D(InLakeBounds.Min.Z, InLakeBounds.Max.Z), true, TreeDepth * 2, InWaterBodyIndex);
 }
 
-void FWaterQuadTree::AddFarMesh(const UMaterialInterface* InFarMeshMaterial, double InFarDistanceMeshExtent, double InFarDistanceMeshHeight)
+void FWaterQuadTree::AddFarMesh(const UMaterialInterface* InFarMeshMaterial, const FBox2D& InInnerRegion, double InFarDistanceMeshExtent, double InFarDistanceMeshHeight)
 {
 	// Checking for not being read only here to keep things consistent with the other Add functions. In reality the FarMesh isn't added to the QuadTree itself, so it could technically be done whenever.
 	ensure(!bIsReadOnly);
@@ -645,9 +643,9 @@ void FWaterQuadTree::AddFarMesh(const UMaterialInterface* InFarMeshMaterial, dou
 	FarMeshData.InstanceData.SetNum(8);
 	FarMeshData.Material = InFarMeshMaterial;
 
-	const FVector2D WaterCenter = GetTileRegion().GetCenter();
-	const FVector2D WaterExtents = GetTileRegion().GetExtent();
-	const FVector2D WaterSize = GetTileRegion().GetSize();
+	const FVector2D WaterCenter = InInnerRegion.GetCenter();
+	const FVector2D WaterExtents = InInnerRegion.GetExtent();
+	const FVector2D WaterSize = InInnerRegion.GetSize();
 	const FVector2D TileOffets[] = { {-1.0, 1.0}, {0.0, 1.0}, {1.0, 1.0}, {1.0, 0.0}, {1.0, -1.0}, {0.0, -1.0}, {-1.0, -1.0}, {-1.0, 0.0} };
 
 	for (int32 i = 0; i < 8; i++)
@@ -707,14 +705,8 @@ void FWaterQuadTree::BuildWaterTileInstanceData(const FTraversalDesc& InTraversa
 {
 	TRACE_CPUPROFILER_EVENT_SCOPE(BuildWaterTileInstanceData);
 	check(bIsReadOnly);
-	if (InTraversalDesc.TessellatedWaterMeshBounds.bIsValid)
-	{
-		NodeData.Nodes[0].SelectLODWithinBounds(NodeData, TreeDepth, InTraversalDesc, Output);
-	}
-	else
-	{
-		NodeData.Nodes[0].SelectLOD(NodeData, TreeDepth, InTraversalDesc, Output);
-	}
+
+	NodeData.Nodes[0].SelectLOD(NodeData, TreeDepth, InTraversalDesc, Output);
 
 	// Append Far Mesh tiles
 	if (FarMeshData.InstanceData.Num() > 0 && FarMeshData.MaterialIndex != INDEX_NONE)
diff --git a/Engine/Plugins/Experimental/Water/Source/Runtime/Private/WaterSubsystem.cpp b/Engine/Plugins/Experimental/Water/Source/Runtime/Private/WaterSubsystem.cpp
--- a/Engine/Plugins/Experimental/Water/Source/Runtime/Private/WaterSubsystem.cpp
+++ b/Engine/Plugins/Experimental/Water/Source/Runtime/Private/WaterSubsystem.cpp
@@ -530,7 +530,7 @@ void UWaterSubsystem::MarkWaterZonesInRegionForRebuild(const FBox2D& InUpdateReg
 
 			if (WaterZoneBounds.Intersect(InUpdateRegion))
 			{
-				WaterZone->MarkForRebuild(InRebuildFlags);
+				WaterZone->MarkForRebuild(InRebuildFlags, InUpdateRegion);
 			}
 		}
 	}
diff --git a/Engine/Plugins/Experimental/Water/Source/Runtime/Private/WaterViewExtension.cpp b/Engine/Plugins/Experimental/Water/Source/Runtime/Private/WaterViewExtension.cpp
--- a/Engine/Plugins/Experimental/Water/Source/Runtime/Private/WaterViewExtension.cpp
+++ b/Engine/Plugins/Experimental/Water/Source/Runtime/Private/WaterViewExtension.cpp
@@ -7,10 +7,18 @@
 #include "SceneView.h"
 #include "WaterMeshComponent.h"
 
-static TAutoConsoleVariable<float> CVarFreezeWaterLODUpdates(
-	TEXT("r.Water.FreezeLODUpdates"),
-	0,
-	TEXT(""),
+static TAutoConsoleVariable<bool> CVarLocalTessellationFreeze(
+	TEXT("r.Water.WaterMesh.LocalTessellation.Freeze"),
+	false,
+	TEXT("Pauses the local tessellation updates to allow the view to move forward without moving the sliding window.\n")
+	TEXT("Can be used to view things outside the sliding window more closely."),
+	ECVF_Default);
+
+static TAutoConsoleVariable<float> CVarLocalTessellationUpdateMargin(
+	TEXT("r.Water.WaterMesh.LocalTessellation.UpdateMargin"),
+	15000.,
+	TEXT("Controls the minimum distance between the view and the edge of the dynamic water mesh when local tessellation is enabled.\n")
+	TEXT("If the view is less than UpdateMargin units away from the edge, it moves the sliding window forward."),
 	ECVF_Default);
 
 // ----------------------------------------------------------------------------------
@@ -31,7 +39,7 @@ void FWaterViewExtension::SetupViewFamily(FSceneViewFamily& InViewFamily)
 
 void FWaterViewExtension::SetupView(FSceneViewFamily& InViewFamily, FSceneView& InView)
 {
-	if (CVarFreezeWaterLODUpdates.GetValueOnGameThread() != 0)
+	if (CVarLocalTessellationFreeze.GetValueOnGameThread())
 	{
 		return;
 	}
@@ -41,64 +49,80 @@ void FWaterViewExtension::SetupView(FSceneViewFamily& InViewFamily, FSceneView&
 	const TWeakObjectPtr<UWorld> WorldPtr = GetWorld();
 	check(WorldPtr.IsValid())
 
+	// Prevent re-entrancy. 
+	// Since the water info render will update the view extensions we could end up with a re-entrant case.
 	static bool bUpdatingWaterInfo = false;
-	if (!bUpdatingWaterInfo)
+	if (bUpdatingWaterInfo)
+	{
+		return;
+	}
+	bUpdatingWaterInfo = true;
+	ON_SCOPE_EXIT { bUpdatingWaterInfo = false; };
+
+	for (const TPair<AWaterZone*, UE::WaterInfo::FRenderingContext>& Pair : WaterInfoContextsToRender)
 	{
-		bUpdatingWaterInfo = true;
+		AWaterZone* WaterZone = Pair.Key;
+		check(WaterZone);
 
-		for (const TPair<AWaterZone*, UE::WaterInfo::FRenderingContext>& Pair : WaterInfoContextsToRender)
+		if (WaterZone->IsLocalOnlyTessellationEnabled())
 		{
-			AWaterZone* WaterZone = Pair.Key;
-			check(WaterZone);
+			UWaterMeshComponent* WaterMesh= WaterZone->GetWaterMeshComponent();
+			check(WaterMesh);
 
-			if (WaterZone->IsLocalOnlyTessellationEnabled())
-			{
-				UWaterMeshComponent* WaterMesh= WaterZone->GetWaterMeshComponent();
-				check(WaterMesh);
+			const double TileSize = WaterMesh->GetTileSize();
 
-				const double TileSize = WaterMesh->GetTileSize();
+			const FVector WaterMeshCenter(ViewLocation.GridSnap(TileSize));
 
-				const FVector WaterMeshCenter = ViewLocation.GridSnap(TileSize);
+			const FVector2D WaterInfoHalfExtent(WaterZone->GetDynamicWaterInfoExtent() / 2.0);
+			const FBox2D WaterInfoBounds(FVector2D(WaterMeshCenter) - WaterInfoHalfExtent, FVector2D(WaterMeshCenter) + WaterInfoHalfExtent);
 
-				const FVector WaterMeshHalfExtent = WaterZone->GetDynamicWaterMeshExtent() / 2.0;
-				const FBox2D WaterMeshBounds(FVector2D(WaterMeshCenter - WaterMeshHalfExtent), FVector2D(WaterMeshCenter + WaterMeshHalfExtent));
+			WaterZone->SetLocalTessellationCenter(WaterMeshCenter);
 
-				WaterZone->SetLocalTessellationCenter(WaterMeshCenter);
-				WaterMesh->PushTessellatedWaterMeshBoundsToPoxy(WaterMeshBounds.ExpandBy(2 * TileSize));
+			// Trigger the next update when the camera is <UpdateMargin> units away from the border of the current window.
+			const FVector2D UpdateMargin(CVarLocalTessellationUpdateMargin.GetValueOnGameThread());
 
-				const FBox2D CurrentCellBounds(FVector2D(ViewLocation) - FVector2D(TileSize / 2.0), FVector2D(ViewLocation) + FVector2D(TileSize / 2.0));
-				UpdateBounds.Add(WaterZone, CurrentCellBounds);
-			}
+			// Keep a minimum of <1., 1.> bounds to avoid updating every frame if the update margin is larger than the zone.
+			const FVector2D UpdateExtents = FVector2D::Max(FVector2D(1., 1.), WaterInfoHalfExtent - UpdateMargin);
+			const FBox2D UpdateBounds(FVector2D(WaterMeshCenter) - UpdateExtents, FVector2D(WaterMeshCenter) + UpdateExtents);
 
-			// Update the material instances now that the tessellated mesh bounds have changed.
-			WaterZone->ForEachWaterBodyComponent([](UWaterBodyComponent* WaterBodyComponent)
-			{
-				WaterBodyComponent->UpdateMaterialInstances();
-				return true;
-			});
+			WaterInfoUpdateBounds.Add(WaterZone, UpdateBounds);
+
+			const FVector2D WaterQuadTreeHalfExtent = WaterMesh->GetLocalTessellationExtentInTiles() * TileSize;
+			const FVector2D WaterQuadTreeCenter = WaterMesh->GetDynamicWaterMeshCenter();
+			const FBox2D WaterQuadTreeBounds(WaterQuadTreeCenter - WaterQuadTreeHalfExtent, WaterQuadTreeCenter + WaterQuadTreeHalfExtent);
 
-			const UE::WaterInfo::FRenderingContext& Context(Pair.Value);
-			UE::WaterInfo::UpdateWaterInfoRendering(WorldPtr.Get()->Scene, Context);
+			// If the new water info bounds would be outside the water mesh, recenter and rebuild the water mesh
+			if (!WaterQuadTreeBounds.IsInside(WaterInfoBounds.ExpandBy(WaterInfoBounds.GetExtent())))
+			{
+				WaterMesh->SetDynamicWaterMeshCenter(FVector2D(WaterMeshCenter));
+			}
 		}
 
-		WaterInfoContextsToRender.Empty();
+		// Update the material instances now that the tessellated mesh bounds have changed.
+		WaterZone->ForEachWaterBodyComponent([](UWaterBodyComponent* WaterBodyComponent)
+		{
+			WaterBodyComponent->UpdateMaterialInstances();
+			return true;
+		});
+
+		const UE::WaterInfo::FRenderingContext& Context(Pair.Value);
+		UE::WaterInfo::UpdateWaterInfoRendering(WorldPtr.Get()->Scene, Context);
+	}
 
-		// Check if the view location is no longer within the current update bounds of a water zone and if so, queue an update for it.
-		for (AWaterZone* WaterZone : TActorRange<AWaterZone>(WorldPtr.Get()))
+	WaterInfoContextsToRender.Empty();
+
+	// Check if the view location is no longer within the current update bounds of a water zone and if so, queue an update for it.
+	for (AWaterZone* WaterZone : TActorRange<AWaterZone>(WorldPtr.Get()))
+	{
+		if (WaterZone->IsLocalOnlyTessellationEnabled())
 		{
-			if (WaterZone->IsLocalOnlyTessellationEnabled())
+			const FBox2D* WaterInfoBounds = WaterInfoUpdateBounds.Find(WaterZone);
+			if (WaterInfoBounds == nullptr || !WaterInfoBounds->IsInside(FVector2D(ViewLocation)))
 			{
-				const FBox2D* ZoneUpdateBounds = UpdateBounds.Find(WaterZone);
-				if (ZoneUpdateBounds == nullptr || !ZoneUpdateBounds->IsInside(FVector2D(ViewLocation)))
-				{
-					WaterZone->MarkForRebuild(EWaterZoneRebuildFlags::UpdateWaterInfoTexture);
-				}
+				WaterZone->MarkForRebuild(EWaterZoneRebuildFlags::UpdateWaterInfoTexture);
 			}
 		}
-
-		bUpdatingWaterInfo = false;
 	}
-
 }
 
 void FWaterViewExtension::PreRenderViewFamily_RenderThread(FRDGBuilder& GraphBuilder, FSceneViewFamily& InViewFamily)
diff --git a/Engine/Plugins/Experimental/Water/Source/Runtime/Private/WaterZoneActor.cpp b/Engine/Plugins/Experimental/Water/Source/Runtime/Private/WaterZoneActor.cpp
--- a/Engine/Plugins/Experimental/Water/Source/Runtime/Private/WaterZoneActor.cpp
+++ b/Engine/Plugins/Experimental/Water/Source/Runtime/Private/WaterZoneActor.cpp
@@ -45,7 +45,7 @@ AWaterZone::AWaterZone(const FObjectInitializer& Initializer)
 	WaterMesh = CreateDefaultSubobject<UWaterMeshComponent>(TEXT("WaterMesh"));
 	SetRootComponent(WaterMesh);
 	ZoneExtent = FVector2D(51200., 51200.);
-	LocalTessellationExtent = FVector(35000., 35000., 10000.);
+	LocalTessellationExtent = FVector(51200., 51200., 10000.);
 	
 #if	WITH_EDITOR
 	// Setup bounds component
@@ -157,17 +157,27 @@ void AWaterZone::DeclareConstructClasses(TArray<FTopLevelAssetPath>& OutConstruc
 }
 #endif
 
-void AWaterZone::MarkForRebuild(EWaterZoneRebuildFlags Flags)
+void AWaterZone::MarkForRebuild(EWaterZoneRebuildFlags Flags, const FBox2D& UpdateRegion)
 {
 	if (EnumHasAnyFlags(Flags, EWaterZoneRebuildFlags::UpdateWaterMesh))
 	{
-		UE_LOG(LogWater, Verbose, TEXT("AWaterZone::MarkForRebuild (UpdateWaterMesh)"));
-		WaterMesh->MarkWaterMeshGridDirty();
+		const FBox2D WaterQuadTreeBounds = WaterMesh->GetWaterQuadTree().GetTileRegion();
+		// Suppress water mesh updates which occur outside the bounds of the water quad tree.
+		if ((!UpdateRegion.bIsValid) || UpdateRegion.Intersect(WaterQuadTreeBounds))
+		{
+			UE_LOG(LogWater, Verbose, TEXT("AWaterZone::MarkForRebuild (UpdateWaterMesh)"));
+			WaterMesh->MarkWaterMeshGridDirty();
+		}
 	}
 	if (EnumHasAnyFlags(Flags, EWaterZoneRebuildFlags::UpdateWaterInfoTexture))
 	{
-		UE_LOG(LogWater, Verbose, TEXT("AWaterZone::MarkForRebuild (UpdateWaterInfoTexture)"));
-		bNeedsWaterInfoRebuild = true;
+		const FBox WaterInfoBounds = GetDynamicWaterInfoBounds();
+		const FBox2D WaterInfoBounds2D(FVector2D(WaterInfoBounds.Min), FVector2D(WaterInfoBounds.Max));
+		if ((!UpdateRegion.bIsValid) || UpdateRegion.Intersect(WaterInfoBounds2D))
+		{
+			UE_LOG(LogWater, Verbose, TEXT("AWaterZone::MarkForRebuild (UpdateWaterInfoTexture)"));
+			bNeedsWaterInfoRebuild = true;
+		}
 	}
 }
 
@@ -198,7 +208,8 @@ void AWaterZone::AddWaterBodyComponent(UWaterBodyComponent* WaterBodyComponent)
 			RebuildFlags |= EWaterZoneRebuildFlags::UpdateWaterMesh;
 		}
 
-		MarkForRebuild(RebuildFlags);
+		const FBox WaterBodyBounds = WaterBodyComponent->Bounds.GetBox();
+		MarkForRebuild(RebuildFlags, FBox2D(FVector2D(WaterBodyBounds.Min), FVector2D(WaterBodyBounds.Max)));
 	}
 }
 
@@ -219,7 +230,8 @@ void AWaterZone::RemoveWaterBodyComponent(UWaterBodyComponent* WaterBodyComponen
 			RebuildFlags |= EWaterZoneRebuildFlags::UpdateWaterMesh;
 		}
 		
-		MarkForRebuild(RebuildFlags);
+		const FBox WaterBodyBounds = WaterBodyComponent->Bounds.GetBox();
+		MarkForRebuild(RebuildFlags, FBox2D(FVector2D(WaterBodyBounds.Min), FVector2D(WaterBodyBounds.Max)));
 	}
 }
 
@@ -515,7 +527,7 @@ bool AWaterZone::UpdateWaterInfoTexture()
 	return true;
 }
 
-FVector AWaterZone::GetDynamicWaterMeshCenter() const
+FVector AWaterZone::GetDynamicWaterInfoCenter() const
 {
 	if (IsLocalOnlyTessellationEnabled())
 	{
@@ -527,7 +539,7 @@ FVector AWaterZone::GetDynamicWaterMeshCenter() const
 	}
 }
 
-FVector AWaterZone::GetDynamicWaterMeshExtent() const
+FVector AWaterZone::GetDynamicWaterInfoExtent() const
 {
 	if (IsLocalOnlyTessellationEnabled())
 	{
@@ -540,6 +552,13 @@ FVector AWaterZone::GetDynamicWaterMeshExtent() const
 	}
 }
 
+FBox AWaterZone::GetDynamicWaterInfoBounds() const
+{
+	const FVector WaterInfoCenter(GetDynamicWaterInfoCenter());
+	const FVector WaterInfoHalfExtent(GetDynamicWaterInfoExtent() / 2);
+	return FBox(WaterInfoCenter - WaterInfoHalfExtent, WaterInfoCenter + WaterInfoHalfExtent);
+}
+
 void AWaterZone::OnLevelAddedToWorld(ULevel* InLevel, UWorld* InWorld)
 {
 	OnLevelChanged(InLevel, InWorld);
@@ -557,7 +576,7 @@ void AWaterZone::OnLevelChanged(ULevel* InLevel, UWorld* InWorld)
 		return;
 	}
 
-	const FBox WaterZoneBounds = GetZoneBounds();
+	const FBox WaterZoneBounds = GetDynamicWaterInfoBounds();
 	const bool bContainsActorsAffectingWaterZone = Algo::AnyOf(InLevel->Actors, [this, &WaterZoneBounds](const AActor* Actor)
 	{
 		return IsAffectingWaterZone(WaterZoneBounds, Actor);
@@ -580,9 +599,9 @@ bool AWaterZone::IsAffectingWaterZone(const FBox& InWaterZoneBounds, const AActo
 	{
 		if (const UWaterBodyComponent* WaterBodyComponent = WaterBodyActor->GetWaterBodyComponent())
 		{
-			if (WaterBodyComponent->GetWaterZone() == this)
+			if (WaterBodyComponent->GetWaterZone() == this && WaterBodyComponent->AffectsWaterInfo())
 			{
-				return WaterBodyComponent->AffectsWaterInfo();
+				return InWaterZoneBounds.IntersectXY(WaterBodyComponent->Bounds.GetBox());
 			}
 		}
 	}
diff --git a/Engine/Plugins/Experimental/Water/Source/Runtime/Public/WaterMeshComponent.h b/Engine/Plugins/Experimental/Water/Source/Runtime/Public/WaterMeshComponent.h
--- a/Engine/Plugins/Experimental/Water/Source/Runtime/Public/WaterMeshComponent.h
+++ b/Engine/Plugins/Experimental/Water/Source/Runtime/Public/WaterMeshComponent.h
@@ -49,8 +49,6 @@ public:
 	/** Use this instead of GetMaterialRelevance, since this one will go over all materials from all tiles */
 	FMaterialRelevance GetWaterMaterialRelevance(ERHIFeatureLevel::Type InFeatureLevel) const;
 
-	void PushTessellatedWaterMeshBoundsToPoxy(const FBox2D& TessellatedWaterMeshBounds);
-
 	const FWaterQuadTree& GetWaterQuadTree() const { return WaterQuadTree; }
 
 	const TSet<TObjectPtr<UMaterialInterface>>& GetUsedMaterialsSet() const { return UsedMaterials; }
@@ -64,6 +62,10 @@ public:
 	void SetExtentInTiles(FIntPoint NewExtentInTiles);
 	FIntPoint GetExtentInTiles() const { return ExtentInTiles; }
 
+	FIntPoint GetLocalTessellationExtentInTiles() const { return LocalTessellationExtentInTiles; }
+	void SetDynamicWaterMeshCenter(const FVector2D& NewCenter);
+	FVector2D GetDynamicWaterMeshCenter() const { return DynamicWaterMeshCenter; }
+
 	void SetTileSize(float NewTileSize);
 	float GetTileSize() const { return TileSize; }
 
@@ -100,6 +102,13 @@ private:
 	UPROPERTY(EditAnywhere, Category = Rendering, meta = (ClampMin = "1", AllowPrivateAcces = "true"))
 	FIntPoint ExtentInTiles = FIntPoint(64, 64);
 
+	/** The extent of the system in number of tiles when local tessellation is enabled. Maximum number of tiles for this system will be ExtentInTiles.X*2*ExtentInTiles.Y*2 */
+	UPROPERTY(EditAnywhere, Category = Rendering, meta = (ClampMin = "1"))
+	FIntPoint LocalTessellationExtentInTiles = FIntPoint(32, 32);
+
+	/** The current center of the dynamic water mesh. Updated by the water view extension whenever the view location crosses the update bounds. */
+	FVector2D DynamicWaterMeshCenter = FVector2D::ZeroVector;
+
 	/** Tiles containing water, stored in a quad tree */
 	FWaterQuadTree WaterQuadTree;
 
diff --git a/Engine/Plugins/Experimental/Water/Source/Runtime/Public/WaterQuadTree.h b/Engine/Plugins/Experimental/Water/Source/Runtime/Public/WaterQuadTree.h
--- a/Engine/Plugins/Experimental/Water/Source/Runtime/Public/WaterQuadTree.h
+++ b/Engine/Plugins/Experimental/Water/Source/Runtime/Public/WaterQuadTree.h
@@ -109,7 +109,6 @@ struct FWaterQuadTree
 		FVector PreViewTranslation = FVector::ZeroVector;
 		FConvexVolume Frustum;
 		bool bLODMorphingEnabled = true;
-		FBox2D TessellatedWaterMeshBounds = FBox2D(ForceInit);
 
 #if !(UE_BUILD_SHIPPING || UE_BUILD_TEST)
 		// Debug
@@ -136,8 +135,8 @@ struct FWaterQuadTree
 	/** Add Lake by giving a closed spline that represents the lake */
 	void AddLake(const TArray<FVector2D>& InPoly, const FBox& InLakeBounds, uint32 InWaterBodyIndex);
 
-	/** Add an automatically generated mesh (8 quads) skirt around the main water quadtree which extends out InFarDistanceMeshExtent, is placed at Z value InFarDistanceMeshHeight and is rendered using InFarMeshMaterial */
-	void AddFarMesh(const UMaterialInterface* InFarMeshMaterial, double InFarDistanceMeshExtent, double InFarDistanceMeshHeight);
+	/** Add an automatically generated mesh (8 quads) skirt around InInnerRegion which extends out InFarDistanceMeshExtent, is placed at Z value InFarDistanceMeshHeight and is rendered using InFarMeshMaterial */
+	void AddFarMesh(const UMaterialInterface* InFarMeshMaterial, const FBox2D& InInnerRegion, double InFarDistanceMeshExtent, double InFarDistanceMeshHeight);
 
 	/** Assign an index to each material */
 	void BuildMaterialIndices();
diff --git a/Engine/Plugins/Experimental/Water/Source/Runtime/Public/WaterViewExtension.h b/Engine/Plugins/Experimental/Water/Source/Runtime/Public/WaterViewExtension.h
--- a/Engine/Plugins/Experimental/Water/Source/Runtime/Public/WaterViewExtension.h
+++ b/Engine/Plugins/Experimental/Water/Source/Runtime/Public/WaterViewExtension.h
@@ -32,7 +32,7 @@ private:
 	 * For each water zone, store the bounds of the tile from which the water zone was last rendered.
 	 * When the view location crosses the bounds, submit a new WaterInfo update to reflect the new active area
 	 */
-	TMap<AWaterZone*, FBox2D> UpdateBounds;
+	TMap<AWaterZone*, FBox2D> WaterInfoUpdateBounds;
 };
 
 #if UE_ENABLE_INCLUDE_ORDER_DEPRECATED_IN_5_2
diff --git a/Engine/Plugins/Experimental/Water/Source/Runtime/Public/WaterZoneActor.h b/Engine/Plugins/Experimental/Water/Source/Runtime/Public/WaterZoneActor.h
--- a/Engine/Plugins/Experimental/Water/Source/Runtime/Public/WaterZoneActor.h
+++ b/Engine/Plugins/Experimental/Water/Source/Runtime/Public/WaterZoneActor.h
@@ -32,7 +32,7 @@ public:
 	UWaterMeshComponent* GetWaterMeshComponent() { return WaterMesh; }
 	const UWaterMeshComponent* GetWaterMeshComponent() const { return WaterMesh; }
 
-	void MarkForRebuild(EWaterZoneRebuildFlags Flags);
+	void MarkForRebuild(EWaterZoneRebuildFlags Flags, const FBox2D& RebuildRegion = FBox2D(EForceInit::ForceInitToZero));
 	void Update();
 		
 	/** Execute a predicate function on each valid water body within the water zone. Predicate should return false for early exit. */
@@ -52,8 +52,9 @@ public:
 
 	uint32 GetVelocityBlurRadius() const { return VelocityBlurRadius; }
 
-	FVector GetDynamicWaterMeshCenter() const;
-	FVector GetDynamicWaterMeshExtent() const;
+	FVector GetDynamicWaterInfoCenter() const;
+	FVector GetDynamicWaterInfoExtent() const;
+	FBox GetDynamicWaterInfoBounds() const;
 
 	bool IsLocalOnlyTessellationEnabled() const { return bEnableLocalOnlyTessellation; }
 	void SetLocalTessellationCenter(const FVector& NewCenter) { LocalTessellationCenter = NewCenter;}
@@ -79,14 +80,14 @@ public:
 	virtual FBox GetStreamingBounds() const override;
 #endif //WITH_EDITOR
 
-	UE_DEPRECATED(5.3, "Function renamed to GetDynamicWaterMeshCenter")
-	FVector GetTessellatedWaterMeshCenter() const { return GetDynamicWaterMeshCenter(); }
+	UE_DEPRECATED(5.3, "Function renamed to GetDynamicWaterInfoCenter")
+	FVector GetTessellatedWaterMeshCenter() const { return GetDynamicWaterInfoCenter(); }
 
 	UE_DEPRECATED(5.3, "Function renamed to SetLocalTessellationCenter")
 	void SetTessellatedWaterMeshCenter(FVector NewCenter) { SetLocalTessellationCenter(NewCenter); }
 
-	UE_DEPRECATED(5.3, "Function renamed to GetDynamicWaterMeshExtent")
-	FVector GetTessellatedWaterMeshExtent() const { return GetDynamicWaterMeshExtent(); }
+	UE_DEPRECATED(5.3, "Function renamed to GetDynamicWaterInfoExtent")
+	FVector GetTessellatedWaterMeshExtent() const { return GetDynamicWaterInfoExtent(); }
 
 	UE_DEPRECATED(5.3, "Function renamed to IsLocalTessellationEnabled.")
 	bool IsNonTessellatedLODMeshEnabled() const { return IsLocalOnlyTessellationEnabled(); }
