Water: disable mapcheck error when a water body with static mesh enabled doesn't have an owning water zone.

diff --git a/Engine/Plugins/Experimental/Water/Source/Runtime/Private/WaterBodyComponent.cpp b/Engine/Plugins/Experimental/Water/Source/Runtime/Private/WaterBodyComponent.cpp
--- a/Engine/Plugins/Experimental/Water/Source/Runtime/Private/WaterBodyComponent.cpp
+++ b/Engine/Plugins/Experimental/Water/Source/Runtime/Private/WaterBodyComponent.cpp
@@ -1145,11 +1145,14 @@ TArray<TSharedRef<FTokenizedMessage>> UWaterBodyComponent::CheckWaterBodyStatus(
 				AWaterZone* WaterZone = GetWaterZone();
 				if (WaterZone == nullptr)
 				{
-					Result.Add(FTokenizedMessage::Create(EMessageSeverity::Error)
-						->AddToken(FUObjectToken::Create(this))
-						->AddToken(FTextToken::Create(FText::Format(
-							LOCTEXT("MapCheck_Message_MissingWaterZone", "Water body {0} requires a WaterZone actor to be rendered. Please add one to the map. "),
-							FText::FromString(GetWaterBodyActor()->GetActorLabel())))));
+					if (!StaticMeshSettings.bEnableWaterBodyStaticMesh)
+					{
+						Result.Add(FTokenizedMessage::Create(EMessageSeverity::Error)
+							->AddToken(FUObjectToken::Create(this))
+							->AddToken(FTextToken::Create(FText::Format(
+								LOCTEXT("MapCheck_Message_MissingWaterZone", "Water body {0} without a static mesh fallback requires a WaterZone actor to be rendered. Please add one to the map or enable the static mesh fallback. "),
+								FText::FromString(GetWaterBodyActor()->GetActorLabel())))));
+					}
 				}
 				else
 				{
