Verify that Water Body Actors are exposed in the Create menu

diff --git a/Engine/Plugins/Experimental/Water/Source/Editor/Private/Tests/Display/WaterExposedInContentMenu.cpp b/Engine/Plugins/Experimental/Water/Source/Editor/Private/Tests/Display/WaterExposedInContentMenu.cpp
--- a/Engine/Plugins/Experimental/Water/Source/Editor/Private/Tests/Display/WaterExposedInContentMenu.cpp
+++ b/Engine/Plugins/Experimental/Water/Source/Editor/Private/Tests/Display/WaterExposedInContentMenu.cpp
@@ -4,28 +4,42 @@
 #include "IPlacementModeModule.h"
 
 #if WITH_AUTOMATION_TESTS
+// TestRail: 27425803
+// TestRail: 36679397
+// TestRail: 36679398
+// TestRail: 36679399
+IMPLEMENT_COMPLEX_AUTOMATION_TEST(FWaterExposedInContentMenuTest, "Editor.Plugins.Tools.Water", EAutomationTestFlags::EditorContext | EAutomationTestFlags::EngineFilter)
 
-// TestRail: C27425803
-IMPLEMENT_SIMPLE_AUTOMATION_TEST(FWaterRiverBodyExposedInContentMenu, "Editor.Plugins.Tools.Water.RiverBodyExposedInContentMenu", EAutomationTestFlags::EditorContext | EAutomationTestFlags::EngineFilter)
+void FWaterExposedInContentMenuTest::GetTests(TArray<FString>& OutBeautifiedNames, TArray<FString>& OutTestCommands) const
+{
+	OutTestCommands.Add(TEXT("Water Body River"));
+	OutBeautifiedNames.Add(TEXT("RiverBodyExposedInContentMenu"));
+	OutTestCommands.Add(TEXT("Water Body Lake"));
+	OutBeautifiedNames.Add(TEXT("LakeBodyExposedInContentMenu"));
+	OutTestCommands.Add(TEXT("Water Body Ocean"));
+	OutBeautifiedNames.Add(TEXT("OceanBodyExposedInContentMenu"));
+	OutTestCommands.Add(TEXT("Water Body Island"));
+	OutBeautifiedNames.Add(TEXT("IslandBodyExposedInContentMenu"));
+}
 
-bool FWaterRiverBodyExposedInContentMenu::RunTest(const FString& Parameters)
+bool FWaterExposedInContentMenuTest::RunTest(const FString& Parameters)
 {
 	// Getting the module
 	const IPlacementModeModule& PlacementModeModule = IPlacementModeModule::Get();
-	// Water Body River is only accessible through All Classes
+	// Water Body Actors are only accessible through All Classes
 	const FName PlacementModeCategoryHandle = TEXT("AllClasses");
-	FString WaterBodyRiverName = TEXT("Water Body River");
-	
+	FString WaterBodyName = Parameters;
+
 	TArray<TSharedPtr<FPlaceableItem>> OutItems;
-	// We need to manually load All Classes to search for Water Body River
+	// We need to manually load All Classes to search for a Water Body Actor
 	PlacementModeModule.Get().RegenerateItemsForCategory(PlacementModeCategoryHandle);
 
 	PlacementModeModule.GetItemsForCategory(PlacementModeCategoryHandle, OutItems);
-	bool RiverBodyFound = OutItems.ContainsByPredicate([&WaterBodyRiverName](const TSharedPtr<FPlaceableItem>& Item) {
-		return Item->DisplayName.ToString().Contains(WaterBodyRiverName);
+	bool WaterBodyFound = OutItems.ContainsByPredicate([&WaterBodyName](const TSharedPtr<FPlaceableItem>& Item) {
+		return Item->DisplayName.ToString().Contains(WaterBodyName);
 		});
-	
-	TestTrue(TEXT("Water Body River is not found in Create Menu"), RiverBodyFound);
+
+	TestTrue("Water Body " + WaterBodyName + " is not found in Create Menu", WaterBodyFound);
 
 	return true;
 }
