Water: cleanup invalid water static mesh components before tryign to reuse them.

diff --git a/Engine/Plugins/Experimental/Water/Source/Runtime/Private/WaterBodyComponent.cpp b/Engine/Plugins/Experimental/Water/Source/Runtime/Private/WaterBodyComponent.cpp
--- a/Engine/Plugins/Experimental/Water/Source/Runtime/Private/WaterBodyComponent.cpp
+++ b/Engine/Plugins/Experimental/Water/Source/Runtime/Private/WaterBodyComponent.cpp
@@ -1974,6 +1974,7 @@ void UWaterBodyComponent::UpdateWaterBodyStaticMeshComponents()
 		return;
 	}
 
+	WaterBodyActor->CleanupInvalidStaticMeshComponents();
 	const TArray<TObjectPtr<UWaterBodyStaticMeshComponent>>& WaterBodyStaticMeshComponents = WaterBodyActor->GetWaterBodyStaticMeshComponents();
 
 	const bool bShouldUseStaticMesh = StaticMeshSettings.bEnableWaterBodyStaticMesh && (GetWaterMeshOverride() == nullptr) && (GetWaterBodyType() != EWaterBodyType::Transition);
