Water: fixed persistent maps which had TWeakObjPtr keys.

diff --git a/Engine/Plugins/Experimental/Water/Source/Editor/Private/WaterLandscapeBrush.cpp b/Engine/Plugins/Experimental/Water/Source/Editor/Private/WaterLandscapeBrush.cpp
--- a/Engine/Plugins/Experimental/Water/Source/Editor/Private/WaterLandscapeBrush.cpp
+++ b/Engine/Plugins/Experimental/Water/Source/Editor/Private/WaterLandscapeBrush.cpp
@@ -54,7 +54,7 @@ void AWaterLandscapeBrush::AddActorInternal(AActor* Actor, const UWorld* ThisWor
 
 		if (InCache)
 		{
-			Cache.Add(TWeakObjectPtr<AActor>(Actor), InCache);
+			Cache.Add(TSoftObjectPtr<AActor>(Actor), InCache);
 		}
 
 		if (bTriggerEvent)
@@ -75,7 +75,7 @@ void AWaterLandscapeBrush::RemoveActorInternal(AActor* Actor)
 	if (Index != INDEX_NONE)
 	{
 		ActorsAffectingLandscape.RemoveAt(Index);
-		Cache.Remove(TWeakObjectPtr<AActor>(Actor));
+		Cache.Remove(TSoftObjectPtr<AActor>(Actor));
 
 		OnActorsAffectingLandscapeChanged();
 
@@ -134,7 +134,7 @@ void AWaterLandscapeBrush::UpdateActors(bool bInTriggerEvents)
 	ClearActors();
 
 	// Backup Cache
-	TMap<TWeakObjectPtr<AActor>, TObjectPtr<UObject>> PreviousCache;
+	TMap<TSoftObjectPtr<AActor>, TObjectPtr<UObject>> PreviousCache;
 	Swap(Cache, PreviousCache);
 
 	if (UWorld* World = GetWorld())
@@ -144,7 +144,7 @@ void AWaterLandscapeBrush::UpdateActors(bool bInTriggerEvents)
 			AActor* Actor = *It;
 			if (IWaterBrushActorInterface* WaterBrushActor = Cast<IWaterBrushActorInterface>(Actor))
 			{
-				const TObjectPtr<UObject>* FoundCache = PreviousCache.Find(TWeakObjectPtr<AActor>(Actor));
+				const TObjectPtr<UObject>* FoundCache = PreviousCache.Find(TSoftObjectPtr<AActor>(Actor));
 				const bool bTriggerEvent = false;
 				const bool bModify = false;
 				AddActorInternal(Actor, World, FoundCache != nullptr ? *FoundCache : nullptr, bTriggerEvent, bModify);
@@ -398,7 +398,7 @@ void AWaterLandscapeBrush::AddReferencedObjects(UObject* InThis, FReferenceColle
 	Super::AddReferencedObjects(This, Collector);
 
 	// TODO [jonathan.bard] : remove : probably not necessary since it's now a uproperty :
-	for (TPair<TWeakObjectPtr<AActor>, TObjectPtr<UObject>>& Pair : This->Cache)
+	for (TPair<TSoftObjectPtr<AActor>, TObjectPtr<UObject>>& Pair : This->Cache)
 	{
 		Collector.AddReferencedObject(Pair.Value);
 	}
@@ -434,14 +434,14 @@ void AWaterLandscapeBrush::SetActorCache(AActor* InActor, UObject* InCache)
 		return;
 	}
 
-	TObjectPtr<UObject>& Value = Cache.FindOrAdd(TWeakObjectPtr<AActor>(InActor));
+	TObjectPtr<UObject>& Value = Cache.FindOrAdd(TSoftObjectPtr<AActor>(InActor));
 	Value = InCache;
 }
 
 
 UObject* AWaterLandscapeBrush::GetActorCache(AActor* InActor, TSubclassOf<UObject> CacheClass) const
 {
-	TObjectPtr<UObject> const* ValuePtr = Cache.Find(TWeakObjectPtr<AActor>(InActor));
+	TObjectPtr<UObject> const* ValuePtr = Cache.Find(TSoftObjectPtr<AActor>(InActor));
 	if (ValuePtr && (*ValuePtr) && (*ValuePtr)->IsA(*CacheClass))
 	{
 		return *ValuePtr;
@@ -451,7 +451,7 @@ UObject* AWaterLandscapeBrush::GetActorCache(AActor* InActor, TSubclassOf<UObjec
 
 void AWaterLandscapeBrush::ClearActorCache(AActor* InActor)
 {
-	Cache.Remove(TWeakObjectPtr<AActor>(InActor));
+	Cache.Remove(TSoftObjectPtr<AActor>(InActor));
 }
 
 void AWaterLandscapeBrush::BlueprintGetRenderTargets_Implementation(UTextureRenderTarget2D* InHeightRenderTarget, UTextureRenderTarget2D*& OutVelocityRenderTarget)
diff --git a/Engine/Plugins/Experimental/Water/Source/Editor/Private/WaterLandscapeBrush.h b/Engine/Plugins/Experimental/Water/Source/Editor/Private/WaterLandscapeBrush.h
--- a/Engine/Plugins/Experimental/Water/Source/Editor/Private/WaterLandscapeBrush.h
+++ b/Engine/Plugins/Experimental/Water/Source/Editor/Private/WaterLandscapeBrush.h
@@ -148,5 +148,5 @@ private:
 #endif // WITH_EDITOR
 
 	UPROPERTY(Transient, DuplicateTransient, VisibleAnywhere, AdvancedDisplay, meta = (Category = "Debug"))
-	TMap<TWeakObjectPtr<AActor>, TObjectPtr<UObject>> Cache;
+	TMap<TSoftObjectPtr<AActor>, TObjectPtr<UObject>> Cache;
 };
