Water: fixed the deprecation code which handles setting OwnedWaterZone pointers when a water body hasn't been saved with a zone pointer yet.

diff --git a/Engine/Plugins/Experimental/Water/Source/Runtime/Private/WaterBodyComponent.cpp b/Engine/Plugins/Experimental/Water/Source/Runtime/Private/WaterBodyComponent.cpp
--- a/Engine/Plugins/Experimental/Water/Source/Runtime/Private/WaterBodyComponent.cpp
+++ b/Engine/Plugins/Experimental/Water/Source/Runtime/Private/WaterBodyComponent.cpp
@@ -370,7 +370,7 @@ void UWaterBodyComponent::UpdateExclusionVolumes()
 	}
 }
 
-void UWaterBodyComponent::UpdateWaterZones()
+void UWaterBodyComponent::UpdateWaterZones(bool bAllowChangesDuringCook /* = false */)
 {
 	if (UWorld* World = GetWorld())
 	{
@@ -382,18 +382,16 @@ void UWaterBodyComponent::UpdateWaterZones()
 		else
 		{
 			// Don't attempt to find a water zone while cooking and just rely on the serialized pointer from the editor.
-			if (IsRunningCookCommandlet())
+			if (IsRunningCookCommandlet() && !bAllowChangesDuringCook)
 			{
 				return;
 			}
 
 			const FBox Bounds3D = Bounds.GetBox();
-			if (const UWaterSubsystem* WaterSubsystem = UWaterSubsystem::GetWaterSubsystem(GetWorld()))
-			{
-				const AActor* ActorOwner = GetTypedOuter<AActor>();
-				const ULevel* PreferredLevel = ActorOwner ? ActorOwner->GetLevel() : nullptr;
-				FoundZone = WaterSubsystem->FindWaterZone(FBox2D(FVector2D(Bounds3D.Min), FVector2D(Bounds3D.Max)), PreferredLevel);
-			}
+
+			const AActor* ActorOwner = GetTypedOuter<AActor>();
+			const ULevel* PreferredLevel = ActorOwner ? ActorOwner->GetLevel() : nullptr;
+			FoundZone = UWaterSubsystem::FindWaterZone(World, FBox2D(FVector2D(Bounds3D.Min), FVector2D(Bounds3D.Max)), PreferredLevel);
 		}
 
 		if (OwningWaterZone != FoundZone)
@@ -1460,7 +1458,8 @@ void UWaterBodyComponent::OnPostRegisterAllComponents()
 	{
 		if (OwningWaterZone.IsNull())
 		{
-			UpdateWaterZones();
+			// This needs to happen during a cook commandlet since the asset wasn't ever saved with a valid pointer in the editor.
+			UpdateWaterZones(/** bAllowChangesDuringCook = */ true);
 		}
 	}
 
diff --git a/Engine/Plugins/Experimental/Water/Source/Runtime/Private/WaterSubsystem.cpp b/Engine/Plugins/Experimental/Water/Source/Runtime/Private/WaterSubsystem.cpp
--- a/Engine/Plugins/Experimental/Water/Source/Runtime/Private/WaterSubsystem.cpp
+++ b/Engine/Plugins/Experimental/Water/Source/Runtime/Private/WaterSubsystem.cpp
@@ -536,10 +536,8 @@ void UWaterSubsystem::MarkWaterZonesInRegionForRebuild(const FBox2D& InUpdateReg
 	}
 }
 
-TSoftObjectPtr<AWaterZone> UWaterSubsystem::FindWaterZone(const FBox2D& Bounds, const TSoftObjectPtr<const ULevel> PreferredLevel) const
+TSoftObjectPtr<AWaterZone> UWaterSubsystem::FindWaterZone(const UWorld* World, const FBox2D& Bounds, const TSoftObjectPtr<const ULevel> PreferredLevel)
 {
-
-	const UWorld* World = GetWorld();
 	if (!World)
 	{
 		return {};
@@ -620,6 +618,11 @@ TSoftObjectPtr<AWaterZone> UWaterSubsystem::FindWaterZone(const FBox2D& Bounds,
 	return Algo::MaxElementBy(ViableZones, [](const TPair<TSoftObjectPtr<AWaterZone>, int32>& A) { return A.Value; })->Key;
 }
 
+TSoftObjectPtr<AWaterZone> UWaterSubsystem::FindWaterZone(const FBox2D& Bounds, const TSoftObjectPtr<const ULevel> PreferredLevel) const
+{
+	return FindWaterZone(GetWorld(), Bounds, PreferredLevel);
+}
+
 void UWaterSubsystem::NotifyWaterScalabilityChangedInternal(IConsoleVariable* CVar)
 {
 	OnWaterScalabilityChanged.Broadcast();
diff --git a/Engine/Plugins/Experimental/Water/Source/Runtime/Public/WaterBodyComponent.h b/Engine/Plugins/Experimental/Water/Source/Runtime/Public/WaterBodyComponent.h
--- a/Engine/Plugins/Experimental/Water/Source/Runtime/Public/WaterBodyComponent.h
+++ b/Engine/Plugins/Experimental/Water/Source/Runtime/Public/WaterBodyComponent.h
@@ -389,8 +389,12 @@ public:
 	UFUNCTION(BlueprintCallable, Category=Water)
 	void SetWaterZoneOverride(const TSoftObjectPtr<AWaterZone>& InWaterZoneOverride);
 
-	/** Registers or this water body with corresponding overlapping water zones and unregisters it from any old zones if they are no longer overlapping. */
-	void UpdateWaterZones();
+	/** 
+	 * Registers or this water body with corresponding overlapping water zones and unregisters it from any old zones if they are no longer overlapping.
+	 *
+	 * @param bAllowChangesDuringCook When disabled, this function will not make any changes during cook and just trust that the serialized pointer is correct. 
+	 */
+	void UpdateWaterZones(bool bAllowChangesDuringCook = false);
 
 	/** Set the navigation area class */
 	void SetNavAreaClass(TSubclassOf<UNavAreaBase> NewWaterNavAreaClass) { WaterNavAreaClass = NewWaterNavAreaClass; }
diff --git a/Engine/Plugins/Experimental/Water/Source/Runtime/Public/WaterSubsystem.h b/Engine/Plugins/Experimental/Water/Source/Runtime/Public/WaterSubsystem.h
--- a/Engine/Plugins/Experimental/Water/Source/Runtime/Public/WaterSubsystem.h
+++ b/Engine/Plugins/Experimental/Water/Source/Runtime/Public/WaterSubsystem.h
@@ -177,6 +177,7 @@ public:
 	void MarkWaterZonesInRegionForRebuild(const FBox2D& InUpdateRegion, EWaterZoneRebuildFlags InRebuildFlags);
 
 	/** Returns the water with the highest priority within the bounds provided. */
+	static TSoftObjectPtr<AWaterZone> FindWaterZone(const UWorld* World, const FBox2D& Bounds, const TSoftObjectPtr<const ULevel> PreferredLevel = {});
 	TSoftObjectPtr<AWaterZone> FindWaterZone(const FBox2D& Bounds, const TSoftObjectPtr<const ULevel> PreferredLevel = {}) const;
 
 #if WITH_EDITOR
