Add CollisionChannelForWaterContacts setting to water plugin, add new collision preset to DefaultEngine for water.

diff --git a/Engine/Plugins/Experimental/Water/Source/Runtime/Public/WaterRuntimeSettings.h b/Engine/Plugins/Experimental/Water/Source/Runtime/Public/WaterRuntimeSettings.h
--- a/Engine/Plugins/Experimental/Water/Source/Runtime/Public/WaterRuntimeSettings.h
+++ b/Engine/Plugins/Experimental/Water/Source/Runtime/Public/WaterRuntimeSettings.h
@@ -54,6 +54,10 @@ public:
 	UPROPERTY(EditAnywhere, config, Category = Collision)
 	TEnumAsByte<ECollisionChannel> CollisionChannelForWaterTraces;
 
+	/** Collision channel to use for water ObjectTypes */
+	UPROPERTY(EditAnywhere, config, Category = Collision)
+	TEnumAsByte<ECollisionChannel> CollisionChannelForWaterObjects;
+
 	/** Material Parameter Collection for everything water-related */
 	UPROPERTY(EditAnywhere, config, Category = Rendering)
 	TSoftObjectPtr<UMaterialParameterCollection> MaterialParameterCollection;
