Water: prevent OwningWaterZone cached pointer to be exported on copy/paste.

diff --git a/Engine/Plugins/Experimental/Water/Source/Runtime/Public/WaterBodyComponent.h b/Engine/Plugins/Experimental/Water/Source/Runtime/Public/WaterBodyComponent.h
--- a/Engine/Plugins/Experimental/Water/Source/Runtime/Public/WaterBodyComponent.h
+++ b/Engine/Plugins/Experimental/Water/Source/Runtime/Public/WaterBodyComponent.h
@@ -663,7 +663,7 @@ protected:
 	UPROPERTY(Transient)
 	mutable TWeakObjectPtr<ALandscapeProxy> Landscape;
 
-	UPROPERTY(Category = Water, VisibleInstanceOnly, AdvancedDisplay)
+	UPROPERTY(Category = Water, VisibleInstanceOnly, AdvancedDisplay, TextExportTransient)
 	TSoftObjectPtr<AWaterZone> OwningWaterZone;
 
 	UPROPERTY(Category = Water, EditAnywhere, BlueprintReadOnly, AdvancedDisplay)
