Water: Water ocean components may not cached their bounds by the time we need to know the owning water zone. Instead, issue an explicit CalcBounds when we try to update the water zone pointer. Additionally, fixed a small bug in how the ocean bounds were calculated and ensured that ocean components fill the owning zone when they are spawned in.

diff --git a/Engine/Plugins/Experimental/Water/Source/Editor/Private/WaterBodyActorFactory.cpp b/Engine/Plugins/Experimental/Water/Source/Editor/Private/WaterBodyActorFactory.cpp
--- a/Engine/Plugins/Experimental/Water/Source/Editor/Private/WaterBodyActorFactory.cpp
+++ b/Engine/Plugins/Experimental/Water/Source/Editor/Private/WaterBodyActorFactory.cpp
@@ -146,6 +146,7 @@ void UWaterBodyOceanActorFactory::PostSpawnActor(UObject* Asset, AActor* NewActo
 		{
 			const double ExistingCollisionHeight = OceanComponent->GetCollisionExtents().Z;
 			OceanComponent->SetCollisionExtents(FVector(OwningWaterZone->GetZoneExtent() / 2.0, ExistingCollisionHeight));
+			OceanComponent->FillWaterZoneWithOcean();
 		}
 	}
 }
diff --git a/Engine/Plugins/Experimental/Water/Source/Runtime/Private/WaterBodyComponent.cpp b/Engine/Plugins/Experimental/Water/Source/Runtime/Private/WaterBodyComponent.cpp
--- a/Engine/Plugins/Experimental/Water/Source/Runtime/Private/WaterBodyComponent.cpp
+++ b/Engine/Plugins/Experimental/Water/Source/Runtime/Private/WaterBodyComponent.cpp
@@ -387,7 +387,7 @@ void UWaterBodyComponent::UpdateWaterZones(bool bAllowChangesDuringCook /* = fal
 				return;
 			}
 
-			const FBox Bounds3D = Bounds.GetBox();
+			const FBox Bounds3D = CalcBounds(GetComponentToWorld()).GetBox();
 
 			const AActor* ActorOwner = GetTypedOuter<AActor>();
 			const ULevel* PreferredLevel = ActorOwner ? ActorOwner->GetLevel() : nullptr;
diff --git a/Engine/Plugins/Experimental/Water/Source/Runtime/Private/WaterBodyOceanComponent.cpp b/Engine/Plugins/Experimental/Water/Source/Runtime/Private/WaterBodyOceanComponent.cpp
--- a/Engine/Plugins/Experimental/Water/Source/Runtime/Private/WaterBodyOceanComponent.cpp
+++ b/Engine/Plugins/Experimental/Water/Source/Runtime/Private/WaterBodyOceanComponent.cpp
@@ -492,15 +492,13 @@ void UWaterBodyOceanComponent::OnPostRegisterAllComponents()
 	
 FBoxSphereBounds UWaterBodyOceanComponent::CalcBounds(const FTransform& LocalToWorld) const
 {
-	if (AWaterZone* WaterZone = GetWaterZone())
-	{
-		const FVector Min(FVector(-OceanExtents / 2.0, -1.0 * GetChannelDepth()));
-		const FVector Max(FVector(OceanExtents / 2.0, 0.0));
+	const FVector Min(FVector(-OceanExtents / 2.0, -1.0 * GetChannelDepth()));
+	const FVector Max(FVector(OceanExtents / 2.0, 0.0));
 
-		// translate the bounds so they are centered on the SavedZoneLocation.
-		return FBoxSphereBounds(FBox(Min, Max)).TransformBy(LocalToWorld).TransformBy(FTransform(FVector(SavedZoneLocation, 0.)));
-	}
-	return FBoxSphereBounds(EForceInit::ForceInit).TransformBy(LocalToWorld);
+	// translate the bounds so they are centered on the SavedZoneLocation.
+	FVector RelativeTranslation(FVector(SavedZoneLocation, 0.) - GetComponentLocation());
+	RelativeTranslation.Z = 0.;
+	return FBoxSphereBounds(FBox(Min, Max)).TransformBy(FTransform(RelativeTranslation)).TransformBy(LocalToWorld);
 }
 
 #if WITH_EDITOR
