Water: Fixed oceans with non-identity scale.

diff --git a/Engine/Plugins/Experimental/Water/Source/Runtime/Private/WaterBodyOceanComponent.cpp b/Engine/Plugins/Experimental/Water/Source/Runtime/Private/WaterBodyOceanComponent.cpp
--- a/Engine/Plugins/Experimental/Water/Source/Runtime/Private/WaterBodyOceanComponent.cpp
+++ b/Engine/Plugins/Experimental/Water/Source/Runtime/Private/WaterBodyOceanComponent.cpp
@@ -178,13 +178,15 @@ static void AddAABBQuadToDynamicMesh(UE::Geometry::FDynamicMesh3& EditMesh, FVec
 		NormalOverlay->AppendElement(FVector3f(0., 0., 1.));
 	}
 
-	const int TriOne = EditMesh.AppendTriangle(VertexAIndex, VertexDIndex, VertexBIndex);
-	const int TriTwo = EditMesh.AppendTriangle(VertexBIndex, VertexDIndex, VertexCIndex);
-
-	check(TriOne != INDEX_NONE && TriTwo != INDEX_NONE);
-	
-	ColorOverlay->SetTriangle(TriOne, FIndex3i(VertexAIndex, VertexDIndex, VertexBIndex));
-	ColorOverlay->SetTriangle(TriTwo, FIndex3i(VertexBIndex, VertexDIndex, VertexCIndex));
+	// Only add triangles for this quad if the min/max are actually the min/max. This inversion could happen if
+	// the island spline extends outside the ocean extent thereby flipping the corner positions around.
+	if (Min.X < Max.X && Min.Y < Max.Y)
+	{
+		const int TriOne = EditMesh.AppendTriangle(VertexAIndex, VertexDIndex, VertexBIndex);
+		ColorOverlay->SetTriangle(TriOne, FIndex3i(VertexAIndex, VertexDIndex, VertexBIndex));
+		const int TriTwo = EditMesh.AppendTriangle(VertexBIndex, VertexDIndex, VertexCIndex);
+		ColorOverlay->SetTriangle(TriTwo, FIndex3i(VertexBIndex, VertexDIndex, VertexCIndex));
+	}
 
 	InOutVertices = {VertexAIndex, VertexBIndex, VertexCIndex, VertexDIndex};
 }
@@ -223,8 +225,14 @@ bool UWaterBodyOceanComponent::GenerateWaterBodyMesh(UE::Geometry::FDynamicMesh3
 	// Expand the island slightly so we aren't intersecting with the spline
 	IslandBounds.Expand(1);
 
-	// #todo: account for scale
-	const FAxisAlignedBox2d OceanBounds = FAxisAlignedBox2d(FVector2d(GetComponentTransform().InverseTransformPositionNoScale(WaterZone->GetRootComponent()->GetRelativeLocation())), OceanExtents.X / 2.0);
+	const FTransform& ComponentTransform = GetComponentTransform();
+	FVector RelativeLocationToZone = ComponentTransform.InverseTransformPosition(FVector(SavedZoneLocation, 0.));
+	RelativeLocationToZone.Z = ComponentTransform.GetLocation().Z;
+
+	const FVector2D OceanExtentScaled = OceanExtents / FVector2D(GetComponentScale());
+
+	const FBox OceanBounds3d = CalcBounds(FTransform::Identity).GetBox();
+	const FAxisAlignedBox2d OceanBounds = FAxisAlignedBox2d(FVector2d(OceanBounds3d.Min), FVector2d(OceanBounds3d.Max));
 	FPolygon2d IslandBoundingPolygon = FPolygon2d::MakeRectangle(IslandBounds.Center(), IslandBounds.Extents().X * 2., IslandBounds.Extents().Y * 2.);
 
 	FConstrainedDelaunay2d Triangulation;
@@ -309,7 +317,7 @@ bool UWaterBodyOceanComponent::GenerateWaterBodyMesh(UE::Geometry::FDynamicMesh3
 	// Each quad has four vertices laid out like so
 	// D --- C
 	// |  \  |
-	// A --- Bo
+	// A --- B
 	{
 		TArray<int32, TInlineAllocator<4>> Vertices({ INDEX_NONE, INDEX_NONE, INDEX_NONE, INDEX_NONE });
 
@@ -492,13 +500,13 @@ void UWaterBodyOceanComponent::OnPostRegisterAllComponents()
 	
 FBoxSphereBounds UWaterBodyOceanComponent::CalcBounds(const FTransform& LocalToWorld) const
 {
-	const FVector Min(FVector(-OceanExtents / 2.0, -1.0 * GetChannelDepth()));
-	const FVector Max(FVector(OceanExtents / 2.0, 0.0));
+	const FTransform& ComponentTransform = GetComponentTransform();
+	FVector RelativeLocationToZone = ComponentTransform.InverseTransformPosition(FVector(SavedZoneLocation, 0.));
+	RelativeLocationToZone.Z = ComponentTransform.GetLocation().Z;
+
+	const FVector2D OceanExtentScaled = (OceanExtents / FVector2D(GetComponentScale())) / 2.;
 
-	// translate the bounds so they are centered on the SavedZoneLocation.
-	FVector RelativeTranslation(FVector(SavedZoneLocation, 0.) - GetComponentLocation());
-	RelativeTranslation.Z = 0.;
-	return FBoxSphereBounds(FBox(Min, Max)).TransformBy(FTransform(RelativeTranslation)).TransformBy(LocalToWorld);
+	return FBoxSphereBounds(RelativeLocationToZone, FVector(OceanExtentScaled.X, OceanExtentScaled.Y, GetChannelDepth()), FMath::Max(OceanExtentScaled.X, OceanExtentScaled.Y)).TransformBy(LocalToWorld);
 }
 
 #if WITH_EDITOR
