Fixed double Z translation of ocean component bounds.

diff --git a/Engine/Plugins/Experimental/Water/Source/Runtime/Private/WaterBodyOceanComponent.cpp b/Engine/Plugins/Experimental/Water/Source/Runtime/Private/WaterBodyOceanComponent.cpp
--- a/Engine/Plugins/Experimental/Water/Source/Runtime/Private/WaterBodyOceanComponent.cpp
+++ b/Engine/Plugins/Experimental/Water/Source/Runtime/Private/WaterBodyOceanComponent.cpp
@@ -199,18 +199,12 @@ bool UWaterBodyOceanComponent::GenerateWaterBodyMesh(UE::Geometry::FDynamicMesh3
 
 	const UWaterSplineComponent* SplineComp = GetWaterSpline();
 	const FVector OceanLocation = GetComponentLocation();
-	const AWaterZone* WaterZone = GetWaterZone();
 	
 	if (SplineComp->GetNumberOfSplineSegments() < 3)
 	{
 		return false;
 	}
 
-	if (WaterZone == nullptr)
-	{
-		return false;
-	}
-
 	FPolygon2d Island;
 	TArray<FVector> PolyLineVertices;
 	SplineComp->ConvertSplineToPolyLine(ESplineCoordinateSpace::Local, FMath::Square(10.f), PolyLineVertices);
@@ -502,7 +496,7 @@ FBoxSphereBounds UWaterBodyOceanComponent::CalcBounds(const FTransform& LocalToW
 {
 	const FTransform& ComponentTransform = GetComponentTransform();
 	FVector RelativeLocationToZone = ComponentTransform.InverseTransformPosition(FVector(SavedZoneLocation, 0.));
-	RelativeLocationToZone.Z = ComponentTransform.GetLocation().Z;
+	RelativeLocationToZone.Z = 0;
 
 	const FVector2D OceanExtentScaled = (OceanExtents / FVector2D(GetComponentScale())) / 2.;
 
