#UE Clear waterbody references before updating them

diff --git a/Engine/Plugins/Experimental/Water/Source/Runtime/Private/WaterBodyComponent.cpp b/Engine/Plugins/Experimental/Water/Source/Runtime/Private/WaterBodyComponent.cpp
--- a/Engine/Plugins/Experimental/Water/Source/Runtime/Private/WaterBodyComponent.cpp
+++ b/Engine/Plugins/Experimental/Water/Source/Runtime/Private/WaterBodyComponent.cpp
@@ -340,6 +340,7 @@ void UWaterBodyComponent::UpdateIslands()
 #if WITH_EDITOR
 	if (GetWorld())
 	{
+		WaterBodyIslands.Empty();
 		for (AWaterBodyIsland* Island : TActorRange<AWaterBodyIsland>(GetWorld()))
 		{
 			Island->UpdateOverlappingWaterBodyComponents();
