Water: Refactored the water data buffer to support additional data such as the Water Zone data.

diff --git a/Engine/Plugins/Experimental/Water/Shaders/Private/GerstnerWaveFunctions.ush b/Engine/Plugins/Experimental/Water/Shaders/Private/GerstnerWaveFunctions.ush
--- a/Engine/Plugins/Experimental/Water/Shaders/Private/GerstnerWaveFunctions.ush
+++ b/Engine/Plugins/Experimental/Water/Shaders/Private/GerstnerWaveFunctions.ush
@@ -6,18 +6,11 @@
 
 #pragma once
 
+#include "WaterDataFunctions.ush"
+
 #define Gravity 980
 #define SOLVE_NORMAL_Z 1
 #define SteepnessThreshold 50
-#define PER_WAVE_DATA_SIZE 2
-
-struct FWaveParams
-{
-	float2 Direction;
-	float Wavelength;
-	float Amplitude;
-	float Steepness;
-};
 
 struct WaveOutput
 {
@@ -283,52 +276,6 @@ float3 FinalizeNormalNew(float3 inNormal)
 	return normalize(float3(inNormal.xy, 1.0f - inNormal.z));
 }
 
-/** Struct containing the decoded water body header data  */	
-struct FWaterBodyHeader
-{
-	// Index into the actual wave data buffer
-	int DataIndex;
-
-	// Number of waves to be read from DataIndex
-	int NumWaves;
-	
-	// Target Wave Mask Depth
-	float TargetWaveMaskDepth;
-};
-
-/** Function to decode the raw data from the water indirection buffer */
-FWaterBodyHeader DecodeWaterBodyHeader(float4 InDataToDecode)
-{
-	FWaterBodyHeader OutWaterBodyHeader = (FWaterBodyHeader)0;
-	OutWaterBodyHeader.DataIndex = InDataToDecode.x; 
-	OutWaterBodyHeader.NumWaves = InDataToDecode.y; 
-	OutWaterBodyHeader.TargetWaveMaskDepth = InDataToDecode.z;
-	
-	return OutWaterBodyHeader;
-}
-
-/** Function to decode the data for an individual gerstner wave from the water data buffer */
-FWaveParams DecodeWaveParams(float4 InDataToDecode0, float4 InDataToDecode1)
-{
-	FWaveParams OutWaveParams = (FWaveParams)0;
-	OutWaveParams.Direction = InDataToDecode0.xy;
-	OutWaveParams.Wavelength = InDataToDecode0.z;
-	OutWaveParams.Amplitude = InDataToDecode0.w;
-	OutWaveParams.Steepness = InDataToDecode1.x;
-
-	return OutWaveParams;
-}
-
-FWaveParams GetWaveDataNew(int InWaveIndex, FWaterBodyHeader InWaterBodyHeader)
-{
-	const int AbsoluteWaveDataIndex = InWaterBodyHeader.DataIndex + (InWaveIndex * PER_WAVE_DATA_SIZE);
-
-	float4 Data0 = View.WaterData[AbsoluteWaveDataIndex];
-	float4 Data1 = View.WaterData[AbsoluteWaveDataIndex + 1];
-		
-	return DecodeWaveParams(Data0, Data1);
-}
-
 FWaveParams GetWaveDataByWaveIndexNew(int InWaterBodyIndex, int InWaveIndex)
 {
 	const FWaterBodyHeader WaterBodyHeader = DecodeWaterBodyHeader(View.WaterIndirection[InWaterBodyIndex]);
@@ -406,8 +353,9 @@ WaveOutput GetAllGerstnerWavesNew(int InWaterBodyIndex, float InTime, FLWCVector
 	WaveOutput CurrentWave = (WaveOutput)0;
 		
 	const FWaterBodyHeader WaterBodyHeader = DecodeWaterBodyHeader(View.WaterIndirection[InWaterBodyIndex]);
+	const int NumWaves = GetWaterBodyData(InWaterBodyIndex).NumWaves;
 
-	for (int i = 0; i < WaterBodyHeader.NumWaves; i++)
+	for (int i = 0; i < NumWaves; i++)
 	{
 		CurrentWave = GetSingleGerstnerWaveNew(i, WaterBodyHeader, InTime, InWorldPos);
 		OutWaves = AddWavesNew(OutWaves, CurrentWave);
@@ -424,8 +372,9 @@ WaveOutput GetRangeGerstnerWavesNew(int InMinWaveIndex, int InMaxWaveIndex, int
 	WaveOutput CurrentWave = (WaveOutput)0;
 		
 	const FWaterBodyHeader WaterBodyHeader = DecodeWaterBodyHeader(View.WaterIndirection[InWaterBodyIndex]);
+	const int NumWaves = GetWaterBodyData(InWaterBodyIndex).NumWaves;
 
-	const int MaxWaveRange = min(InMaxWaveIndex + 1, WaterBodyHeader.NumWaves);
+	const int MaxWaveRange = min(InMaxWaveIndex + 1, NumWaves);
 	const int MinWaveRange = max(0, InMinWaveIndex);
 
 	for (int i = MinWaveRange; i < MaxWaveRange; i++)
@@ -443,11 +392,11 @@ WaveOutput GetRangeGerstnerWavesNew(int InMinWaveIndex, int InMaxWaveIndex, int
 /** Returns the wave attenuation factor according to a given water body and a given water depth. Should match the CPU version (GetWaveAttenuationFactor) */
 float ComputeWaveDepthAttenuationFactorNew(int InWaterBodyIndex, float InWaterDepth)
 {
-	const FWaterBodyHeader WaterBodyHeader = DecodeWaterBodyHeader(View.WaterIndirection[InWaterBodyIndex]);
+	const FWaterBodyData WaterBodyData = GetWaterBodyData(InWaterBodyIndex);
 
 	// Completely erase waves when there are no waves at all (in case waves are computed on a water body with no waves, like a swimming pool) :
-	const float TotalErasureFactor = saturate((float) WaterBodyHeader.NumWaves);
-	const float StrengthCoefficient = exp(-max(InWaterDepth, 0.0f) / (WaterBodyHeader.TargetWaveMaskDepth / 2.0f));
+	const float TotalErasureFactor = saturate((float)WaterBodyData.NumWaves);
+	const float StrengthCoefficient = exp(-max(InWaterDepth, 0.0f) / (WaterBodyData.TargetWaveMaskDepth / 2.0f));
 	return saturate(1.0f - StrengthCoefficient) * TotalErasureFactor;
 }
 
diff --git a/Engine/Plugins/Experimental/Water/Shaders/Private/WaterDataFunctions.ush b/Engine/Plugins/Experimental/Water/Shaders/Private/WaterDataFunctions.ush
new file mode 100644
--- /dev/null
+++ b/Engine/Plugins/Experimental/Water/Shaders/Private/WaterDataFunctions.ush
@@ -0,0 +1,176 @@
+#pragma once
+
+#define PER_WAVE_DATA_SIZE 2
+#define PER_WATERZONE_DATA_SIZE 2
+#define PER_WATERBODY_DATA_SIZE 1
+
+/** Struct containing the decoded water body header data  */	
+struct FWaterBodyHeader
+{
+	// Index of the water zone which owns this body. Points to the position of the water zone in the zone sub-buffer.
+	int WaterZoneIndex;
+
+	// Index into the wave data sub-buffer.
+	int WaveDataIndex;
+
+	// Index into the water body data sub-buffer
+	int WaterBodyDataIndex;
+};
+
+/** Struct containing the decoded per-water zone data. */
+struct FWaterZoneData
+{
+	float2 Location;
+	float2 Extent;
+	float2 HeightExtent;
+	float GroundZMin;
+};
+
+/** Struct containing the decoded per-water body data. */
+struct FWaterBodyData
+{
+	// Number of waves to be read from the wave data buffer
+	int NumWaves;
+	
+	// Target Wave Mask Depth
+	float TargetWaveMaskDepth;
+};
+
+struct FWaveParams
+{
+	float2 Direction;
+	float Wavelength;
+	float Amplitude;
+	float Steepness;
+};
+
+int GetWaterZoneDataOffset()
+{
+	return (int)View.WaterData[0].x;
+}
+
+int GetWaveDataOffset()
+{
+	return (int)View.WaterData[0].y;
+}
+
+int GetWaterBodyDataOffset()
+{
+	return (int)View.WaterData[0].z;
+}
+
+/** Function to decode the raw data from the water indirection buffer */
+FWaterBodyHeader DecodeWaterBodyHeader(float4 InDataToDecode)
+{
+	FWaterBodyHeader OutWaterBodyHeader = (FWaterBodyHeader)0;
+	OutWaterBodyHeader.WaterZoneIndex = InDataToDecode.x;
+	OutWaterBodyHeader.WaveDataIndex = InDataToDecode.y;
+	OutWaterBodyHeader.WaterBodyDataIndex = InDataToDecode.z;
+	
+	return OutWaterBodyHeader;
+}
+
+/** Function to decode the data for an individual gerstner wave from the water data buffer */
+FWaveParams DecodeWaveParams(float4 InDataToDecode0, float4 InDataToDecode1)
+{
+	FWaveParams OutWaveParams = (FWaveParams)0;
+	OutWaveParams.Direction = InDataToDecode0.xy;
+	OutWaveParams.Wavelength = InDataToDecode0.z;
+	OutWaveParams.Amplitude = InDataToDecode0.w;
+	OutWaveParams.Steepness = InDataToDecode1.x;
+
+	return OutWaveParams;
+}
+
+FWaveParams GetWaveDataNew(int InWaveIndex, FWaterBodyHeader InWaterBodyHeader)
+{
+	const int WaveDataOffset = GetWaveDataOffset();
+	const int AbsoluteWaveDataIndex = WaveDataOffset + (InWaterBodyHeader.WaveDataIndex + InWaveIndex) * PER_WAVE_DATA_SIZE;
+
+	float4 Data0 = View.WaterData[AbsoluteWaveDataIndex];
+	float4 Data1 = View.WaterData[AbsoluteWaveDataIndex + 1];
+		
+	return DecodeWaveParams(Data0, Data1);
+}
+
+FWaterBodyData DecodeWaterBodyData(float4 Data)
+{
+	FWaterBodyData OutWaterBodyData = (FWaterBodyData)0;
+	OutWaterBodyData.NumWaves = Data.x;
+	OutWaterBodyData.TargetWaveMaskDepth = Data.y;
+
+	return OutWaterBodyData;
+}
+
+FWaterBodyData GetWaterBodyData(int InWaterBodyIndex)
+{
+	const int WaterBodyDataBufferOffset = GetWaterBodyDataOffset();
+	const FWaterBodyHeader WaterBodyHeader = DecodeWaterBodyHeader(View.WaterIndirection[InWaterBodyIndex]);
+
+	float4 Data = View.WaterData[WaterBodyDataBufferOffset + WaterBodyHeader.WaterBodyDataIndex];
+
+	return DecodeWaterBodyData(Data);
+}
+
+FWaterZoneData DecodeWaterZoneData(float4 Data0, float4 Data1)
+{
+	FWaterZoneData OutWaterZoneData = (FWaterZoneData)0;
+	OutWaterZoneData.Location = Data0.xy;
+	OutWaterZoneData.Extent = Data0.zw;
+	OutWaterZoneData.HeightExtent = Data1.xy;
+	OutWaterZoneData.GroundZMin = Data1.z;
+
+	return OutWaterZoneData;
+}
+
+FWaterZoneData GetWaterZoneData(int InWaterZoneIndex)
+{
+	const int WaterZoneDataBufferOffset = GetWaterZoneDataOffset();
+	const int WaterZoneOffset = WaterZoneDataBufferOffset + (InWaterZoneIndex * PER_WATERZONE_DATA_SIZE);
+
+	float4 Data0 = View.WaterData[WaterZoneOffset];
+	float4 Data1 = View.WaterData[WaterZoneOffset + 1];
+
+	return DecodeWaterZoneData(Data0, Data1);
+}
+
+FWaterZoneData GetWaterZoneDataForWaterBody(int InWaterBodyIndex)
+{
+	const int WaterZoneIndex = DecodeWaterBodyHeader(View.WaterIndirection[InWaterBodyIndex]).WaterZoneIndex;
+
+	return GetWaterZoneData(WaterZoneIndex);
+}
+
+float DecodeWaterInfoGroundHeight(float4 WaterInfoSample, int WaterZoneIndex)
+{
+	const float GroundZNormalized = WaterInfoSample.w;
+	const FWaterZoneData WaterZoneData = GetWaterZoneData(WaterZoneIndex);
+
+	const float GroundZScale = WaterZoneData.HeightExtent.y - WaterZoneData.GroundZMin;
+
+	return GroundZNormalized * GroundZScale + WaterZoneData.GroundZMin;
+}
+
+float DecodeWaterInfoZHeight(float4 WaterInfoSample, int WaterZoneIndex)
+{
+	const float WaterZNormalized = WaterInfoSample.z;
+	const FWaterZoneData WaterZoneData = GetWaterZoneData(WaterZoneIndex);
+
+	const float WaterZScale = WaterZoneData.HeightExtent.y - WaterZoneData.HeightExtent.x;
+	float WaterZ = WaterZNormalized * WaterZScale + WaterZoneData.HeightExtent.x;
+
+	// If below 0, we need to push the height below the height of the terrain
+	if (WaterZNormalized <= 0.0)
+	{
+		WaterZ = DecodeWaterInfoGroundHeight(WaterInfoSample, WaterZoneIndex) - 256.0;
+	}
+
+	return WaterZ;
+}
+
+float2 DecodeWaterInfoVelocity(float4 WaterInfoSample, float MaxVelocity)
+{
+	const float2 NormalizedFlow = WaterInfoSample.xy;
+	return (NormalizedFlow - 0.5) * 2.0 * MaxVelocity;
+}
+
diff --git a/Engine/Plugins/Experimental/Water/Source/Runtime/Private/GerstnerWaterWaveSubsystem.cpp b/Engine/Plugins/Experimental/Water/Source/Runtime/Private/GerstnerWaterWaveSubsystem.cpp
--- a/Engine/Plugins/Experimental/Water/Source/Runtime/Private/GerstnerWaterWaveSubsystem.cpp
+++ b/Engine/Plugins/Experimental/Water/Source/Runtime/Private/GerstnerWaterWaveSubsystem.cpp
@@ -1,7 +1,7 @@
 // Copyright Epic Games, Inc. All Rights Reserved.
 
 #include "GerstnerWaterWaveSubsystem.h"
-#include "GerstnerWaterWaveViewExtension.h"
+#include "WaterViewExtension.h"
 
 #include UE_INLINE_GENERATED_CPP_BY_NAME(GerstnerWaterWaveSubsystem)
 
@@ -19,30 +19,30 @@ void UGerstnerWaterWaveSubsystem::Initialize(FSubsystemCollectionBase& Collectio
 
 void UGerstnerWaterWaveSubsystem::Deinitialize()
 {
-	GerstnerWaterWaveViewExtensions.Empty();
+	WaterViewExtensions.Empty();
 
 	FCoreDelegates::OnBeginFrame.RemoveAll(this);
 
 	Super::Deinitialize();
 }
 
-void UGerstnerWaterWaveSubsystem::Register(FGerstnerWaterWaveViewExtension* InViewExtension)
+void UGerstnerWaterWaveSubsystem::Register(FWaterViewExtension* InViewExtension)
 {
-	check(!GerstnerWaterWaveViewExtensions.Contains(InViewExtension));
-	GerstnerWaterWaveViewExtensions.Add(InViewExtension);
+	check(!WaterViewExtensions.Contains(InViewExtension));
+	WaterViewExtensions.Add(InViewExtension);
 }
 
-void UGerstnerWaterWaveSubsystem::Unregister(FGerstnerWaterWaveViewExtension* InViewExtension)
+void UGerstnerWaterWaveSubsystem::Unregister(FWaterViewExtension* InViewExtension)
 {
-	check(GerstnerWaterWaveViewExtensions.Contains(InViewExtension));
-	GerstnerWaterWaveViewExtensions.Remove(InViewExtension);
+	check(WaterViewExtensions.Contains(InViewExtension));
+	WaterViewExtensions.Remove(InViewExtension);
 }
 
 void UGerstnerWaterWaveSubsystem::GetResourceSizeEx(FResourceSizeEx& CumulativeResourceSize)
 {
 	Super::GetResourceSizeEx(CumulativeResourceSize);
 
-	CumulativeResourceSize.AddDedicatedSystemMemoryBytes(GerstnerWaterWaveViewExtensions.GetAllocatedSize());
+	CumulativeResourceSize.AddDedicatedSystemMemoryBytes(WaterViewExtensions.GetAllocatedSize());
 }
 
 void UGerstnerWaterWaveSubsystem::BeginFrameCallback()
@@ -50,9 +50,9 @@ void UGerstnerWaterWaveSubsystem::BeginFrameCallback()
 	// In case there was a change, all registered view extensions need to update their GPU data : 
 	if (bRebuildGPUData)
 	{
-		for (FGerstnerWaterWaveViewExtension* GerstnerWaterWaveViewExtension : GerstnerWaterWaveViewExtensions)
+		for (FWaterViewExtension* WaterViewExtension : WaterViewExtensions)
 		{
-			GerstnerWaterWaveViewExtension->bRebuildGPUData = true;
+			WaterViewExtension->MarkGPUDataDirty();
 		}
 	}
 	bRebuildGPUData = false;
diff --git a/Engine/Plugins/Experimental/Water/Source/Runtime/Private/GerstnerWaterWaveViewExtension.cpp b/Engine/Plugins/Experimental/Water/Source/Runtime/Private/GerstnerWaterWaveViewExtension.cpp
deleted file mode 100644
diff --git a/Engine/Plugins/Experimental/Water/Source/Runtime/Private/WaterBodyComponent.cpp b/Engine/Plugins/Experimental/Water/Source/Runtime/Private/WaterBodyComponent.cpp
--- a/Engine/Plugins/Experimental/Water/Source/Runtime/Private/WaterBodyComponent.cpp
+++ b/Engine/Plugins/Experimental/Water/Source/Runtime/Private/WaterBodyComponent.cpp
@@ -66,6 +66,7 @@ TAutoConsoleVariable<float> CVarWaterOceanFallbackDepth(
 	ECVF_Default);
 
 const FName UWaterBodyComponent::WaterBodyIndexParamName(TEXT("WaterBodyIndex"));
+const FName UWaterBodyComponent::WaterZoneIndexParamName(TEXT("WaterZoneIndex"));
 const FName UWaterBodyComponent::WaterBodyZOffsetParamName(TEXT("WaterBodyZOffset"));
 const FName UWaterBodyComponent::WaterVelocityAndHeightName(TEXT("WaterVelocityAndHeight"));
 const FName UWaterBodyComponent::GlobalOceanHeightName(TEXT("GlobalOceanHeight"));
@@ -1726,31 +1727,8 @@ bool UWaterBodyComponent::SetDynamicParametersOnMID(UMaterialInstanceDynamic* In
 
 	if (const AWaterZone* WaterZone = GetWaterZone())
 	{
+		InMID->SetScalarParameterValue(WaterZoneIndexParamName, WaterZone->GetWaterZoneIndex());
 		InMID->SetTextureParameterValue(WaterVelocityAndHeightName, WaterZone->WaterInfoTexture);
-
-		const UWaterMeshComponent* WaterMeshComponent = WaterZone->GetWaterMeshComponent();
-		check(WaterMeshComponent);
-
-		// Location is the bottom left of the zone
-		const FVector2D WaterInfoExtent = FVector2D(WaterZone->GetDynamicWaterInfoExtent());
-		const FVector2D WaterInfoLocation = FVector2D(WaterZone->GetDynamicWaterInfoCenter()) - (WaterInfoExtent / 2.f);
-
-		FVector4 WaterArea;
-		WaterArea.X = WaterInfoLocation.X;
-		WaterArea.Y = WaterInfoLocation.Y;
-		WaterArea.Z = WaterInfoExtent.X;
-		WaterArea.W = WaterInfoExtent.Y;
-		InMID->SetDoubleVectorParameterValue(WaterAreaParamName, WaterArea);
-
-		const FVector2f WaterHeightExtents = WaterZone->GetWaterHeightExtents();
-		const float GroundZMin = WaterZone->GetGroundZMin();
-		InMID->SetScalarParameterValue(WaterZMinParamName, WaterHeightExtents.X);
-		InMID->SetScalarParameterValue(WaterZMaxParamName, WaterHeightExtents.Y);
-		InMID->SetScalarParameterValue(GroundZMinParamName, GroundZMin);
-	}
-	else
-	{
-		InMID->SetDoubleVectorParameterValue(WaterAreaParamName, FVector4::Zero());
 	}
 
 	return true;
diff --git a/Engine/Plugins/Experimental/Water/Source/Runtime/Private/WaterBodyManager.cpp b/Engine/Plugins/Experimental/Water/Source/Runtime/Private/WaterBodyManager.cpp
--- a/Engine/Plugins/Experimental/Water/Source/Runtime/Private/WaterBodyManager.cpp
+++ b/Engine/Plugins/Experimental/Water/Source/Runtime/Private/WaterBodyManager.cpp
@@ -3,81 +3,71 @@
 #include "WaterBodyManager.h"
 #include "WaterBodyComponent.h"
 #include "WaterSubsystem.h"
-#include "GerstnerWaterWaveViewExtension.h"
+#include "WaterViewExtension.h"
 
 void FWaterBodyManager::Initialize(UWorld* World)
 {
 	if (World != nullptr)
 	{
-		GerstnerWaterWaveViewExtension = FSceneViewExtensions::NewExtension<FGerstnerWaterWaveViewExtension>(World);
-		GerstnerWaterWaveViewExtension->Initialize();
+		WaterViewExtension = FSceneViewExtensions::NewExtension<FWaterViewExtension>(World);
+		WaterViewExtension->Initialize();
 	}
 }
 
 void FWaterBodyManager::Deinitialize()
 {
-	GerstnerWaterWaveViewExtension->Deinitialize();
-	GerstnerWaterWaveViewExtension.Reset();
+	WaterViewExtension->Deinitialize();
+	WaterViewExtension.Reset();
 }
 
 int32 FWaterBodyManager::AddWaterBodyComponent(UWaterBodyComponent* InWaterBodyComponent)
 {
-	int32 Index = INDEX_NONE;
-	if (UnusedWaterBodyIndices.Num())
-	{
-		Index = UnusedWaterBodyIndices.Pop();
-		check(WaterBodyComponents[Index] == nullptr);
-		WaterBodyComponents[Index] = InWaterBodyComponent;
-	}
-	else
-	{
-		Index = WaterBodyComponents.Add(InWaterBodyComponent);
-	}
-
 	RequestWaveDataRebuild();
-
-	check(Index != INDEX_NONE);
-	return Index;
+	return WaterBodyComponents.Register(InWaterBodyComponent);
 }
 
 void FWaterBodyManager::RemoveWaterBodyComponent(UWaterBodyComponent* InWaterBodyComponent)
 {
-	const int32 WaterBodyIndex = InWaterBodyComponent->GetWaterBodyIndex();
-	check(WaterBodyIndex != INDEX_NONE);
-	UnusedWaterBodyIndices.Add(WaterBodyIndex);
-	WaterBodyComponents[WaterBodyIndex] = nullptr;
-
 	RequestWaveDataRebuild();
+	WaterBodyComponents.Unregister(InWaterBodyComponent, InWaterBodyComponent->GetWaterBodyIndex());
+}
 
-	// Reset all arrays once there are no more waterbodies
-	if (UnusedWaterBodyIndices.Num() == WaterBodyComponents.Num())
+int32 FWaterBodyManager::AddWaterZone(AWaterZone* InWaterZone)
+{
+	RequestGPUDataRebuild();
+	return WaterZones.Register(InWaterZone);
+}
+
+void FWaterBodyManager::RemoveWaterZone(AWaterZone* InWaterZone)
+{
+	RequestGPUDataRebuild();
+	WaterZones.Unregister(InWaterZone, InWaterZone->GetWaterZoneIndex());
+}
+
+void FWaterBodyManager::RequestGPUDataRebuild()
+{
+	if (WaterViewExtension)
 	{
-		UnusedWaterBodyIndices.Empty();
-		WaterBodyComponents.Empty();
+		WaterViewExtension->MarkGPUDataDirty();
 	}
 }
 
 void FWaterBodyManager::RequestWaveDataRebuild()
 {
-	if (GerstnerWaterWaveViewExtension)
-	{
-		GerstnerWaterWaveViewExtension->bRebuildGPUData = true;
-	}
+	RequestGPUDataRebuild();
 
 	// Recompute the maximum of all MaxWaveHeight : 
 	GlobalMaxWaveHeight = 0.0f;
-	for (const UWaterBodyComponent* WaterBodyComponent : WaterBodyComponents)
+	ForEachWaterBodyComponent([this](UWaterBodyComponent* WaterBodyComponent)
 	{
-		if (WaterBodyComponent != nullptr)
-		{
-			GlobalMaxWaveHeight = FMath::Max(GlobalMaxWaveHeight, WaterBodyComponent->GetMaxWaveHeight());
-		}
-	}
+		GlobalMaxWaveHeight = FMath::Max(GlobalMaxWaveHeight, WaterBodyComponent->GetMaxWaveHeight());
+		return true;
+	});
 }
 
 void FWaterBodyManager::ForEachWaterBodyComponent(TFunctionRef<bool(UWaterBodyComponent*)> Pred) const
 {
-	for (UWaterBodyComponent* WaterBodyComponent : WaterBodyComponents)
+	for (UWaterBodyComponent* WaterBodyComponent : WaterBodyComponents.Elements)
 	{
 		if (WaterBodyComponent)
 		{
@@ -96,3 +86,26 @@ void FWaterBodyManager::ForEachWaterBodyComponent(const UWorld* World, TFunction
 		Manager->ForEachWaterBodyComponent(Pred);
 	}
 }
+
+void FWaterBodyManager::ForEachWaterZone(TFunctionRef<bool(AWaterZone*)> Pred) const
+{
+	for (AWaterZone* WaterZone : WaterZones.Elements)
+	{
+		if (WaterZone)
+		{
+			if (!Pred(WaterZone))
+			{
+				return;
+			}
+		}
+	}
+}
+
+void FWaterBodyManager::ForEachWaterZone(const UWorld* World, TFunctionRef<bool(AWaterZone*)> Pred)
+{
+	if (FWaterBodyManager* Manager = UWaterSubsystem::GetWaterBodyManager(World))
+	{
+		Manager->ForEachWaterZone(Pred);
+	}
+}
+
diff --git a/Engine/Plugins/Experimental/Water/Source/Runtime/Private/WaterSubsystem.cpp b/Engine/Plugins/Experimental/Water/Source/Runtime/Private/WaterSubsystem.cpp
--- a/Engine/Plugins/Experimental/Water/Source/Runtime/Private/WaterSubsystem.cpp
+++ b/Engine/Plugins/Experimental/Water/Source/Runtime/Private/WaterSubsystem.cpp
@@ -184,11 +184,11 @@ FWaterBodyManager* UWaterSubsystem::GetWaterBodyManager(const UWorld* InWorld)
 	return nullptr;
 }
 
-TWeakPtr<FWaterViewExtension, ESPMode::ThreadSafe> UWaterSubsystem::GetWaterViewExtension(const UWorld* InWorld)
+FWaterViewExtension* UWaterSubsystem::GetWaterViewExtension(const UWorld* InWorld)
 {
-	if (UWaterSubsystem* Subsystem = GetWaterSubsystem(InWorld))
+	if (FWaterBodyManager* Manager = GetWaterBodyManager(InWorld))
 	{
-		return Subsystem->WaterViewExtension;
+		return Manager->GetWaterViewExtension();
 	}
 	return {};
 }
@@ -253,7 +253,6 @@ void UWaterSubsystem::Initialize(FSubsystemCollectionBase& Collection)
 	check(World != nullptr);
 
 	GetWaterBodyManagerInternal().Initialize(World);
-	WaterViewExtension = FSceneViewExtensions::NewExtension<FWaterViewExtension>(World);
 
 	bUsingSmoothedTime = false;
 	FConsoleVariableDelegate NotifyWaterScalabilityChanged = FConsoleVariableDelegate::CreateUObject(this, &UWaterSubsystem::NotifyWaterScalabilityChangedInternal);
diff --git a/Engine/Plugins/Experimental/Water/Source/Runtime/Private/WaterViewExtension.cpp b/Engine/Plugins/Experimental/Water/Source/Runtime/Private/WaterViewExtension.cpp
--- a/Engine/Plugins/Experimental/Water/Source/Runtime/Private/WaterViewExtension.cpp
+++ b/Engine/Plugins/Experimental/Water/Source/Runtime/Private/WaterViewExtension.cpp
@@ -6,6 +6,9 @@
 #include "EngineUtils.h"
 #include "SceneView.h"
 #include "WaterMeshComponent.h"
+#include "GerstnerWaterWaveSubsystem.h"
+#include "WaterBodyManager.h"
+#include "WaterSubsystem.h"
 
 static TAutoConsoleVariable<bool> CVarLocalTessellationFreeze(
 	TEXT("r.Water.WaterMesh.LocalTessellation.Freeze"),
@@ -25,6 +28,7 @@ static TAutoConsoleVariable<float> CVarLocalTessellationUpdateMargin(
 
 FWaterViewExtension::FWaterViewExtension(const FAutoRegister& AutoReg, UWorld* InWorld)
 	: FWorldSceneViewExtension(AutoReg, InWorld)
+	, WaterGPUData(MakeShared<FWaterGPUResources, ESPMode::ThreadSafe>())
 {
 }
 
@@ -32,10 +36,246 @@ FWaterViewExtension::~FWaterViewExtension()
 {
 }
 
-void FWaterViewExtension::SetupViewFamily(FSceneViewFamily& InViewFamily)
+void FWaterViewExtension::Initialize()
 {
+	// Register the view extension to the Gerstner Wave subsystem so we can rebuild the water gpu data when waves change.
+	if (UGerstnerWaterWaveSubsystem* GerstnerWaterWaveSubsystem = GEngine->GetEngineSubsystem<UGerstnerWaterWaveSubsystem>())
+	{
+		GerstnerWaterWaveSubsystem->Register(this);
+	}
 }
 
+void FWaterViewExtension::Deinitialize()
+{
+	ENQUEUE_RENDER_COMMAND(DeallocateWaterInstanceDataBuffer)
+	(
+		// Copy the shared ptr into a local copy for this lambda, this will increase the ref count and keep it alive on the renderthread until this lambda is executed
+		[WaterGPUData=WaterGPUData](FRHICommandListImmediate& RHICmdList){}
+	);
+
+	if (UGerstnerWaterWaveSubsystem* GerstnerWaterWaveSubsystem = GEngine->GetEngineSubsystem<UGerstnerWaterWaveSubsystem>())
+	{
+		GerstnerWaterWaveSubsystem->Unregister(this);
+	}
+}
+
+void FWaterViewExtension::SetupViewFamily(FSceneViewFamily& InViewFamily)
+{
+	if (bRebuildGPUData)
+	{
+		TRACE_CPUPROFILER_EVENT_SCOPE(Water::RebuildWaterGPUData);
+
+		const UWorld* WorldPtr = GetWorld();
+		check(WorldPtr != nullptr);
+		const FWaterBodyManager* WaterBodyManager = UWaterSubsystem::GetWaterBodyManager(WorldPtr);
+		check(WaterBodyManager);
+
+
+		struct FWaterIndirection
+		{
+			float WaterZoneIndex;
+			float WaveDataIndex;
+			float WaterBodyDataIndex;
+
+			float _Padding; // Unused 
+		};
+		static_assert(sizeof(FWaterIndirection) == 1 * sizeof(FVector4f));
+
+		struct FWaterZoneData
+		{
+			FVector2f Location;
+			FVector2f Extent;
+			FVector2f HeightExtent;
+			float GroundZMin;
+
+			float _Padding; // Unused
+
+			FWaterZoneData(const FVector2f& InLocation, const FVector2f& InExtent, const FVector2f& InHeightExtent, float InGroundZMin)
+				: Location(InLocation), Extent(InExtent), HeightExtent(InHeightExtent), GroundZMin(InGroundZMin) {}
+		};
+		static_assert(sizeof(FWaterZoneData) == 2 * sizeof(FVector4f));
+
+		struct FGerstnerWaveData
+		{
+			FVector2f Direction;
+			float WaveLength;
+			float Amplitude;
+			float Steepness;
+
+			float _Padding[3]; // Unused
+
+			FGerstnerWaveData(const FGerstnerWave& Wave)
+				: Direction(FVector2D(Wave.Direction)), WaveLength(Wave.WaveLength), Amplitude(Wave.Amplitude), Steepness(Wave.Steepness) {}
+		};
+		static_assert(sizeof(FGerstnerWaveData) == 2 * sizeof(FVector4f));
+
+		struct FWaterBodyData
+		{
+			float NumWaves;
+			float TargetWaveMaskDepth;
+			
+			float _Padding[2]; // Unused;
+		};
+		static_assert(sizeof(FWaterBodyData) == sizeof(FVector4f));
+
+		// Water Indirection Buffer layout:
+		// -------------------------------------------------------------------------------
+		// || WaterZoneIndex | WaveDataIndex | WaterBodyDataIndex | *Unused* ||   ...   ||
+		// -------------------------------------------------------------------------------
+		//
+		// Water Data Buffer layout:
+		// -----------------------------------------------------------------------------
+		// ||| WaterZone Data | ..  || GerstnerWaveData | ... || WaterBodyData | ... |||
+		// -----------------------------------------------------------------------------
+		//
+
+		TArray<FWaterZoneData> WaterZoneData;
+
+		{
+			WaterBodyManager->ForEachWaterZone([&WaterZoneData](AWaterZone* WaterZone)
+			{
+				// #todo_water: LWC
+				const FVector2f ZoneLocation = FVector2f(FVector2D(WaterZone->GetDynamicWaterInfoCenter()));
+
+				const FVector2f ZoneExtent = FVector2f(FVector2D(WaterZone->GetDynamicWaterInfoExtent()));
+				const FVector2f WaterHeightExtents = WaterZone->GetWaterHeightExtents();
+				const float GroundZMin = WaterZone->GetGroundZMin();
+
+				WaterZoneData.Emplace(ZoneLocation, ZoneExtent, WaterHeightExtents, GroundZMin);
+
+				return true;
+			});
+		}
+
+
+		TArray<FWaterIndirection> WaterIndirection;
+		TArray<FGerstnerWaveData> WaveData;
+		TArray<FWaterBodyData> WaterBodyData;
+		{
+			const int32 NumWaterBodies =  WaterBodyManager->NumWaterBodies();
+			WaterIndirection.Reserve(NumWaterBodies);
+			WaterBodyData.Reserve(NumWaterBodies);
+
+			TArray<const UGerstnerWaterWaves*> AllGerstnerWaves;
+
+			WaterBodyManager->ForEachWaterBodyComponent([&WaterIndirection, &WaveData, &WaterBodyData, &AllGerstnerWaves](UWaterBodyComponent* WaterBodyComponent)
+			{
+				const int32 WaterZoneIndex = WaterBodyComponent->GetWaterZone() ? WaterBodyComponent->GetWaterZone()->GetWaterZoneIndex() : -1;
+
+				FWaterIndirection& WaterIndirectionEntry = WaterIndirection.AddZeroed_GetRef();
+				WaterIndirectionEntry.WaterZoneIndex = WaterZoneIndex;
+				WaterIndirectionEntry.WaterBodyDataIndex = WaterBodyData.Num();
+
+				FWaterBodyData& WaterBodyDataEntry = WaterBodyData.AddZeroed_GetRef();
+				WaterBodyDataEntry.TargetWaveMaskDepth = WaterBodyComponent->TargetWaveMaskDepth;
+
+				// #todo_water: reuse waves instead of duplicating.
+				if (WaterBodyComponent->HasWaves())
+				{
+					const UWaterWavesBase* WaterWavesBase = WaterBodyComponent->GetWaterWaves();
+					check(WaterWavesBase != nullptr);
+					if (const UGerstnerWaterWaves* GerstnerWaves = Cast<const UGerstnerWaterWaves>(WaterWavesBase->GetWaterWaves()))
+					{
+						int32 WaveDataIndex = AllGerstnerWaves.IndexOfByKey(GerstnerWaves);
+						if (WaveDataIndex == INDEX_NONE)
+						{
+							WaveDataIndex = AllGerstnerWaves.Num();
+							AllGerstnerWaves.Add(GerstnerWaves);
+						}
+
+						const TArray<FGerstnerWave>& Waves = GerstnerWaves->GetGerstnerWaves();
+
+						WaterIndirectionEntry.WaveDataIndex = WaveDataIndex;
+						WaterBodyDataEntry.NumWaves = Waves.Num();
+					}
+				}
+				return true;
+			});
+
+			for (int32 GerstnerWavesIndex = 0; GerstnerWavesIndex < AllGerstnerWaves.Num(); ++GerstnerWavesIndex)
+			{
+				// Some max value
+				constexpr int32 MaxWavesPerGerstnerWaves = 4096;
+
+				const TArray<FGerstnerWave>& Waves = AllGerstnerWaves[GerstnerWavesIndex]->GetGerstnerWaves();
+				
+				// Where the data for this set of waves starts
+				const int32 WaveDataBase = WaveData.Num();
+
+				// Allocate for the waves in this water body
+				const int32 NumWaves = FMath::Min(Waves.Num(), MaxWavesPerGerstnerWaves);
+				WaveData.AddZeroed(NumWaves);
+
+				for (int32 WaveIndex = 0; WaveIndex < NumWaves; WaveIndex++)
+				{
+					const uint32 WavesDataIndex = WaveDataBase + WaveIndex;
+					WaveData[WavesDataIndex] = FGerstnerWaveData(Waves[WaveIndex]);
+				}
+			}
+		}
+
+		TResourceArray<FVector4f> WaterIndirectionBuffer;
+		TResourceArray<FVector4f> WaterDataBuffer;
+
+		// The first element of the WaterDataBuffer contains the offsets to each of the sub-buffers.
+		// X = WaterZoneDataOffset
+		// Y = WaterWaveDataOffset
+		// Z = WaterBodyDataOffset
+		// W = Unused
+		WaterDataBuffer.AddZeroed();
+
+		// Transform the individual arrays into the single buffer:
+		{
+			/** Copy a buffer of arbitrary PoD into a float4 resource array. Returns the starting offset of the source buffer in the dest buffer. */
+			auto AppendDataToFloat4Buffer = []<typename T>(TResourceArray<FVector4f>& Dest, const TArray<T>& Source)
+			{
+				constexpr int32 NumFloat4PerElement = (sizeof(T) / sizeof(FVector4f));
+				const int32 StartOffset = Dest.Num();
+				Dest.AddUninitialized(Source.Num() * NumFloat4PerElement);
+				FMemory::Memcpy(
+					Dest.GetData() + StartOffset,
+					Source.GetData(),
+					Source.Num() * sizeof(T));
+				return StartOffset;
+			};
+
+			AppendDataToFloat4Buffer(WaterIndirectionBuffer, WaterIndirection);
+
+			const int32 ZoneDataOffset = AppendDataToFloat4Buffer(WaterDataBuffer, WaterZoneData);
+			const int32 WaveDataOffset = AppendDataToFloat4Buffer(WaterDataBuffer, WaveData);
+			const int32 WaterBodyDataOffset = AppendDataToFloat4Buffer(WaterDataBuffer, WaterBodyData);
+
+			// Store the offsets to each sub-buffer in the first entry.
+			// If this layout ever changes, corresponding decode functions must be updated in GerstnerWaveFunctions.ush!
+			FVector4f& OffsetData = WaterDataBuffer[0];
+			OffsetData.X = ZoneDataOffset;
+			OffsetData.Y = WaveDataOffset;
+			OffsetData.Z = WaterBodyDataOffset;
+			OffsetData.W = 0.f;
+		}
+
+		if (WaterIndirectionBuffer.Num() == 0)
+		{
+			WaterIndirectionBuffer.AddZeroed();
+		}
+
+		ENQUEUE_RENDER_COMMAND(AllocateWaterInstanceDataBuffer)
+		(
+			[WaterGPUData=WaterGPUData, WaterDataBuffer, WaterIndirectionBuffer](FRHICommandListImmediate& RHICmdList) mutable
+			{
+				FRHIResourceCreateInfo CreateInfoData(TEXT("WaterDataBuffer"), &WaterDataBuffer);
+				WaterGPUData->DataBuffer = RHICmdList.CreateBuffer(WaterDataBuffer.GetResourceDataSize(), BUF_VertexBuffer | BUF_ShaderResource | BUF_Static, sizeof(FVector4f), ERHIAccess::SRVMask, CreateInfoData);
+				WaterGPUData->DataSRV = RHICmdList.CreateShaderResourceView(WaterGPUData->DataBuffer, sizeof(FVector4f), PF_A32B32G32R32F);
+
+				FRHIResourceCreateInfo CreateInfoIndirection(TEXT("WaterIndirectionBuffer"), &WaterIndirectionBuffer);
+				WaterGPUData->IndirectionBuffer = RHICmdList.CreateBuffer(WaterIndirectionBuffer.GetResourceDataSize(), BUF_VertexBuffer | BUF_ShaderResource | BUF_Static, sizeof(FVector4f), ERHIAccess::SRVMask, CreateInfoIndirection);
+				WaterGPUData->IndirectionSRV = RHICmdList.CreateShaderResourceView(WaterGPUData->IndirectionBuffer, sizeof(FVector4f), PF_A32B32G32R32F);
+			}
+		);
+
+		bRebuildGPUData = false;
+	}
+}
 
 void FWaterViewExtension::SetupView(FSceneViewFamily& InViewFamily, FSceneView& InView)
 {
@@ -96,14 +336,10 @@ void FWaterViewExtension::SetupView(FSceneViewFamily& InViewFamily, FSceneView&
 			{
 				WaterMesh->SetDynamicWaterMeshCenter(FVector2D(WaterMeshCenter));
 			}
-		}
 
-		// Update the material instances now that the tessellated mesh bounds have changed.
-		WaterZone->ForEachWaterBodyComponent([](UWaterBodyComponent* WaterBodyComponent)
-		{
-			WaterBodyComponent->UpdateMaterialInstances();
-			return true;
-		});
+			// Mark GPU data dirty since we have a new WaterArea parameter and need to push this to water bodies.
+			MarkGPUDataDirty();
+		}
 
 		const UE::WaterInfo::FRenderingContext& Context(Pair.Value);
 		UE::WaterInfo::UpdateWaterInfoRendering(WorldPtr.Get()->Scene, Context);
@@ -129,8 +365,22 @@ void FWaterViewExtension::PreRenderViewFamily_RenderThread(FRDGBuilder& GraphBui
 {
 }
 
+void FWaterViewExtension::PreRenderView_RenderThread(FRDGBuilder& GraphBuilder, FSceneView& InView)
+{
+	if (WaterGPUData->DataSRV && WaterGPUData->IndirectionSRV)
+	{
+		InView.WaterDataBuffer = WaterGPUData->DataSRV;
+		InView.WaterIndirectionBuffer = WaterGPUData->IndirectionSRV;
+	}
+}
+
 void FWaterViewExtension::MarkWaterInfoTextureForRebuild(const UE::WaterInfo::FRenderingContext& RenderContext)
 {
 	WaterInfoContextsToRender.Emplace(RenderContext.ZoneToRender, RenderContext);
+	MarkGPUDataDirty();
 }
 
+void FWaterViewExtension::MarkGPUDataDirty()
+{
+	bRebuildGPUData = true;
+}
diff --git a/Engine/Plugins/Experimental/Water/Source/Runtime/Private/WaterZoneActor.cpp b/Engine/Plugins/Experimental/Water/Source/Runtime/Private/WaterZoneActor.cpp
--- a/Engine/Plugins/Experimental/Water/Source/Runtime/Private/WaterZoneActor.cpp
+++ b/Engine/Plugins/Experimental/Water/Source/Runtime/Private/WaterZoneActor.cpp
@@ -149,6 +149,32 @@ void AWaterZone::PostLoad()
 #endif // WITH_EDITORONLY_DATA
 }
 
+void AWaterZone::PostRegisterAllComponents()
+{
+	Super::PostRegisterAllComponents();
+
+	// PostRegisterAllComponents is called many times in a row during BP reinstancing, property changes, etc.
+	// Only register to the water body manager if we haven't already (don't have an index yet).
+	FWaterBodyManager* Manager = UWaterSubsystem::GetWaterBodyManager(GetWorld());
+	if (Manager && !IsTemplate() && (WaterZoneIndex == INDEX_NONE))
+	{
+		WaterZoneIndex = Manager->AddWaterZone(this);
+	}
+}
+
+void AWaterZone::PostUnregisterAllComponents()
+{
+	Super::PostUnregisterAllComponents();
+
+	// We must check for the index because UnregisterAllComponents can be called multiple times in a row by PostEditChangeProperty, etc.
+	FWaterBodyManager* Manager = UWaterSubsystem::GetWaterBodyManager(GetWorld());
+	if (Manager && !IsTemplate() && (WaterZoneIndex != INDEX_NONE))
+	{
+		Manager->RemoveWaterZone(this);
+	}
+	WaterZoneIndex = INDEX_NONE;
+}
+
 #if WITH_EDITORONLY_DATA
 void AWaterZone::DeclareConstructClasses(TArray<FTopLevelAssetPath>& OutConstructClasses, const UClass* SpecificSubclass)
 {
@@ -507,7 +533,18 @@ bool AWaterZone::UpdateWaterInfoTexture()
 #endif // WITH_EDITOR
 
 		const ETextureRenderTargetFormat Format = bHalfPrecisionTexture ? ETextureRenderTargetFormat::RTF_RGBA16f : RTF_RGBA32f;
-		WaterInfoTexture = FWaterUtils::GetOrCreateTransientRenderTarget2D(WaterInfoTexture, TEXT("WaterInfoTexture"), RenderTargetResolution, Format);
+		UTextureRenderTarget2D* OldTexture = WaterInfoTexture;
+		WaterInfoTexture = FWaterUtils::GetOrCreateTransientRenderTarget2D(OldTexture, TEXT("WaterInfoTexture"), RenderTargetResolution, Format);
+
+		// The water info texture is different, we need to bind the newly created texture to all registered water bodies
+		if (WaterInfoTexture != OldTexture)
+		{
+			ForEachWaterBodyComponent([WaterInfoTexture = WaterInfoTexture](UWaterBodyComponent* WaterBodyComponent)
+			{
+				WaterBodyComponent->UpdateMaterialInstances();
+				return true;
+			});
+		}
 
 		UE::WaterInfo::FRenderingContext Context;
 		Context.ZoneToRender = this;
@@ -516,9 +553,9 @@ bool AWaterZone::UpdateWaterInfoTexture()
 		Context.CaptureZ = FMath::Max(WaterZMax, GroundZMax) + CaptureZOffset;
 		Context.TextureRenderTarget = WaterInfoTexture;
 
-		if (TWeakPtr<FWaterViewExtension> WaterViewExtension = UWaterSubsystem::GetWaterViewExtension(World); WaterViewExtension.IsValid())
+		if (FWaterViewExtension* WaterViewExtension = UWaterSubsystem::GetWaterViewExtension(World))
 		{
-			WaterViewExtension.Pin()->MarkWaterInfoTextureForRebuild(Context);
+			WaterViewExtension->MarkWaterInfoTextureForRebuild(Context);
 		}
 
 		UE_LOG(LogWater, Verbose, TEXT("Queued Water Info texture update"));
diff --git a/Engine/Plugins/Experimental/Water/Source/Runtime/Public/GerstnerWaterWaveSubsystem.h b/Engine/Plugins/Experimental/Water/Source/Runtime/Public/GerstnerWaterWaveSubsystem.h
--- a/Engine/Plugins/Experimental/Water/Source/Runtime/Public/GerstnerWaterWaveSubsystem.h
+++ b/Engine/Plugins/Experimental/Water/Source/Runtime/Public/GerstnerWaterWaveSubsystem.h
@@ -7,7 +7,7 @@
 
 struct FResourceSizeEx;
 
-class FGerstnerWaterWaveViewExtension;
+class FWaterViewExtension;
 
 /** UGerstnerWaterWaveSubsystem manages all UGerstnerWaterWaves objects, regardless of which world they belong to (it's a UEngineSubsystem) */
 UCLASS()
@@ -22,8 +22,8 @@ public:
 	virtual void Initialize(FSubsystemCollectionBase& Collection) override;
 	virtual void Deinitialize() override;
 
-	void Register(FGerstnerWaterWaveViewExtension* InViewExtension);
-	void Unregister(FGerstnerWaterWaveViewExtension* InViewExtension);
+	void Register(FWaterViewExtension* ViewExtension);
+	void Unregister(FWaterViewExtension* ViewExtension);
 
 	void RebuildGPUData() { bRebuildGPUData = true; }
 
@@ -35,7 +35,7 @@ private:
 	void BeginFrameCallback();
 
 private:
-	TArray<FGerstnerWaterWaveViewExtension*> GerstnerWaterWaveViewExtensions;
+	TArray<FWaterViewExtension*> WaterViewExtensions;
 	bool bRebuildGPUData = true;
 };
 
diff --git a/Engine/Plugins/Experimental/Water/Source/Runtime/Public/WaterBodyComponent.h b/Engine/Plugins/Experimental/Water/Source/Runtime/Public/WaterBodyComponent.h
--- a/Engine/Plugins/Experimental/Water/Source/Runtime/Public/WaterBodyComponent.h
+++ b/Engine/Plugins/Experimental/Water/Source/Runtime/Public/WaterBodyComponent.h
@@ -549,6 +549,7 @@ public:
 
 	/** Public static constants : */
 	static const FName WaterBodyIndexParamName;
+	static const FName WaterZoneIndexParamName;
 	static const FName WaterBodyZOffsetParamName;
 	static const FName WaterVelocityAndHeightName;
 	static const FName GlobalOceanHeightName;
diff --git a/Engine/Plugins/Experimental/Water/Source/Runtime/Public/WaterBodyManager.h b/Engine/Plugins/Experimental/Water/Source/Runtime/Public/WaterBodyManager.h
--- a/Engine/Plugins/Experimental/Water/Source/Runtime/Public/WaterBodyManager.h
+++ b/Engine/Plugins/Experimental/Water/Source/Runtime/Public/WaterBodyManager.h
@@ -5,7 +5,8 @@
 #include "Templates/SharedPointer.h"
 
 class UWaterBodyComponent;
-class FGerstnerWaterWaveViewExtension;
+class AWaterZone;
+class FWaterViewExtension;
 
 class WATER_API FWaterBodyManager
 {
@@ -23,6 +24,12 @@ public:
 	/** Unregister any water body upon removal to the world */
 	void RemoveWaterBodyComponent(UWaterBodyComponent* InWaterBodyComponent);
 
+	int32 AddWaterZone(AWaterZone* InWaterZone);
+	void RemoveWaterZone(AWaterZone* InWaterZone);
+
+	/** Recomputes water gpu data whenever it changes on one of the managed water types. */
+	void RequestGPUDataRebuild();
+
 	/** Recomputes wave-related data whenever it changes on one of water bodies. */
 	void RequestWaveDataRebuild();
 
@@ -35,16 +42,80 @@ public:
 	/** Execute a predicate function on each valid water body. Predicate should return false for early exit. */
 	static void ForEachWaterBodyComponent (const UWorld* World, TFunctionRef<bool(UWaterBodyComponent*)> Pred);
 
-	bool HasAnyWaterBodies() const { return WaterBodyComponents.Num() > 0; }
+	void ForEachWaterZone(TFunctionRef<bool(AWaterZone*)> Pred) const;
+	static void ForEachWaterZone(const UWorld* World, TFunctionRef<bool(AWaterZone*)> Pred);
+
+	bool HasAnyWaterBodies() const { return WaterBodyComponents.Num > 0; }
+
+	int32 NumWaterBodies() const { return WaterBodyComponents.Num; }
+
+	int32 NumWaterZones() const { return WaterZones.Num; }
+
+	FWaterViewExtension* GetWaterViewExtension() { return WaterViewExtension.Get(); }
 
 private:
-	/** List of components registered to this manager. May contain nullptr indices (indicated by the UnusedWaterBodyIndices array). */
-	TArray<UWaterBodyComponent*> WaterBodyComponents;
-	TArray<int32> UnusedWaterBodyIndices;
+
+	/**
+	* TWaterContainer<T> wraps a TArray<T> to support reusing dead indices while always maintaining stability for existing indices.
+	*
+	* The Elements array may contain nullptr entries.
+	*/
+	template <typename T>
+	class TWaterContainer 
+	{
+	public:
+		int32 Register(T* InElement)
+		{
+			int32 Index = INDEX_NONE;
+			if (UnusedIndices.Num())
+			{
+				Index = UnusedIndices.Pop(/*bAllowShrinking = */false);
+				check(Elements[Index] == nullptr);
+				Elements[Index] = InElement;
+			}
+			else
+			{
+				Index = Elements.Add(InElement);
+			}
+
+			++Num;
+			check(Num <= Elements.Num());
+
+			check(Index != INDEX_NONE);
+			return Index;
+		}
+
+		void Unregister(const T* InElement, int32 OldIndex)
+		{
+			check(OldIndex != INDEX_NONE);
+			UnusedIndices.Add(OldIndex);
+			Elements[OldIndex] = nullptr;
+
+			--Num;
+			check(Num >= 0);
+
+			// Empty all arrays once there are no more elements
+			if (UnusedIndices.Num() == Elements.Num())
+			{
+				UnusedIndices.Empty();
+				Elements.Empty();
+			}
+		}
+
+		TArray<T*> Elements;
+		TArray<int32> UnusedIndices;
+		int32 Num = 0;
+	};
+
+	/** List of components registered to this manager. */
+	TWaterContainer<UWaterBodyComponent> WaterBodyComponents;
+
+	/** List of Water zones registered to this manager. */
+	TWaterContainer<AWaterZone> WaterZones;
 
 	float GlobalMaxWaveHeight = 0.0f;
 
-	TSharedPtr<FGerstnerWaterWaveViewExtension, ESPMode::ThreadSafe> GerstnerWaterWaveViewExtension;
+	TSharedPtr<FWaterViewExtension, ESPMode::ThreadSafe> WaterViewExtension;
 };
 
 #if UE_ENABLE_INCLUDE_ORDER_DEPRECATED_IN_5_2
diff --git a/Engine/Plugins/Experimental/Water/Source/Runtime/Public/WaterSubsystem.h b/Engine/Plugins/Experimental/Water/Source/Runtime/Public/WaterSubsystem.h
--- a/Engine/Plugins/Experimental/Water/Source/Runtime/Public/WaterSubsystem.h
+++ b/Engine/Plugins/Experimental/Water/Source/Runtime/Public/WaterSubsystem.h
@@ -100,16 +100,7 @@ public:
 	static FWaterBodyManager* GetWaterBodyManager(const UWorld* InWorld);
 
 	/** Static helper function to get a weak ptr to the water scene view extension for a given world. */
-	static TWeakPtr<FWaterViewExtension, ESPMode::ThreadSafe> GetWaterViewExtension(const UWorld* InWorld);
-
-	UE_DEPRECATED(5.1, "Please use FWaterBodyManager::ForEachWaterBodyComponent instead.")
-	static void ForEachWaterBodyComponent(const UWorld* World, TFunctionRef<bool(UWaterBodyComponent*)> Predicate) {}
-
-	UE_DEPRECATED(5.1, "This will become a private member. Prefer calling GetWaterBodyManager instead")
-	FWaterBodyManager WaterBodyManager;
-
-	UE_DEPRECATED(5.1, "There may be multiple water zones per level. Prefer calling the GetWaterZone on a per-water body basis or iterating over all zones.")
-	AWaterZone* GetWaterZoneActor(ULevel* InPreferredOuterLevel = nullptr) const { return nullptr; }
+	static FWaterViewExtension* GetWaterViewExtension(const UWorld* InWorld);
 
 	ABuoyancyManager* GetBuoyancyManager() const { return BuoyancyManager; }
 
@@ -255,7 +246,7 @@ private:
 
 	FUnderwaterPostProcessVolume UnderwaterPostProcessVolume;
 
-	TSharedPtr<FWaterViewExtension> WaterViewExtension;
+	FWaterBodyManager WaterBodyManager;
 
 #if WITH_EDITOR
 	FDelegateHandle OnHeightmapStreamedHandle;
diff --git a/Engine/Plugins/Experimental/Water/Source/Runtime/Public/WaterViewExtension.h b/Engine/Plugins/Experimental/Water/Source/Runtime/Public/WaterViewExtension.h
--- a/Engine/Plugins/Experimental/Water/Source/Runtime/Public/WaterViewExtension.h
+++ b/Engine/Plugins/Experimental/Water/Source/Runtime/Public/WaterViewExtension.h
@@ -14,16 +14,21 @@ public:
 	FWaterViewExtension(const FAutoRegister& AutoReg, UWorld* InWorld);
 	~FWaterViewExtension();
 
+	void Initialize();
+	void Deinitialize();
+
 	// FSceneViewExtensionBase implementation : 
 	virtual void BeginRenderViewFamily(FSceneViewFamily& InViewFamily) override {}
 
 	virtual void SetupView(FSceneViewFamily& InViewFamily, FSceneView& InView) override;
 	virtual void SetupViewFamily(FSceneViewFamily& InViewFamily) override;
 	virtual void PreRenderViewFamily_RenderThread(FRDGBuilder& GraphBuilder, FSceneViewFamily& InViewFamily) override;
+	virtual void PreRenderView_RenderThread(FRDGBuilder& GraphBuilder, FSceneView& InView) override;
 	// End FSceneViewExtensionBase implementation
 
 	void MarkWaterInfoTextureForRebuild(const UE::WaterInfo::FRenderingContext& RenderContext);
 
+	void MarkGPUDataDirty();
 private:
 	/** Queued Water Info rendering contexts to submit for rendering on the next SetupView call */
 	TMap<AWaterZone*, UE::WaterInfo::FRenderingContext> WaterInfoContextsToRender;
@@ -33,6 +38,19 @@ private:
 	 * When the view location crosses the bounds, submit a new WaterInfo update to reflect the new active area
 	 */
 	TMap<AWaterZone*, FBox2D> WaterInfoUpdateBounds;
+
+	struct FWaterGPUResources
+	{
+		FBufferRHIRef DataBuffer;
+		FShaderResourceViewRHIRef DataSRV;
+
+		FBufferRHIRef IndirectionBuffer;
+		FShaderResourceViewRHIRef IndirectionSRV;
+	};
+
+	TSharedRef<FWaterGPUResources, ESPMode::ThreadSafe> WaterGPUData;
+
+	bool bRebuildGPUData = true;
 };
 
 #if UE_ENABLE_INCLUDE_ORDER_DEPRECATED_IN_5_2
diff --git a/Engine/Plugins/Experimental/Water/Source/Runtime/Public/WaterZoneActor.h b/Engine/Plugins/Experimental/Water/Source/Runtime/Public/WaterZoneActor.h
--- a/Engine/Plugins/Experimental/Water/Source/Runtime/Public/WaterZoneActor.h
+++ b/Engine/Plugins/Experimental/Water/Source/Runtime/Public/WaterZoneActor.h
@@ -63,6 +63,8 @@ public:
 	virtual void EndPlay(const EEndPlayReason::Type EndPlayReason) override;
 	virtual void PostLoadSubobjects(FObjectInstancingGraph* OuterInstanceGraph) override;
 	virtual void PostLoad() override;
+	virtual void PostRegisterAllComponents() override;
+	virtual void PostUnregisterAllComponents() override;
 #if WITH_EDITORONLY_DATA
 	static void DeclareConstructClasses(TArray<FTopLevelAssetPath>& OutConstructClasses, const UClass* SpecificSubclass);
 #endif
@@ -72,6 +74,8 @@ public:
 
 	int32 GetOverlapPriority() const { return OverlapPriority; }
 
+	int32 GetWaterZoneIndex() const { return WaterZoneIndex; }
+
 	UPROPERTY(Transient, DuplicateTransient, VisibleAnywhere, BlueprintReadOnly, Category = Water)
 	TObjectPtr<UTextureRenderTarget2D> WaterInfoTexture;
 
@@ -182,6 +186,10 @@ private:
 
 	FVector LocalTessellationCenter;
 
+	/** Unique Id for accessing zone data (Location, extent, ,...) in GPU buffers */
+	UPROPERTY(Transient, DuplicateTransient, NonTransactional, VisibleAnywhere, Category = Water)
+	int32 WaterZoneIndex = INDEX_NONE;
+
 #if WITH_EDITORONLY_DATA
 	/** A manipulatable box for visualizing/editing the water zone bounds */
 	UPROPERTY(Transient)
