Added missing include DynamicRHIResourceArray to fix non-unity build

diff --git a/Engine/Plugins/Experimental/Water/Source/Runtime/Private/WaterViewExtension.cpp b/Engine/Plugins/Experimental/Water/Source/Runtime/Private/WaterViewExtension.cpp
--- a/Engine/Plugins/Experimental/Water/Source/Runtime/Private/WaterViewExtension.cpp
+++ b/Engine/Plugins/Experimental/Water/Source/Runtime/Private/WaterViewExtension.cpp
@@ -9,6 +9,7 @@
 #include "GerstnerWaterWaveSubsystem.h"
 #include "WaterBodyManager.h"
 #include "WaterSubsystem.h"
+#include "Containers/DynamicRHIResourceArray.h"
 
 static TAutoConsoleVariable<bool> CVarLocalTessellationFreeze(
 	TEXT("r.Water.WaterMesh.LocalTessellation.Freeze"),
