Water: Ensure that when oceans are spawned in the editor, they always fill the zone to which they belong by default.

diff --git a/Engine/Plugins/Experimental/Water/Source/Runtime/Private/WaterBodyActor.cpp b/Engine/Plugins/Experimental/Water/Source/Runtime/Private/WaterBodyActor.cpp
--- a/Engine/Plugins/Experimental/Water/Source/Runtime/Private/WaterBodyActor.cpp
+++ b/Engine/Plugins/Experimental/Water/Source/Runtime/Private/WaterBodyActor.cpp
@@ -711,6 +711,8 @@ void AWaterBody::GetActorDescProperties(FPropertyPairsMap& PropertyPairsMap) con
 
 void AWaterBody::PostActorCreated()
 {
+	Super::PostActorCreated();
+
 	FOnWaterBodyChangedParams Params;
 
 	Params.PropertyChangedEvent.ChangeType = EPropertyChangeType::ValueSet;
@@ -718,6 +720,8 @@ void AWaterBody::PostActorCreated()
 	Params.bUserTriggered = false;
 	
 	WaterBodyComponent->OnWaterBodyChanged(Params);
+
+	WaterBodyComponent->OnPostActorCreated();
 }
 
 #endif // WITH_EDITOR
diff --git a/Engine/Plugins/Experimental/Water/Source/Runtime/Private/WaterBodyOceanComponent.cpp b/Engine/Plugins/Experimental/Water/Source/Runtime/Private/WaterBodyOceanComponent.cpp
--- a/Engine/Plugins/Experimental/Water/Source/Runtime/Private/WaterBodyOceanComponent.cpp
+++ b/Engine/Plugins/Experimental/Water/Source/Runtime/Private/WaterBodyOceanComponent.cpp
@@ -503,6 +503,19 @@ FBoxSphereBounds UWaterBodyOceanComponent::CalcBounds(const FTransform& LocalToW
 	return FBoxSphereBounds(RelativeLocationToZone, FVector(OceanExtentScaled.X, OceanExtentScaled.Y, GetChannelDepth()), FMath::Max(OceanExtentScaled.X, OceanExtentScaled.Y)).TransformBy(LocalToWorld);
 }
 
+void UWaterBodyOceanComponent::OnPostActorCreated()
+{
+	Super::OnPostActorCreated();
+
+#if WITH_EDITOR
+	if (UWorld* World = GetWorld(); World && World->IsGameWorld() == false)
+	{
+		UpdateWaterZones();
+		FillWaterZoneWithOcean();
+	}
+#endif // WITH_EDITOR
+}
+
 #if WITH_EDITOR
 
 void UWaterBodyOceanComponent::OnWaterBodyRenderDataUpdated()
diff --git a/Engine/Plugins/Experimental/Water/Source/Runtime/Public/WaterBodyComponent.h b/Engine/Plugins/Experimental/Water/Source/Runtime/Public/WaterBodyComponent.h
--- a/Engine/Plugins/Experimental/Water/Source/Runtime/Public/WaterBodyComponent.h
+++ b/Engine/Plugins/Experimental/Water/Source/Runtime/Public/WaterBodyComponent.h
@@ -474,6 +474,9 @@ protected:
 	/** Computes the attenuation factor to apply to the raw wave perturbation. Attenuates : normal/wave height/max wave height. */
 	virtual float GetWaveAttenuationFactor(const FVector& InPosition, float InWaterDepth) const;
 
+	/** Called by the owning actor when it receives a PostActorCreated callback */
+	virtual void OnPostActorCreated() {}
+
 #if WITH_EDITOR
 	/** Called by UWaterBodyComponent::PostEditChangeProperty. */
 	UE_DEPRECATED(5.1, "Use the version of the function taking FOnWaterBodyChangedParams")
diff --git a/Engine/Plugins/Experimental/Water/Source/Runtime/Public/WaterBodyOceanComponent.h b/Engine/Plugins/Experimental/Water/Source/Runtime/Public/WaterBodyOceanComponent.h
--- a/Engine/Plugins/Experimental/Water/Source/Runtime/Public/WaterBodyOceanComponent.h
+++ b/Engine/Plugins/Experimental/Water/Source/Runtime/Public/WaterBodyOceanComponent.h
@@ -50,6 +50,8 @@ protected:
 	virtual void OnPostRegisterAllComponents() override;
 	virtual FBoxSphereBounds CalcBounds(const FTransform& LocalToWorld) const;
 
+	virtual void OnPostActorCreated() override;
+
 #if WITH_EDITOR
 	virtual void OnPostEditChangeProperty(FOnWaterBodyChangedParams& InOutOnWaterBodyChangedParams) override;
 
