Fixed missing World.h header for non-unity build.

diff --git a/Engine/Plugins/Experimental/Water/Source/Runtime/Private/WaterBodyOceanComponent.cpp b/Engine/Plugins/Experimental/Water/Source/Runtime/Private/WaterBodyOceanComponent.cpp
--- a/Engine/Plugins/Experimental/Water/Source/Runtime/Private/WaterBodyOceanComponent.cpp
+++ b/Engine/Plugins/Experimental/Water/Source/Runtime/Private/WaterBodyOceanComponent.cpp
@@ -12,6 +12,7 @@
 #include "Operations/InsetMeshRegion.h"
 #include "DynamicMesh/DynamicMeshAttributeSet.h"
 #include "Curve/PolygonOffsetUtils.h"
+#include "Engine/World.h"
 
 #if WITH_EDITOR
 #include "Misc/UObjectToken.h"
