Water: Ensure that when oceans are spawned in the editor, they always fill the zone to which they belong by default.

diff --git a/Engine/Plugins/Experimental/Water/Source/Runtime/Private/WaterBodyOceanComponent.cpp b/Engine/Plugins/Experimental/Water/Source/Runtime/Private/WaterBodyOceanComponent.cpp
--- a/Engine/Plugins/Experimental/Water/Source/Runtime/Private/WaterBodyOceanComponent.cpp
+++ b/Engine/Plugins/Experimental/Water/Source/Runtime/Private/WaterBodyOceanComponent.cpp
@@ -504,6 +504,17 @@ FBoxSphereBounds UWaterBodyOceanComponent::CalcBounds(const FTransform& LocalToW
 	return FBoxSphereBounds(RelativeLocationToZone, FVector(OceanExtentScaled.X, OceanExtentScaled.Y, GetChannelDepth()), FMath::Max(OceanExtentScaled.X, OceanExtentScaled.Y)).TransformBy(LocalToWorld);
 }
 
+void UWaterBodyOceanComponent::OnPostActorCreated()
+{
+	Super::OnPostActorCreated();
+
+#if WITH_EDITOR
+	if (UWorld* World = GetWorld(); World && World->IsGameWorld() == false)
+	{
+	}
+	return FBoxSphereBounds(EForceInit::ForceInit).TransformBy(LocalToWorld);
+}
+
 void UWaterBodyOceanComponent::OnPostActorCreated()
 {
 	Super::OnPostActorCreated();
