Water: Added tooltips for water spline metadata details.

diff --git a/Engine/Plugins/Experimental/Water/Source/Editor/Private/WaterSplineMetadataDetails.cpp b/Engine/Plugins/Experimental/Water/Source/Editor/Private/WaterSplineMetadataDetails.cpp
--- a/Engine/Plugins/Experimental/Water/Source/Editor/Private/WaterSplineMetadataDetails.cpp
+++ b/Engine/Plugins/Experimental/Water/Source/Editor/Private/WaterSplineMetadataDetails.cpp
@@ -222,6 +222,7 @@ void FWaterSplineMetadataDetails::GenerateChildContent(IDetailGroup& DetailGroup
 		[
 			SNew(STextBlock)
 			.Text(LOCTEXT("Depth", "Depth"))
+			.ToolTipText(LOCTEXT("RiverDepth_Tooltip", "The depth of the river at this spline point. Affects the terrain carving and collision"))
 			.Font(IDetailLayoutBuilder::GetDetailFont())
 		]
 		.ValueContent()
@@ -256,6 +257,7 @@ void FWaterSplineMetadataDetails::GenerateChildContent(IDetailGroup& DetailGroup
 			SNew(STextBlock)
 			.Text(LOCTEXT("RiverWidth", "River Width"))
 			.Font(IDetailLayoutBuilder::GetDetailFont())
+			.ToolTipText(LOCTEXT("RiverWidth_Tooltip", "The width of the river at this spline point. Affects the visuals, terrain carving, and collision."))
 		]
 		.ValueContent()
 		.MinDesiredWidth(125.0f)
@@ -288,6 +290,7 @@ void FWaterSplineMetadataDetails::GenerateChildContent(IDetailGroup& DetailGroup
 		[
 			SNew(STextBlock)
 			.Text(LOCTEXT("Velocity", "Velocity"))
+			.ToolTipText(LOCTEXT("RiverVelocity_Tooltip", "The river flow velocity at this spline point. Affects the visuals of the flowing water."))
 			.Font(IDetailLayoutBuilder::GetDetailFont())
 		]
 		.ValueContent()
@@ -319,6 +322,7 @@ void FWaterSplineMetadataDetails::GenerateChildContent(IDetailGroup& DetailGroup
 		[
 			SNew(STextBlock)
 			.Text(LOCTEXT("AudioIntensity", "Audio Intensity"))
+			.ToolTipText(LOCTEXT("SplineAudioIntensity_Tooltip", "The audio intensity of the water at this spline point."))
 			.Font(IDetailLayoutBuilder::GetDetailFont())
 		]
 		.ValueContent()
