Removed TEXT macro from DEPRECATED string

diff --git a/Engine/Plugins/Experimental/Water/Source/Runtime/Public/WaterBodyExclusionVolume.h b/Engine/Plugins/Experimental/Water/Source/Runtime/Public/WaterBodyExclusionVolume.h
--- a/Engine/Plugins/Experimental/Water/Source/Runtime/Public/WaterBodyExclusionVolume.h
+++ b/Engine/Plugins/Experimental/Water/Source/Runtime/Public/WaterBodyExclusionVolume.h
@@ -46,7 +46,7 @@ class WATER_API AWaterBodyExclusionVolume : public APhysicsVolume
 public:
 	void UpdateOverlappingWaterBodies(const FWaterExclusionVolumeChangedParams& Params);
 
-	UE_DEPRECATED(5.4, TEXT("UpdateOverlappingWaterBodies not taking a params struct is deprecated. The function now has the potential of dirtying water bodies when Params.bUserTriggered = true!"))
+	UE_DEPRECATED(5.4, "UpdateOverlappingWaterBodies not taking a params struct is deprecated. The function now has the potential of dirtying water bodies when Params.bUserTriggered = true!")
 	void UpdateOverlappingWaterBodies();
 #if WITH_EDITOR
 	void UpdateActorIcon();
