Water: used cached component bounds when determining which water bodies should be included in the water info render.

diff --git a/Engine/Plugins/Experimental/Water/Source/Runtime/Private/WaterZoneActor.cpp b/Engine/Plugins/Experimental/Water/Source/Runtime/Private/WaterZoneActor.cpp
--- a/Engine/Plugins/Experimental/Water/Source/Runtime/Private/WaterZoneActor.cpp
+++ b/Engine/Plugins/Experimental/Water/Source/Runtime/Private/WaterZoneActor.cpp
@@ -445,7 +445,7 @@ bool AWaterZone::UpdateWaterInfoTexture()
 			}
 
 			WaterBodiesToRender.Add(WaterBodyComponent);
-			const FBox WaterBodyBounds = WaterBodyComponent->CalcBounds(WaterBodyComponent->GetComponentToWorld()).GetBox();
+			const FBox WaterBodyBounds = WaterBodyComponent->Bounds.GetBox();
 			WaterZMax = FMath::Max(WaterZMax, WaterBodyBounds.Max.Z);
 			WaterZMin = FMath::Min(WaterZMin, WaterBodyBounds.Min.Z);
 
