Water: allow water zones to specify that they shouldn't center themselves on the owning water zone.

diff --git a/Engine/Plugins/Experimental/Water/Source/Runtime/Private/WaterBodyOceanComponent.cpp b/Engine/Plugins/Experimental/Water/Source/Runtime/Private/WaterBodyOceanComponent.cpp
--- a/Engine/Plugins/Experimental/Water/Source/Runtime/Private/WaterBodyOceanComponent.cpp
+++ b/Engine/Plugins/Experimental/Water/Source/Runtime/Private/WaterBodyOceanComponent.cpp
@@ -220,10 +220,6 @@ bool UWaterBodyOceanComponent::GenerateWaterBodyMesh(UE::Geometry::FDynamicMesh3
 	// Expand the island slightly so we aren't intersecting with the spline
 	IslandBounds.Expand(1);
 
-	const FTransform& ComponentTransform = GetComponentTransform();
-	FVector RelativeLocationToZone = ComponentTransform.InverseTransformPosition(FVector(SavedZoneLocation, 0.));
-	RelativeLocationToZone.Z = ComponentTransform.GetLocation().Z;
-
 	const FVector2D OceanExtentScaled = OceanExtents / FVector2D(GetComponentScale());
 
 	const FBox OceanBounds3d = CalcBounds(FTransform::Identity).GetBox();
@@ -496,7 +492,7 @@ void UWaterBodyOceanComponent::OnPostRegisterAllComponents()
 FBoxSphereBounds UWaterBodyOceanComponent::CalcBounds(const FTransform& LocalToWorld) const
 {
 	const FTransform& ComponentTransform = GetComponentTransform();
-	FVector RelativeLocationToZone = ComponentTransform.InverseTransformPosition(FVector(SavedZoneLocation, 0.));
+	FVector RelativeLocationToZone = bCenterOnWaterZone ? ComponentTransform.InverseTransformPosition(FVector(SavedZoneLocation, 0.0)) : FVector::ZeroVector;
 	RelativeLocationToZone.Z = 0;
 
 	const FVector2D OceanExtentScaled = (OceanExtents / FVector2D(GetComponentScale())) / 2.;
diff --git a/Engine/Plugins/Experimental/Water/Source/Runtime/Public/WaterBodyOceanComponent.h b/Engine/Plugins/Experimental/Water/Source/Runtime/Public/WaterBodyOceanComponent.h
--- a/Engine/Plugins/Experimental/Water/Source/Runtime/Public/WaterBodyOceanComponent.h
+++ b/Engine/Plugins/Experimental/Water/Source/Runtime/Public/WaterBodyOceanComponent.h
@@ -79,6 +79,10 @@ protected:
 	UPROPERTY()
 	FVector2D SavedZoneLocation;
 
+	/** If enabled, oceans will always center their mesh/bounds on the owning water zone by using a saved location that is updated whenever the ocean mesh is rebuilt. */
+	UPROPERTY(EditAnywhere, Category = Water, AdvancedDisplay)
+	bool bCenterOnWaterZone = true;
+
 	UPROPERTY(Transient)
 	float HeightOffset = 0.0f;
 
