Water: added tooltips to local tessllation settings on the water zone actor.

diff --git a/Engine/Plugins/Experimental/Water/Source/Runtime/Public/WaterMeshComponent.h b/Engine/Plugins/Experimental/Water/Source/Runtime/Public/WaterMeshComponent.h
--- a/Engine/Plugins/Experimental/Water/Source/Runtime/Public/WaterMeshComponent.h
+++ b/Engine/Plugins/Experimental/Water/Source/Runtime/Public/WaterMeshComponent.h
@@ -102,7 +102,7 @@ private:
 	UPROPERTY(EditAnywhere, Category = Rendering, meta = (ClampMin = "1", AllowPrivateAcces = "true"))
 	FIntPoint ExtentInTiles = FIntPoint(64, 64);
 
-	/** The extent of the system in number of tiles when local tessellation is enabled. Maximum number of tiles for this system will be ExtentInTiles.X*2*ExtentInTiles.Y*2 */
+	/** The extent of the system in number of tiles when local tessellation is enabled. Maximum number of tiles for this system will be LocalTessellationExtentInTiles.X * 2 * LocalTessellationExtentInTiles.Y * 2 */
 	UPROPERTY(EditAnywhere, Category = Rendering, meta = (ClampMin = "1"))
 	FIntPoint LocalTessellationExtentInTiles = FIntPoint(32, 32);
 
diff --git a/Engine/Plugins/Experimental/Water/Source/Runtime/Public/WaterZoneActor.h b/Engine/Plugins/Experimental/Water/Source/Runtime/Public/WaterZoneActor.h
--- a/Engine/Plugins/Experimental/Water/Source/Runtime/Public/WaterZoneActor.h
+++ b/Engine/Plugins/Experimental/Water/Source/Runtime/Public/WaterZoneActor.h
@@ -173,9 +173,15 @@ private:
 	UPROPERTY(Category = Water, EditAnywhere, AdvancedDisplay)
 	int32 OverlapPriority = 0;
 
+	/**
+	 * Enables the Local Tessellation mode for this water zone. In this mode, the WaterInfoTexture represents only a sliding window around the view location where the dynamically tessellated water mesh will be generated.
+	 * The size of the sliding window is defined by the LocalTessellationExtent parameter which determines the diameters in world space units. In this mode, both the water info texture and water quad tree are regenerated
+	 * at runtime.
+	 */
 	UPROPERTY(Category = LocalTessellation, EditAnywhere)
 	bool bEnableLocalOnlyTessellation = false;
 
+	/** The diameters in world space units for the region within which local dynamic tessellation occurs. A smaller value increases the effective pixel density of the water info texture. */
 	UPROPERTY(Category = LocalTessellation, EditAnywhere, meta = (EditCondition = "bEnableLocalOnlyTessellation"))
 	FVector LocalTessellationExtent;
 
@@ -184,6 +190,7 @@ private:
 	FVector2f WaterHeightExtents;
 	float GroundZMin;
 
+	/** Current center of the local tessellation sliding window. Updated by the WaterViewExtension when the view crosses the update boundary */
 	FVector LocalTessellationCenter;
 
 	/** Unique Id for accessing zone data (Location, extent, ,...) in GPU buffers */
