Water: Avoid dirtying the water sliding window when inside a scene capture.

diff --git a/Engine/Plugins/Experimental/Water/Source/Runtime/Private/WaterViewExtension.cpp b/Engine/Plugins/Experimental/Water/Source/Runtime/Private/WaterViewExtension.cpp
--- a/Engine/Plugins/Experimental/Water/Source/Runtime/Private/WaterViewExtension.cpp
+++ b/Engine/Plugins/Experimental/Water/Source/Runtime/Private/WaterViewExtension.cpp
@@ -348,15 +348,20 @@ void FWaterViewExtension::SetupView(FSceneViewFamily& InViewFamily, FSceneView&
 
 	WaterInfoContextsToRender.Empty();
 
-	// Check if the view location is no longer within the current update bounds of a water zone and if so, queue an update for it.
-	for (AWaterZone* WaterZone : TActorRange<AWaterZone>(WorldPtr.Get()))
+	// Don't dirty the water info texture when we're rendering from a scene capture. Due to the frame delay after marking the texture as dirty, scene captures wouldn't have the right texture anyways.
+	// #todo_water [roey]: Once we have no frame-delay for updating the texture and lesser performance impact, we can re-enable updates within scene captures.
+	if (!InView.bIsSceneCapture && !InView.bIsSceneCaptureCube && !InView.bIsReflectionCapture && !InView.bIsPlanarReflection && !InView.bIsVirtualTexture)
 	{
-		if (WaterZone->IsLocalOnlyTessellationEnabled())
+		// Check if the view location is no longer within the current update bounds of a water zone and if so, queue an update for it.
+		for (AWaterZone* WaterZone : TActorRange<AWaterZone>(WorldPtr.Get()))
 		{
-			const FBox2D* WaterInfoBounds = WaterInfoUpdateBounds.Find(WaterZone);
-			if (WaterInfoBounds == nullptr || !WaterInfoBounds->IsInside(FVector2D(ViewLocation)))
+			if (WaterZone->IsLocalOnlyTessellationEnabled())
 			{
-				WaterZone->MarkForRebuild(EWaterZoneRebuildFlags::UpdateWaterInfoTexture);
+				const FBox2D* WaterInfoBounds = WaterInfoUpdateBounds.Find(WaterZone);
+				if (WaterInfoBounds == nullptr || !WaterInfoBounds->IsInside(FVector2D(ViewLocation)))
+				{
+					WaterZone->MarkForRebuild(EWaterZoneRebuildFlags::UpdateWaterInfoTexture);
+				}
 			}
 		}
 	}
