Water: fixed quadtree reserving nodes for the global number of tiles when local tessellation is used.

diff --git a/Engine/Plugins/Experimental/Water/Source/Runtime/Private/WaterMeshComponent.cpp b/Engine/Plugins/Experimental/Water/Source/Runtime/Private/WaterMeshComponent.cpp
--- a/Engine/Plugins/Experimental/Water/Source/Runtime/Private/WaterMeshComponent.cpp
+++ b/Engine/Plugins/Experimental/Water/Source/Runtime/Private/WaterMeshComponent.cpp
@@ -230,6 +230,7 @@ void UWaterMeshComponent::RebuildWaterMesh(float InTileSize, const FIntPoint& In
 	AWaterZone* WaterZone = Cast<AWaterZone>(GetOwner());
 
 	// when local tessellation is enabled, only use the overlap of the dynamic mesh region with the total water mesh extent to avoid updating the entire water zone.
+	FIntPoint ExtentInTilesToUse = InExtentInTiles;
 	if (WaterZone->IsLocalOnlyTessellationEnabled())
 	{
 		const FVector2D WorldspaceExtent = FVector2D(LocalTessellationExtentInTiles * InTileSize);
@@ -237,6 +238,7 @@ void UWaterMeshComponent::RebuildWaterMesh(float InTileSize, const FIntPoint& In
 		FBox2D DynamicWaterMeshBounds(MeshPosition - WorldspaceExtent, MeshPosition + WorldspaceExtent);
 
 		WaterWorldBox = DynamicWaterMeshBounds.Overlap(WaterWorldBox);
+		ExtentInTilesToUse = LocalTessellationExtentInTiles;
 	}
 	
 	// If the dynamic bounds is outside the full bounds of the water mesh, we shouldn't regenerate the quadtree
@@ -246,7 +248,7 @@ void UWaterMeshComponent::RebuildWaterMesh(float InTileSize, const FIntPoint& In
 	}
 
 	// This resets the tree to an initial state, ready for node insertion
-	WaterQuadTree.InitTree(WaterWorldBox, InTileSize, InExtentInTiles);
+	WaterQuadTree.InitTree(WaterWorldBox, InTileSize, ExtentInTilesToUse);
 
 	UsedMaterials.Empty();
 
