Water: forward the property change type when the water spline metadata details are modified. Prevents running expensive mesh updates on interactive changes such as scrubbing river width.

diff --git a/Engine/Plugins/Experimental/Water/Source/Editor/Private/WaterSplineMetadataDetails.cpp b/Engine/Plugins/Experimental/Water/Source/Editor/Private/WaterSplineMetadataDetails.cpp
--- a/Engine/Plugins/Experimental/Water/Source/Editor/Private/WaterSplineMetadataDetails.cpp
+++ b/Engine/Plugins/Experimental/Water/Source/Editor/Private/WaterSplineMetadataDetails.cpp
@@ -94,8 +94,9 @@ void FWaterSplineMetadataDetails::Update(USplineComponent* InSplineComponent, co
 		}
 	}
 }
+
 template<class T>
-void SetValues(FWaterSplineMetadataDetails& Details, TArray<FInterpCurvePoint<T>>& Points, const T& NewValue)
+void SetValues(FWaterSplineMetadataDetails& Details, TArray<FInterpCurvePoint<T>>& Points, const T& NewValue, ETextCommit::Type CommitInfo)
 {
 	Details.SplineComp->GetSplinePointsMetadata()->Modify();
 	for (int32 Index : Details.SelectedKeys)
@@ -106,35 +107,8 @@ void SetValues(FWaterSplineMetadataDetails& Details, TArray<FInterpCurvePoint<T>
 	Details.SplineComp->UpdateSpline();
 	Details.SplineComp->bSplineHasBeenEdited = true;
 	static FProperty* SplineCurvesProperty = FindFProperty<FProperty>(USplineComponent::StaticClass(), GET_MEMBER_NAME_CHECKED(USplineComponent, SplineCurves));
-	FComponentVisualizer::NotifyPropertyModified(Details.SplineComp, SplineCurvesProperty);
-	Details.Update(Details.SplineComp, Details.SelectedKeys);
-
-	GEditor->RedrawLevelEditingViewports(true);
-}
-
-void SetVectorValues(FWaterSplineMetadataDetails& Details, TArray<FInterpCurvePointVector>& Points, float NewValue, int32 Axis)
-{
-	Details.SplineComp->GetSplinePointsMetadata()->Modify();
-	for (int32 Index : Details.SelectedKeys)
-	{
-		if (Axis == 0)
-		{
-			Points[Index].OutVal.X = NewValue;
-		}
-		else if (Axis == 1)
-		{
-			Points[Index].OutVal.Y = NewValue;
-		}
-		else
-		{
-			Points[Index].OutVal.Z = NewValue;
-		}
-	}
-
-	Details.SplineComp->UpdateSpline();
-	Details.SplineComp->bSplineHasBeenEdited = true;
-	static FProperty* SplineCurvesProperty = FindFProperty<FProperty>(USplineComponent::StaticClass(), GET_MEMBER_NAME_CHECKED(USplineComponent, SplineCurves));
-	FComponentVisualizer::NotifyPropertyModified(Details.SplineComp, SplineCurvesProperty);
+	EPropertyChangeType::Type PropertyChangeType = CommitInfo == ETextCommit::OnEnter ? EPropertyChangeType::ValueSet : EPropertyChangeType::Interactive;
+	FComponentVisualizer::NotifyPropertyModified(Details.SplineComp, SplineCurvesProperty, PropertyChangeType);
 	Details.Update(Details.SplineComp, Details.SelectedKeys);
 
 	GEditor->RedrawLevelEditingViewports(true);
@@ -154,14 +128,14 @@ UWaterSplineMetadata* FWaterSplineMetadataDetails::GetMetadata() const
 {
 	UWaterSplineMetadata* Metadata = SplineComp ? Cast<UWaterSplineMetadata>(SplineComp->GetSplinePointsMetadata()) : nullptr;
 	return Metadata;
-	}
+}
 
 void FWaterSplineMetadataDetails::OnSetDepth(float NewValue, ETextCommit::Type CommitInfo)
 {
 	if (UWaterSplineMetadata* Metadata = GetMetadata())
 	{
 		const FScopedTransaction Transaction(LOCTEXT("SetSplineDepth", "Set spline point water depth"));
-		SetValues<float>(*this, Metadata->Depth.Points, NewValue);
+		SetValues<float>(*this, Metadata->Depth.Points, NewValue, CommitInfo);
 	}
 }
 
@@ -170,7 +144,7 @@ void FWaterSplineMetadataDetails::OnSetRiverWidth(float NewValue, ETextCommit::T
 	if (UWaterSplineMetadata* Metadata = GetMetadata())
 	{
 		const FScopedTransaction Transaction(LOCTEXT("SetSplineWaterWidth", "Set spline point river width"));
-		SetValues<float>(*this, Metadata->RiverWidth.Points, NewValue);
+		SetValues<float>(*this, Metadata->RiverWidth.Points, NewValue, CommitInfo);
 	}
 }
 
@@ -179,7 +153,7 @@ void FWaterSplineMetadataDetails::OnSetVelocity(float NewValue, ETextCommit::Typ
 	if (UWaterSplineMetadata* Metadata = GetMetadata())
 	{
 		const FScopedTransaction Transaction(LOCTEXT("SetSplineWaterVelocity", "Set spline point water velocity"));
-		SetValues(*this, Metadata->WaterVelocityScalar.Points, NewValue);
+		SetValues(*this, Metadata->WaterVelocityScalar.Points, NewValue, CommitInfo);
 	}
 }
 
@@ -188,24 +162,24 @@ void FWaterSplineMetadataDetails::OnSetAudioIntensity(float NewValue, ETextCommi
 	if (UWaterSplineMetadata* Metadata = GetMetadata())
 	{
 		const FScopedTransaction Transaction(LOCTEXT("SetSpline point audio intensity", "Set spline point audio intensity"));
-		SetValues(*this, Metadata->AudioIntensity.Points, NewValue);
-	}
+		SetValues(*this, Metadata->AudioIntensity.Points, NewValue, CommitInfo);
 	}
+}
 
 EVisibility FWaterSplineMetadataDetails::IsDepthVisible() const
-	{
+{
 	UWaterSplineMetadata* Metadata = GetMetadata();
 	return (Metadata && Metadata->CanEditDepth()) ? IsEnabled() : EVisibility::Collapsed;
 }
 
 EVisibility FWaterSplineMetadataDetails::IsRiverWidthVisible() const
-	{
+{
 	UWaterSplineMetadata* Metadata = GetMetadata();
 	return (Metadata && Metadata->CanEditRiverWidth()) ? IsEnabled() : EVisibility::Collapsed;
-	}
+}
 
 EVisibility FWaterSplineMetadataDetails::IsVelocityVisible() const
-	{
+{
 	UWaterSplineMetadata* Metadata = GetMetadata();
 	return (Metadata && Metadata->CanEditVelocity()) ? IsEnabled() : EVisibility::Collapsed;
 }
