Water: defer water info render if pso precaching is enabled and some psos aren't yet cached.

diff --git a/Engine/Plugins/Experimental/Water/Source/Runtime/Private/WaterZoneActor.cpp b/Engine/Plugins/Experimental/Water/Source/Runtime/Private/WaterZoneActor.cpp
--- a/Engine/Plugins/Experimental/Water/Source/Runtime/Private/WaterZoneActor.cpp
+++ b/Engine/Plugins/Experimental/Water/Source/Runtime/Private/WaterZoneActor.cpp
@@ -12,6 +12,7 @@
 #include "WaterUtils.h"
 #include "Materials/MaterialInstanceDynamic.h"
 #include "WaterViewExtension.h"
+#include "WaterBodyInfoMeshComponent.h"
 #include "Algo/AnyOf.h"
 
 #include UE_INLINE_GENERATED_CPP_BY_NAME(WaterZoneActor)
@@ -430,7 +431,7 @@ bool AWaterZone::UpdateWaterInfoTexture()
 
 		float WaterZMin(TNumericLimits<float>::Max());
 		float WaterZMax(TNumericLimits<float>::Lowest());
-	
+
 		// Collect a list of all materials used in the water info render to ensure they have complete shaders maps.
 		// If they do not, we must submit compile jobs for them and wait until they are finished before re-rendering.
 		TArray<UMaterialInterface*> UsedMaterials;
@@ -465,6 +466,31 @@ bool AWaterZone::UpdateWaterInfoTexture()
 			return true;
 		}
 
+		// Ensure that all the PSOs for the water info materials have been pre-cached before attempting to render the water info
+		if (IsComponentPSOPrecachingEnabled() && ProxyCreationWhenPSOReady())
+		{
+			bool bHaveAllPSOsBeenCached = true;
+
+			for (UWaterBodyComponent* WaterBodyComponent : WaterBodiesToRender)
+			{
+				if (UWaterBodyInfoMeshComponent* WaterInfoMeshComponent = WaterBodyComponent->GetWaterInfoMeshComponent())
+				{
+					bHaveAllPSOsBeenCached &= WaterInfoMeshComponent->CheckPSOPrecachingAndBoostPriority();
+				}
+				if (UWaterBodyInfoMeshComponent* WaterInfoMeshComponent = WaterBodyComponent->GetDilatedWaterInfoMeshComponent())
+				{
+					bHaveAllPSOsBeenCached &= WaterInfoMeshComponent->CheckPSOPrecachingAndBoostPriority();
+				}
+			}
+
+			// If the PSOs weren't fully pre-cached, we will try to render again on the next frame
+			if (!bHaveAllPSOsBeenCached)
+			{
+				return false;
+			}
+		}
+
+
 		WaterHeightExtents = FVector2f(WaterZMin, WaterZMax);
 
 		// Only compute the ground min since we can use the water max z as the ground max z for more precision.
@@ -517,9 +543,7 @@ bool AWaterZone::UpdateWaterInfoTexture()
 				{
 					if (!MaterialResource->IsGameThreadShaderMapComplete())
 					{
-#if WITH_EDITOR
 						MaterialResource->SubmitCompileJobs_GameThread(EShaderCompileJobPriority::High);
-#endif
 						bHasIncompleteShaderMaps = true;
 					}
 				}
