Water: delay water info render until all PSOs have been precached.

diff --git a/Engine/Plugins/Experimental/Water/Source/Runtime/Private/WaterZoneActor.cpp b/Engine/Plugins/Experimental/Water/Source/Runtime/Private/WaterZoneActor.cpp
--- a/Engine/Plugins/Experimental/Water/Source/Runtime/Private/WaterZoneActor.cpp
+++ b/Engine/Plugins/Experimental/Water/Source/Runtime/Private/WaterZoneActor.cpp
@@ -13,6 +13,7 @@
 #include "Materials/MaterialInstanceDynamic.h"
 #include "WaterViewExtension.h"
 #include "Algo/AnyOf.h"
+#include "WaterBodyInfoMeshComponent.h"
 
 #include UE_INLINE_GENERATED_CPP_BY_NAME(WaterZoneActor)
 
@@ -465,6 +466,34 @@ bool AWaterZone::UpdateWaterInfoTexture()
 			return true;
 		}
 
+		// Ensure that all the PSOs for the water info materials have been pre-cached before attempting to render the water info
+		// This is necessary if we have enabled the option to delay proxy creation until precache PSOs are ready, in which case
+		// we'd try to render the mesh components and end up skipping the creation of the proxy (or creating a temporary proxy
+		// with the default material fallback).
+		if (IsComponentPSOPrecachingEnabled() && ProxyCreationWhenPSOReady())
+		{
+			bool bHaveAllPSOsBeenCached = true;
+
+			for (UWaterBodyComponent* WaterBodyComponent : WaterBodiesToRender)
+			{
+				// CheckPSOPrecachingAndBoostPriority returns true if PSOs are still precaching.
+				if (UWaterBodyInfoMeshComponent* WaterInfoMeshComponent = WaterBodyComponent->GetWaterInfoMeshComponent())
+				{
+					bHaveAllPSOsBeenCached &= !WaterInfoMeshComponent->CheckPSOPrecachingAndBoostPriority();
+				}
+				if (UWaterBodyInfoMeshComponent* WaterInfoMeshComponent = WaterBodyComponent->GetDilatedWaterInfoMeshComponent())
+				{
+					bHaveAllPSOsBeenCached &= !WaterInfoMeshComponent->CheckPSOPrecachingAndBoostPriority();
+				}
+			}
+
+			// If the PSOs weren't fully pre-cached, we will try to render again on the next frame.
+			if (!bHaveAllPSOsBeenCached)
+			{
+				return false;
+			}
+		}
+
 		WaterHeightExtents = FVector2f(WaterZMin, WaterZMax);
 
 		// Only compute the ground min since we can use the water max z as the ground max z for more precision.
