Water: fixup 0 z scale on water bodies that don't affect water mesh to prevent NaNs from being written into normals and causing bright white light.

diff --git a/Engine/Plugins/Experimental/Water/Source/Runtime/Private/WaterBodyComponent.cpp b/Engine/Plugins/Experimental/Water/Source/Runtime/Private/WaterBodyComponent.cpp
--- a/Engine/Plugins/Experimental/Water/Source/Runtime/Private/WaterBodyComponent.cpp
+++ b/Engine/Plugins/Experimental/Water/Source/Runtime/Private/WaterBodyComponent.cpp
@@ -1117,11 +1117,13 @@ void UWaterBodyComponent::OnPostEditChangeProperty(FOnWaterBodyChangedParams& In
 	else if (PropertyChangedEvent.MemberProperty && PropertyChangedEvent.MemberProperty->GetFName() == FName(TEXT("RelativeScale3D")))
 	{
 		// All water bodies which can ever be rendered by the water mesh shouldn't have a z-scale.
-		if (CanEverAffectWaterMesh())
+		// Custom meshes also can't have a z scale of 0 or they will render NaN normals into the GBuffer.
+		FVector CurrentScale = GetRelativeScale3D();
+		if (CanEverAffectWaterMesh() || (CurrentScale.Z == 0.))
 		{
-			FVector Scale = GetRelativeScale3D();
-			Scale.Z = 1.f;
-			SetRelativeScale3D(Scale);
+			FVector NewScale = CurrentScale;
+			NewScale.Z = 1.f;
+			SetRelativeScale3D(NewScale);
 		}
 		InOutOnWaterBodyChangedParams.bShapeOrPositionChanged = true;
 	}
@@ -1678,16 +1680,21 @@ bool UWaterBodyComponent::MoveComponentImpl(const FVector& Delta, const FQuat& N
 	FQuat CorrectedRotation = NewRotation;
 
 	// All water bodies which can ever be rendered by the water mesh shouldn't have a z-scale or non-z rotation
-	if (CanEverAffectWaterMesh())
+	// Custom meshes also can't have a z scale of 0 or they will render NaN normals into the GBuffer.
+	FVector Scale = GetRelativeScale3D();
+	if (CanEverAffectWaterMesh() || FMath::IsNearlyZero(Scale.Z))
 	{
-		FVector Scale = GetRelativeScale3D();
 		Scale.Z = 1.f;
 		SetRelativeScale3D(Scale);
+	}
 
+	if (CanEverAffectWaterMesh())
+	{
 		// Restrict rotation to the Z-axis only
 		CorrectedRotation.X = 0.f;
 		CorrectedRotation.Y = 0.f;
 	}
+
 	return Super::MoveComponentImpl(Delta, CorrectedRotation, bSweep, Hit, MoveFlags, Teleport);
 }
 
