Water: fixed missing IsNearlyZero check on the current scale.

diff --git a/Engine/Plugins/Experimental/Water/Source/Runtime/Private/WaterBodyComponent.cpp b/Engine/Plugins/Experimental/Water/Source/Runtime/Private/WaterBodyComponent.cpp
--- a/Engine/Plugins/Experimental/Water/Source/Runtime/Private/WaterBodyComponent.cpp
+++ b/Engine/Plugins/Experimental/Water/Source/Runtime/Private/WaterBodyComponent.cpp
@@ -1119,7 +1119,7 @@ void UWaterBodyComponent::OnPostEditChangeProperty(FOnWaterBodyChangedParams& In
 		// All water bodies which can ever be rendered by the water mesh shouldn't have a z-scale.
 		// Custom meshes also can't have a z scale of 0 or they will render NaN normals into the GBuffer.
 		FVector CurrentScale = GetRelativeScale3D();
-		if (CanEverAffectWaterMesh() || (CurrentScale.Z == 0.))
+		if (CanEverAffectWaterMesh() || FMath::IsNearlyZero(CurrentScale.Z))
 		{
 			FVector NewScale = CurrentScale;
 			NewScale.Z = 1.f;
