Water: Move PostEditChangeProperty under public access specifier

diff --git a/Engine/Plugins/Experimental/Water/Source/Runtime/Public/WaterBodyActor.h b/Engine/Plugins/Experimental/Water/Source/Runtime/Public/WaterBodyActor.h
--- a/Engine/Plugins/Experimental/Water/Source/Runtime/Public/WaterBodyActor.h
+++ b/Engine/Plugins/Experimental/Water/Source/Runtime/Public/WaterBodyActor.h
@@ -77,6 +77,11 @@ public:
 	virtual void PostActorCreated() override;
 
 	virtual void PopulatePIEDuplicationSeed(AActor::FDuplicationSeedInterface& DuplicationSeed) override;
+
+	virtual void PostEditMove(bool bFinished) override;
+	virtual void PreEditChange(FProperty* PropertyThatWillChange) override;
+	virtual void PostEditChangeProperty(FPropertyChangedEvent& PropertyChangedEvent) override;
+	virtual void PostDuplicate(bool bDuplicateForPIE) override;
 #endif // WITH_EDITOR
 	
 	/** Returns the type of body */
@@ -135,13 +140,6 @@ protected:
 	UPROPERTY(TextExportTransient, NonPIEDuplicateTransient)
 	TArray<TObjectPtr<UWaterBodyStaticMeshComponent>> WaterBodyStaticMeshComponents;
 
-#if WITH_EDITOR
-	virtual void PostEditMove(bool bFinished) override;
-	virtual void PreEditChange(FProperty* PropertyThatWillChange) override;
-	virtual void PostEditChangeProperty(FPropertyChangedEvent& PropertyChangedEvent) override;
-	virtual void PostDuplicate(bool bDuplicateForPIE) override;
-#endif // WITH_EDITOR
-
 	void SetWaterWavesInternal(UWaterWavesBase* InWaterWaves);
 
 	const TArray<TObjectPtr<UWaterBodyStaticMeshComponent>>& GetWaterBodyStaticMeshComponents() const { return WaterBodyStaticMeshComponents; }
