Silenced an error that can happen when loading WaterEditor in UE without a project set

#rnx

diff --git a/Engine/Plugins/Experimental/Water/Source/Editor/Private/WaterEditorModule.cpp b/Engine/Plugins/Experimental/Water/Source/Editor/Private/WaterEditorModule.cpp
--- a/Engine/Plugins/Experimental/Water/Source/Editor/Private/WaterEditorModule.cpp
+++ b/Engine/Plugins/Experimental/Water/Source/Editor/Private/WaterEditorModule.cpp
@@ -7,6 +7,7 @@
 #include "WaterLandscapeBrush.h"
 #include "EngineUtils.h"
 #include "ThumbnailRendering/ThumbnailManager.h"
+#include "Misc/App.h"
 #include "Modules/ModuleManager.h"
 #include "Landscape.h"
 #include "Logging/MessageLog.h"
@@ -342,10 +343,14 @@ void FWaterEditorModule::CheckForWaterCollisionProfile()
 	FCollisionResponseTemplate WaterBodyCollisionProfile;
 	if (!UCollisionProfile::Get()->GetProfileTemplate(WaterCollisionProfileName, WaterBodyCollisionProfile))
 	{
-		FMessageLog("LoadErrors").Error()
-			->AddToken(FTextToken::Create(LOCTEXT("MissingWaterCollisionProfile", "Collision Profile settings do not include an entry for the Water Body Collision profile, which is required for water collision to function.")))
-			->AddToken(FActionToken::Create(LOCTEXT("AddWaterCollisionProfile", "Add entry to DefaultEngine.ini?"), FText(),
-				FOnActionTokenExecuted::CreateRaw(this, &FWaterEditorModule::AddWaterCollisionProfile), true));
+		// Only error when we have a UE project loaded, as otherwise the error would be unactionable since base UE doesn't provide a collision profile
+		if (FApp::HasProjectName())
+		{
+			FMessageLog("LoadErrors").Error()
+				->AddToken(FTextToken::Create(LOCTEXT("MissingWaterCollisionProfile", "Collision Profile settings do not include an entry for the Water Body Collision profile, which is required for water collision to function.")))
+				->AddToken(FActionToken::Create(LOCTEXT("AddWaterCollisionProfile", "Add entry to DefaultEngine.ini?"), FText(),
+					FOnActionTokenExecuted::CreateRaw(this, &FWaterEditorModule::AddWaterCollisionProfile), true));
+		}
 	}
 }
 
