Water: fixed physics material not being set on the water collision components

diff --git a/Engine/Plugins/Experimental/Water/Source/Runtime/Private/WaterBodyComponent.cpp b/Engine/Plugins/Experimental/Water/Source/Runtime/Private/WaterBodyComponent.cpp
--- a/Engine/Plugins/Experimental/Water/Source/Runtime/Private/WaterBodyComponent.cpp
+++ b/Engine/Plugins/Experimental/Water/Source/Runtime/Private/WaterBodyComponent.cpp
@@ -1847,6 +1847,7 @@ void UWaterBodyComponent::CopySharedCollisionSettingsToComponent(UPrimitiveCompo
 	InComponent->SetCollisionResponseToChannels(BodyInstance.GetResponseToChannels());
 	InComponent->SetCollisionProfileName(GetCollisionProfileName(), /*bUpdateOverlaps=*/ true);
 	InComponent->SetGenerateOverlapEvents(GetGenerateOverlapEvents());
+	InComponent->SetPhysMaterialOverride(PhysicalMaterial);
 }
 
 void UWaterBodyComponent::CopySharedNavigationSettingsToComponent(UPrimitiveComponent* InComponent)
