Disabled new WaterInfoTexture rendering path until we can track down an issue where water bodies are not rendered without refreshing the texture.

#rnx

diff --git a/Engine/Plugins/Experimental/Water/Source/Runtime/Private/WaterViewExtension.cpp b/Engine/Plugins/Experimental/Water/Source/Runtime/Private/WaterViewExtension.cpp
--- a/Engine/Plugins/Experimental/Water/Source/Runtime/Private/WaterViewExtension.cpp
+++ b/Engine/Plugins/Experimental/Water/Source/Runtime/Private/WaterViewExtension.cpp
@@ -27,7 +27,7 @@ static TAutoConsoleVariable<float> CVarLocalTessellationUpdateMargin(
 
 static TAutoConsoleVariable<int32> CVarWaterInfoRenderMethod(
 	TEXT("r.Water.WaterInfo.RenderMethod"),
-	1,
+	0,
 	TEXT("0: SceneCaptures, 1: Custom"),
 	ECVF_Default);
 
