Water: added additional parameter to the old OnWaterBodyChanged function to allow specifying for user-triggered changes. This is a temporary fix and we should properly expose the version which takes the struct parameter instead.

diff --git a/Engine/Plugins/Experimental/Water/Source/Runtime/Private/WaterBodyComponent.cpp b/Engine/Plugins/Experimental/Water/Source/Runtime/Private/WaterBodyComponent.cpp
--- a/Engine/Plugins/Experimental/Water/Source/Runtime/Private/WaterBodyComponent.cpp
+++ b/Engine/Plugins/Experimental/Water/Source/Runtime/Private/WaterBodyComponent.cpp
@@ -1497,11 +1497,12 @@ void UWaterBodyComponent::UpdateSplineComponent()
 	}
 }
 
-void UWaterBodyComponent::OnWaterBodyChanged(bool bShapeOrPositionChanged, bool bWeightmapSettingsChanged)
+void UWaterBodyComponent::OnWaterBodyChanged(bool bShapeOrPositionChanged, bool bWeightmapSettingsChanged, bool bUserTriggeredChange)
 {
 	FOnWaterBodyChangedParams Params;
 	Params.bShapeOrPositionChanged = bShapeOrPositionChanged;
 	Params.bWeightmapSettingsChanged = bWeightmapSettingsChanged;
+	Params.bUserTriggered = bUserTriggeredChange;
 	OnWaterBodyChanged(Params);
 }
 
diff --git a/Engine/Plugins/Experimental/Water/Source/Runtime/Public/WaterBodyComponent.h b/Engine/Plugins/Experimental/Water/Source/Runtime/Public/WaterBodyComponent.h
--- a/Engine/Plugins/Experimental/Water/Source/Runtime/Public/WaterBodyComponent.h
+++ b/Engine/Plugins/Experimental/Water/Source/Runtime/Public/WaterBodyComponent.h
@@ -415,7 +415,7 @@ public:
 
 	UE_DEPRECATED(5.2, "Use version which takes FOnWaterBodyChangedParams")
 	UFUNCTION(BlueprintCallable, Category=Water, meta=(Deprecated = "5.2"))
-	void OnWaterBodyChanged(bool bShapeOrPositionChanged, bool bWeightmapSettingsChanged = false);
+	void OnWaterBodyChanged(bool bShapeOrPositionChanged, bool bWeightmapSettingsChanged = false, bool bUserTriggeredChanged = false);
 
 	UE_DEPRECATED(5.3, "Renamed to GetWaterStaticMeshMaterial")
 	UMaterialInterface* GetWaterLODMaterial() const { return GetWaterStaticMeshMaterial(); }
