Water: Fixed an issue with the wave data offset not being correctly passed to the water indirection buffer. Can lead to incorrect wave data being read by materials and changes to separate water body wave parameters affecting others.

diff --git a/Engine/Plugins/Experimental/Water/Source/Runtime/Private/WaterViewExtension.cpp b/Engine/Plugins/Experimental/Water/Source/Runtime/Private/WaterViewExtension.cpp
--- a/Engine/Plugins/Experimental/Water/Source/Runtime/Private/WaterViewExtension.cpp
+++ b/Engine/Plugins/Experimental/Water/Source/Runtime/Private/WaterViewExtension.cpp
@@ -160,65 +160,64 @@ void FWaterViewExtension::SetupViewFamily(FSceneViewFamily& InViewFamily)
 		TArray<FWaterBodyData> WaterBodyData;
 		{
 			const int32 NumWaterBodies =  WaterBodyManager->NumWaterBodies();
-			WaterIndirection.Reserve(NumWaterBodies);
+			WaterIndirection.SetNumZeroed(NumWaterBodies);
 			WaterBodyData.Reserve(NumWaterBodies);
 
-			TArray<const UGerstnerWaterWaves*> AllGerstnerWaves;
+			TMap<const UGerstnerWaterWaves*, int32> GerstnerWavesIndices;
 
-			WaterBodyManager->ForEachWaterBodyComponent([&WaterIndirection, &WaveData, &WaterBodyData, &AllGerstnerWaves](UWaterBodyComponent* WaterBodyComponent)
+			WaterBodyManager->ForEachWaterBodyComponent([&WaterIndirection, &WaveData, &WaterBodyData, &GerstnerWavesIndices](UWaterBodyComponent* WaterBodyComponent)
 			{
 				const int32 WaterZoneIndex = WaterBodyComponent->GetWaterZone() ? WaterBodyComponent->GetWaterZone()->GetWaterZoneIndex() : -1;
 
-				FWaterIndirection& WaterIndirectionEntry = WaterIndirection.AddZeroed_GetRef();
+				check(WaterBodyComponent->GetWaterBodyIndex() < WaterIndirection.Num());
+				FWaterIndirection& WaterIndirectionEntry = WaterIndirection[WaterBodyComponent->GetWaterBodyIndex()];
 				WaterIndirectionEntry.WaterZoneIndex = WaterZoneIndex;
 				WaterIndirectionEntry.WaterBodyDataIndex = WaterBodyData.Num();
 
 				FWaterBodyData& WaterBodyDataEntry = WaterBodyData.AddZeroed_GetRef();
 				WaterBodyDataEntry.TargetWaveMaskDepth = WaterBodyComponent->TargetWaveMaskDepth;
 
-				// #todo_water: reuse waves instead of duplicating.
 				if (WaterBodyComponent->HasWaves())
 				{
 					const UWaterWavesBase* WaterWavesBase = WaterBodyComponent->GetWaterWaves();
 					check(WaterWavesBase != nullptr);
 					if (const UGerstnerWaterWaves* GerstnerWaves = Cast<const UGerstnerWaterWaves>(WaterWavesBase->GetWaterWaves()))
 					{
-						int32 WaveDataIndex = AllGerstnerWaves.IndexOfByKey(GerstnerWaves);
-						if (WaveDataIndex == INDEX_NONE)
+						int32* WaveDataIndex = GerstnerWavesIndices.Find(GerstnerWaves);
+
+						if (WaveDataIndex == nullptr)
 						{
-							WaveDataIndex = AllGerstnerWaves.Num();
-							AllGerstnerWaves.Add(GerstnerWaves);
+							// Where the data for this set of waves starts
+							const int32 WaveDataBase = WaveData.Num();
+
+							WaveDataIndex = &GerstnerWavesIndices.Add(GerstnerWaves, WaveDataBase);
+
+							// Some max value
+							constexpr int32 MaxWavesPerGerstnerWaves = 4096;
+
+							const TArray<FGerstnerWave>& Waves = GerstnerWaves->GetGerstnerWaves();
+							
+							// Allocate for the waves in this water body
+							const int32 NumWaves = FMath::Min(Waves.Num(), MaxWavesPerGerstnerWaves);
+							WaveData.AddZeroed(NumWaves);
+
+							for (int32 WaveIndex = 0; WaveIndex < NumWaves; WaveIndex++)
+							{
+								const uint32 WavesDataIndex = WaveDataBase + WaveIndex;
+								WaveData[WavesDataIndex] = FGerstnerWaveData(Waves[WaveIndex]);
+							}
 						}
 
 						const TArray<FGerstnerWave>& Waves = GerstnerWaves->GetGerstnerWaves();
 
-						WaterIndirectionEntry.WaveDataIndex = WaveDataIndex;
+						check(WaveDataIndex);
+
+						WaterIndirectionEntry.WaveDataIndex = *WaveDataIndex;
 						WaterBodyDataEntry.NumWaves = Waves.Num();
 					}
 				}
 				return true;
 			});
-
-			for (int32 GerstnerWavesIndex = 0; GerstnerWavesIndex < AllGerstnerWaves.Num(); ++GerstnerWavesIndex)
-			{
-				// Some max value
-				constexpr int32 MaxWavesPerGerstnerWaves = 4096;
-
-				const TArray<FGerstnerWave>& Waves = AllGerstnerWaves[GerstnerWavesIndex]->GetGerstnerWaves();
-				
-				// Where the data for this set of waves starts
-				const int32 WaveDataBase = WaveData.Num();
-
-				// Allocate for the waves in this water body
-				const int32 NumWaves = FMath::Min(Waves.Num(), MaxWavesPerGerstnerWaves);
-				WaveData.AddZeroed(NumWaves);
-
-				for (int32 WaveIndex = 0; WaveIndex < NumWaves; WaveIndex++)
-				{
-					const uint32 WavesDataIndex = WaveDataBase + WaveIndex;
-					WaveData[WavesDataIndex] = FGerstnerWaveData(Waves[WaveIndex]);
-				}
-			}
 		}
 
 		TResourceArray<FVector4f> WaterIndirectionBuffer;
