Water: fixed preallocation size of WaterIndirectionBuffer to property support up the max water body index, rather than the number of water bodies. Since we don't reassign indices when water bodies are removed, it might be possible to have an index > NumWaterBodies.

diff --git a/Engine/Plugins/Experimental/Water/Source/Runtime/Private/WaterBodyManager.cpp b/Engine/Plugins/Experimental/Water/Source/Runtime/Private/WaterBodyManager.cpp
--- a/Engine/Plugins/Experimental/Water/Source/Runtime/Private/WaterBodyManager.cpp
+++ b/Engine/Plugins/Experimental/Water/Source/Runtime/Private/WaterBodyManager.cpp
@@ -23,7 +23,8 @@ void FWaterBodyManager::Deinitialize()
 int32 FWaterBodyManager::AddWaterBodyComponent(UWaterBodyComponent* InWaterBodyComponent)
 {
 	RequestWaveDataRebuild();
-	const int32 WaterIndex = WaterBodyComponents.Register(InWaterBodyComponent);
+	int32 LowestFreeIndex = 0;
+	const int32 WaterIndex = WaterBodyComponents.EmplaceAtLowestFreeIndex(LowestFreeIndex, InWaterBodyComponent);
 	OnWaterBodyAdded.Broadcast(InWaterBodyComponent);
 	return WaterIndex;
 }
@@ -31,20 +32,21 @@ int32 FWaterBodyManager::AddWaterBodyComponent(UWaterBodyComponent* InWaterBodyC
 void FWaterBodyManager::RemoveWaterBodyComponent(UWaterBodyComponent* InWaterBodyComponent)
 {
 	RequestWaveDataRebuild();
-	WaterBodyComponents.Unregister(InWaterBodyComponent, InWaterBodyComponent->GetWaterBodyIndex());
+	WaterBodyComponents.RemoveAt(InWaterBodyComponent->GetWaterBodyIndex());
 	OnWaterBodyRemoved.Broadcast(InWaterBodyComponent);
 }
 
 int32 FWaterBodyManager::AddWaterZone(AWaterZone* InWaterZone)
 {
 	RequestGPUDataRebuild();
-	return WaterZones.Register(InWaterZone);
+	int32 LowestFreeIndex = 0;
+	return WaterZones.EmplaceAtLowestFreeIndex(LowestFreeIndex, InWaterZone);
 }
 
 void FWaterBodyManager::RemoveWaterZone(AWaterZone* InWaterZone)
 {
 	RequestGPUDataRebuild();
-	WaterZones.Unregister(InWaterZone, InWaterZone->GetWaterZoneIndex());
+	WaterZones.RemoveAt(InWaterZone->GetWaterZoneIndex());
 }
 
 void FWaterBodyManager::RequestGPUDataRebuild()
@@ -70,7 +72,7 @@ void FWaterBodyManager::RequestWaveDataRebuild()
 
 void FWaterBodyManager::ForEachWaterBodyComponent(TFunctionRef<bool(UWaterBodyComponent*)> Pred) const
 {
-	for (UWaterBodyComponent* WaterBodyComponent : WaterBodyComponents.Elements)
+	for (UWaterBodyComponent* WaterBodyComponent : WaterBodyComponents)
 	{
 		if (WaterBodyComponent)
 		{
@@ -92,7 +94,7 @@ void FWaterBodyManager::ForEachWaterBodyComponent(const UWorld* World, TFunction
 
 void FWaterBodyManager::ForEachWaterZone(TFunctionRef<bool(AWaterZone*)> Pred) const
 {
-	for (AWaterZone* WaterZone : WaterZones.Elements)
+	for (AWaterZone* WaterZone : WaterZones)
 	{
 		if (WaterZone)
 		{
@@ -112,3 +114,9 @@ void FWaterBodyManager::ForEachWaterZone(const UWorld* World, TFunctionRef<bool(
 	}
 }
 
+void FWaterBodyManager::Shrink()
+{
+	WaterBodyComponents.Shrink();
+	WaterZones.Shrink();
+}
+
diff --git a/Engine/Plugins/Experimental/Water/Source/Runtime/Private/WaterViewExtension.cpp b/Engine/Plugins/Experimental/Water/Source/Runtime/Private/WaterViewExtension.cpp
--- a/Engine/Plugins/Experimental/Water/Source/Runtime/Private/WaterViewExtension.cpp
+++ b/Engine/Plugins/Experimental/Water/Source/Runtime/Private/WaterViewExtension.cpp
@@ -74,9 +74,12 @@ void FWaterViewExtension::SetupViewFamily(FSceneViewFamily& InViewFamily)
 
 		const UWorld* WorldPtr = GetWorld();
 		check(WorldPtr != nullptr);
-		const FWaterBodyManager* WaterBodyManager = UWaterSubsystem::GetWaterBodyManager(WorldPtr);
+		FWaterBodyManager* WaterBodyManager = UWaterSubsystem::GetWaterBodyManager(WorldPtr);
 		check(WaterBodyManager);
 
+		// Shrink the water manager storage to avoid over-preallocating in the WaterIndirectionBuffer.
+		WaterBodyManager->Shrink();
+
 
 		struct FWaterIndirection
 		{
@@ -160,7 +163,8 @@ void FWaterViewExtension::SetupViewFamily(FSceneViewFamily& InViewFamily)
 		TArray<FWaterBodyData> WaterBodyData;
 		{
 			const int32 NumWaterBodies =  WaterBodyManager->NumWaterBodies();
-			WaterIndirection.SetNumZeroed(NumWaterBodies);
+			// Pre-set up to the max water body index. Some entries may be empty and NumWaterBodies != MaxIndex
+			WaterIndirection.SetNumZeroed(WaterBodyManager->MaxWaterBodyIndex());
 			WaterBodyData.Reserve(NumWaterBodies);
 
 			TMap<const UGerstnerWaterWaves*, int32> GerstnerWavesIndices;
diff --git a/Engine/Plugins/Experimental/Water/Source/Runtime/Public/WaterBodyManager.h b/Engine/Plugins/Experimental/Water/Source/Runtime/Public/WaterBodyManager.h
--- a/Engine/Plugins/Experimental/Water/Source/Runtime/Public/WaterBodyManager.h
+++ b/Engine/Plugins/Experimental/Water/Source/Runtime/Public/WaterBodyManager.h
@@ -49,11 +49,15 @@ public:
 	void ForEachWaterZone(TFunctionRef<bool(AWaterZone*)> Pred) const;
 	static void ForEachWaterZone(const UWorld* World, TFunctionRef<bool(AWaterZone*)> Pred);
 
-	bool HasAnyWaterBodies() const { return WaterBodyComponents.Num > 0; }
+	bool HasAnyWaterBodies() const { return WaterBodyComponents.Num() > 0; }
 
-	int32 NumWaterBodies() const { return WaterBodyComponents.Num; }
+	int32 NumWaterBodies() const { return WaterBodyComponents.Num(); }
+	int32 MaxWaterBodyIndex() const { return WaterBodyComponents.GetMaxIndex(); }
 
-	int32 NumWaterZones() const { return WaterZones.Num; }
+	/** Shrinks the sparse array storage for water body components and water zones. Ensures that MaxIndex == MaxAllocatedIndex */
+	void Shrink();
+
+	int32 NumWaterZones() const { return WaterZones.Num(); }
 
 	FWaterViewExtension* GetWaterViewExtension() { return WaterViewExtension.Get(); }
 
@@ -62,64 +66,11 @@ public:
 	FWaterBodyEvent OnWaterBodyRemoved;
 
 private:
-
-	/**
-	* TWaterContainer<T> wraps a TArray<T> to support reusing dead indices while always maintaining stability for existing indices.
-	*
-	* The Elements array may contain nullptr entries.
-	*/
-	template <typename T>
-	class TWaterContainer 
-	{
-	public:
-		int32 Register(T* InElement)
-		{
-			int32 Index = INDEX_NONE;
-			if (UnusedIndices.Num())
-			{
-				Index = UnusedIndices.Pop(/*bAllowShrinking = */false);
-				check(Elements[Index] == nullptr);
-				Elements[Index] = InElement;
-			}
-			else
-			{
-				Index = Elements.Add(InElement);
-			}
-
-			++Num;
-			check(Num <= Elements.Num());
-
-			check(Index != INDEX_NONE);
-			return Index;
-		}
-
-		void Unregister(const T* InElement, int32 OldIndex)
-		{
-			check(OldIndex != INDEX_NONE);
-			UnusedIndices.Add(OldIndex);
-			Elements[OldIndex] = nullptr;
-
-			--Num;
-			check(Num >= 0);
-
-			// Empty all arrays once there are no more elements
-			if (UnusedIndices.Num() == Elements.Num())
-			{
-				UnusedIndices.Empty();
-				Elements.Empty();
-			}
-		}
-
-		TArray<T*> Elements;
-		TArray<int32> UnusedIndices;
-		int32 Num = 0;
-	};
-
 	/** List of components registered to this manager. */
-	TWaterContainer<UWaterBodyComponent> WaterBodyComponents;
+	TSparseArray<UWaterBodyComponent*> WaterBodyComponents;
 
 	/** List of Water zones registered to this manager. */
-	TWaterContainer<AWaterZone> WaterZones;
+	TSparseArray<AWaterZone*> WaterZones;
 
 	float GlobalMaxWaveHeight = 0.0f;
 
