Water: Fixed assumption in the water quad tree that a river couldn't overlap with both a lake and ocean of differing priorities in the same quad tree tile. This could lead to a random choice of transition type/material not respecting the overlap priority.

diff --git a/Engine/Plugins/Experimental/Water/Source/Runtime/Private/WaterQuadTree.cpp b/Engine/Plugins/Experimental/Water/Source/Runtime/Private/WaterQuadTree.cpp
--- a/Engine/Plugins/Experimental/Water/Source/Runtime/Private/WaterQuadTree.cpp
+++ b/Engine/Plugins/Experimental/Water/Source/Runtime/Private/WaterQuadTree.cpp
@@ -334,13 +334,19 @@ void FWaterQuadTree::FNode::AddNodes(FNodeData& InNodeData, const FBox& InMeshBo
 	// Check is this node should be marked for material overlap 
 	if (InWaterBody.IsRiver() && ((ThisWaterBody.IsLake() && InWaterBody.RiverToLakeMaterial) || (ThisWaterBody.IsOcean() && InWaterBody.RiverToOceanMaterial)))
 	{
-		// If the incoming water body is a river with a transition, and if the existing water body is either ocean or lake, we set transition water body index
-		TransitionWaterBodyIndex = (uint16)WaterBodyIndex;
+		// If the incoming water body is a river with a transition, and if the existing water body is either ocean or lake, we set transition water body index, but only if the new water body has a higher priority
+		if ((TransitionWaterBodyIndex == 0) || (ThisWaterBody.Priority >= InNodeData.WaterBodyRenderData[TransitionWaterBodyIndex].Priority))
+		{
+			TransitionWaterBodyIndex = (uint16)WaterBodyIndex;
+		}
 	}
 	else if (ThisWaterBody.IsRiver() && ((InWaterBody.IsLake() && ThisWaterBody.RiverToLakeMaterial) || (InWaterBody.IsOcean() && ThisWaterBody.RiverToOceanMaterial)))
 	{
-		// If the existing water body is a river with a transition, and if the incoming water body is either ocean or lake, we set transition water body index
-		TransitionWaterBodyIndex = (uint16)InWaterBodyIndex;
+		// If the existing water body is a river with a transition, and if the incoming water body is either ocean or lake, we set transition water body index, but only if the new water body has a higher priority
+		if (TransitionWaterBodyIndex == 0 || InWaterBody.Priority >= InNodeData.WaterBodyRenderData[TransitionWaterBodyIndex].Priority)
+		{
+			TransitionWaterBodyIndex = (uint16)InWaterBodyIndex;
+		}
 	}
 
 	// Assign the render data here (based on priority)
