Water: avoid updating water bodies when exclusion volumes update their overlapping water bodies. Instead, only update the water body directly when added/removed or when the exclusion volume is modified directly in the editor or loaded.

diff --git a/Engine/Plugins/Experimental/Water/Source/Runtime/Private/WaterBodyExclusionVolume.cpp b/Engine/Plugins/Experimental/Water/Source/Runtime/Private/WaterBodyExclusionVolume.cpp
--- a/Engine/Plugins/Experimental/Water/Source/Runtime/Private/WaterBodyExclusionVolume.cpp
+++ b/Engine/Plugins/Experimental/Water/Source/Runtime/Private/WaterBodyExclusionVolume.cpp
@@ -84,10 +84,6 @@ void AWaterBodyExclusionVolume::UpdateOverlappingWaterBodies(const FWaterExclusi
 			{
 				ExistingBody->RemoveExclusionVolume(this);
 			}
-			else
-			{
-				ExistingBody->OnWaterBodyChanged(WaterBodyChangedParams);
-			}
 		}
 	}
 }
@@ -119,6 +115,26 @@ void AWaterBodyExclusionVolume::PostRegisterAllComponents()
 	FWaterExclusionVolumeChangedParams Params;
 	Params.bUserTriggered = false;
 	UpdateOverlappingWaterBodies(Params);
+
+	UpdateAffectedWaterBodyCollisions(Params);
+}
+
+void AWaterBodyExclusionVolume::UpdateAffectedWaterBodyCollisions(const FWaterExclusionVolumeChangedParams& Params)
+{
+	TSoftObjectPtr<AWaterBodyExclusionVolume> SoftThis(this);
+
+	FOnWaterBodyChangedParams WaterBodyChangedParams(Params.PropertyChangedEvent);
+	WaterBodyChangedParams.bUserTriggered = Params.bUserTriggered;
+	WaterBodyChangedParams.bShapeOrPositionChanged = false;
+	WaterBodyChangedParams.bWeightmapSettingsChanged = false;
+	FWaterBodyManager::ForEachWaterBodyComponent(GetWorld(), [SoftThis, &WaterBodyChangedParams](UWaterBodyComponent* WaterBodyComponent)
+	{
+		if (WaterBodyComponent->ContainsExclusionVolume(SoftThis))
+		{
+			WaterBodyComponent->OnWaterBodyChanged(WaterBodyChangedParams);
+		}
+		return true;
+	});
 }
 
 void AWaterBodyExclusionVolume::PostLoad()
@@ -189,6 +205,8 @@ void AWaterBodyExclusionVolume::PostEditMove(bool bFinished)
 	Params.PropertyChangedEvent.ChangeType = bFinished ? EPropertyChangeType::ValueSet : EPropertyChangeType::Interactive;
 	Params.bUserTriggered = true;
 	UpdateOverlappingWaterBodies(Params);
+
+	UpdateAffectedWaterBodyCollisions(Params);
 }
 
 void AWaterBodyExclusionVolume::PostEditUndo()
@@ -198,6 +216,8 @@ void AWaterBodyExclusionVolume::PostEditUndo()
 	FWaterExclusionVolumeChangedParams Params;
 	Params.bUserTriggered = true;
 	UpdateOverlappingWaterBodies(Params);
+
+	UpdateAffectedWaterBodyCollisions(Params);
 }
 
 void AWaterBodyExclusionVolume::PostEditImport()
@@ -207,6 +227,8 @@ void AWaterBodyExclusionVolume::PostEditImport()
 	FWaterExclusionVolumeChangedParams Params;
 	Params.bUserTriggered = true;
 	UpdateOverlappingWaterBodies(Params);
+
+	UpdateAffectedWaterBodyCollisions(Params);
 }
 
 void AWaterBodyExclusionVolume::PostEditChangeProperty(FPropertyChangedEvent& PropertyChangedEvent)
@@ -216,6 +238,8 @@ void AWaterBodyExclusionVolume::PostEditChangeProperty(FPropertyChangedEvent& Pr
 	FWaterExclusionVolumeChangedParams Params(PropertyChangedEvent);
 	Params.bUserTriggered = true;
 	UpdateOverlappingWaterBodies(Params);
+
+	UpdateAffectedWaterBodyCollisions(Params);
 }
 
 FName AWaterBodyExclusionVolume::GetCustomIconName() const
diff --git a/Engine/Plugins/Experimental/Water/Source/Runtime/Public/WaterBodyExclusionVolume.h b/Engine/Plugins/Experimental/Water/Source/Runtime/Public/WaterBodyExclusionVolume.h
--- a/Engine/Plugins/Experimental/Water/Source/Runtime/Public/WaterBodyExclusionVolume.h
+++ b/Engine/Plugins/Experimental/Water/Source/Runtime/Public/WaterBodyExclusionVolume.h
@@ -66,6 +66,8 @@ protected:
 	virtual FName GetCustomIconName() const override;
 #endif // WITH_EDITOR
 
+	/** Updates all water bodies affected by this exclusion volume to rebuild due to a change in this exclusion volume. */
+	void UpdateAffectedWaterBodyCollisions(const FWaterExclusionVolumeChangedParams& Params);
 public:
 	/** Determines the behavior of the WaterBodies list. */
 	UPROPERTY(EditAnywhere, Category = Water)
