Make sure to mark the WaterZone's BoundsComponent as editor only.  It is an editor only, transient object.  Not being marked as editor only was causing validation errors when trying to set a WaterZone actor as editor only.

diff --git a/Engine/Plugins/Experimental/Water/Source/Runtime/Private/WaterZoneActor.cpp b/Engine/Plugins/Experimental/Water/Source/Runtime/Private/WaterZoneActor.cpp
--- a/Engine/Plugins/Experimental/Water/Source/Runtime/Private/WaterZoneActor.cpp
+++ b/Engine/Plugins/Experimental/Water/Source/Runtime/Private/WaterZoneActor.cpp
@@ -60,6 +60,7 @@ AWaterZone::AWaterZone(const FObjectInitializer& Initializer)
 		BoundsComponent->SetupAttachment(WaterMesh);
 		// Bounds component extent is half-extent, ZoneExtent is full extent.
 		BoundsComponent->SetBoxExtent(FVector(ZoneExtent / 2., 8192.));
+		BoundsComponent->bIsEditorOnly = true;
 	}
 
 	if (GIsEditor && !IsTemplate())
