Water: exposed a function to set river transition materials from blueprint and added a water mesh update when setting these materials to ensure the new MIDs are set.

diff --git a/Engine/Plugins/Experimental/Water/Source/Runtime/Private/WaterBodyComponent.cpp b/Engine/Plugins/Experimental/Water/Source/Runtime/Private/WaterBodyComponent.cpp
--- a/Engine/Plugins/Experimental/Water/Source/Runtime/Private/WaterBodyComponent.cpp
+++ b/Engine/Plugins/Experimental/Water/Source/Runtime/Private/WaterBodyComponent.cpp
@@ -426,6 +426,22 @@ AWaterZone* UWaterBodyComponent::GetWaterZone() const
 	return OwningWaterZone.Get();
 }
 
+void UWaterBodyComponent::MarkOwningWaterZoneForRebuild(EWaterZoneRebuildFlags InRebuildFlags, bool bInOnlyWithinWaterBodyBounds) const
+{
+	if (AWaterZone* WaterZone = GetWaterZone())
+	{
+		if (bInOnlyWithinWaterBodyBounds)
+		{
+			const FBox WaterBodyBounds = Bounds.GetBox();
+			WaterZone->MarkForRebuild(InRebuildFlags, FBox2D(FVector2D(WaterBodyBounds.Min), FVector2D(WaterBodyBounds.Max)));
+		}
+		else
+		{
+			WaterZone->MarkForRebuild(InRebuildFlags);
+		}
+	}
+}
+
 void UWaterBodyComponent::SetWaterZoneOverride(const TSoftObjectPtr<AWaterZone>& InWaterZoneOverride)
 {
 	WaterZoneOverride = InWaterZoneOverride;
@@ -864,6 +880,9 @@ void UWaterBodyComponent::UpdateMaterialInstances()
 	CreateOrUpdateWaterInfoMID();
 	CreateOrUpdateWaterStaticMeshMID();
 	CreateOrUpdateUnderwaterPostProcessMID();
+
+	// Update the water mesh since it will not contain the up-to-date MIDs:
+	MarkOwningWaterZoneForRebuild(EWaterZoneRebuildFlags::UpdateWaterMesh);
 }
 
 bool UWaterBodyComponent::UpdateWaterHeight()
@@ -1033,8 +1052,7 @@ void UWaterBodyComponent::UpdateComponentVisibility(bool bAllowWaterZoneRebuild)
 					RebuildFlags |= EWaterZoneRebuildFlags::UpdateWaterInfoTexture;
 				}
 
-				const FBox WaterBodyBounds = Bounds.GetBox();
-				WaterZone->MarkForRebuild(RebuildFlags, FBox2D(FVector2D(WaterBodyBounds.Min), FVector2D(WaterBodyBounds.Max)));
+				MarkOwningWaterZoneForRebuild(RebuildFlags);
 			}
 		}
 	}
@@ -1924,7 +1942,7 @@ void UWaterBodyComponent::OnWaterBodyRenderDataUpdated()
 	{
 		if (AWaterZone* WaterZone = GetWaterZone())
 		{
-			WaterZone->MarkForRebuild(EWaterZoneRebuildFlags::UpdateWaterInfoTexture);
+			MarkOwningWaterZoneForRebuild(EWaterZoneRebuildFlags::UpdateWaterInfoTexture, /* bOnlyWithinWaterBodyBounds = */ false);
 		}
 	}
 
diff --git a/Engine/Plugins/Experimental/Water/Source/Runtime/Private/WaterBodyRiverComponent.cpp b/Engine/Plugins/Experimental/Water/Source/Runtime/Private/WaterBodyRiverComponent.cpp
--- a/Engine/Plugins/Experimental/Water/Source/Runtime/Private/WaterBodyRiverComponent.cpp
+++ b/Engine/Plugins/Experimental/Water/Source/Runtime/Private/WaterBodyRiverComponent.cpp
@@ -344,10 +344,11 @@ FBoxSphereBounds UWaterBodyRiverComponent::CalcBounds(const FTransform& LocalToW
 
 void UWaterBodyRiverComponent::UpdateMaterialInstances()
 {
-	Super::UpdateMaterialInstances();
-
 	CreateOrUpdateLakeTransitionMID();
 	CreateOrUpdateOceanTransitionMID();
+
+	// Must be called after the transition MIDs are created. Super::UpdateMaterialInstances will rebuild the water mesh if necessary to push new MIDs.
+	Super::UpdateMaterialInstances();
 }
 
 void UWaterBodyRiverComponent::SetLakeTransitionMaterial(UMaterialInterface* InMaterial)
@@ -362,6 +363,19 @@ void UWaterBodyRiverComponent::SetOceanTransitionMaterial(UMaterialInterface* In
 	UpdateMaterialInstances();
 }
 
+void UWaterBodyRiverComponent::SetLakeAndOceanTransitionMaterials(UMaterialInterface* InLakeTransition, UMaterialInterface* InOceanTransition)
+{
+	const bool bUpdateInstances = LakeTransitionMaterial != InLakeTransition || OceanTransitionMaterial != InOceanTransition;
+
+	LakeTransitionMaterial = InLakeTransition;
+	OceanTransitionMaterial = InOceanTransition;
+
+	if (bUpdateInstances)
+	{
+		UpdateMaterialInstances();
+	}
+}
+
 void UWaterBodyRiverComponent::Reset()
 {
 	for (USplineMeshComponent* Comp : SplineMeshComponents)
diff --git a/Engine/Plugins/Experimental/Water/Source/Runtime/Public/WaterBodyComponent.h b/Engine/Plugins/Experimental/Water/Source/Runtime/Public/WaterBodyComponent.h
--- a/Engine/Plugins/Experimental/Water/Source/Runtime/Public/WaterBodyComponent.h
+++ b/Engine/Plugins/Experimental/Water/Source/Runtime/Public/WaterBodyComponent.h
@@ -9,6 +9,7 @@
 #include "WaterCurveSettings.h"
 #include "WaterBodyStaticMeshSettings.h"
 #include "WaterSplineMetadata.h"
+#include "WaterZoneActor.h"
 #include "WaterBodyTypes.h"
 
 class AWaterBody;
@@ -221,9 +222,11 @@ public:
 	UMaterialInterface* GetWaterStaticMeshMaterial() const { return WaterStaticMeshMaterial; }
 
 	/** Sets water material */
+	UFUNCTION(BlueprintCallable, Category = Rendering)
 	void SetWaterMaterial(UMaterialInterface* InMaterial);
 
 	/** Sets water static mesh material */
+	UFUNCTION(BlueprintCallable, Category = Rendering)
 	void SetWaterStaticMeshMaterial(UMaterialInterface* InMaterial);
 
 	/** Returns water MID */
@@ -248,6 +251,7 @@ public:
 	UMaterialInstanceDynamic* GetWaterInfoMaterialInstance();
 
 	/** Sets under water post process material */
+	UFUNCTION(BlueprintCallable, Category = Rendering)
 	void SetUnderwaterPostProcessMaterial(UMaterialInterface* InMaterial);
 
 	UFUNCTION(BlueprintCallable, Category = Rendering)
@@ -453,6 +457,12 @@ protected:
 
 	virtual void OnUpdateBody(bool bWithExclusionVolumes) {}
 
+	/** 
+	 * Marks the owning water zone for rebuild. 
+	 * If bOnlyWithinWaterBodyBounds is set, updates to the water zone that aren't relevant within the bounds of the water body are suppressed.
+	 */
+	void MarkOwningWaterZoneForRebuild(EWaterZoneRebuildFlags InRebuildFlags, bool bInOnlyWithinWaterBodyBounds = true) const;
+
 	/** Called when the WaterBodyActor has had all its components registered. */
 	virtual void OnPostRegisterAllComponents();
 
diff --git a/Engine/Plugins/Experimental/Water/Source/Runtime/Public/WaterBodyRiverComponent.h b/Engine/Plugins/Experimental/Water/Source/Runtime/Public/WaterBodyRiverComponent.h
--- a/Engine/Plugins/Experimental/Water/Source/Runtime/Public/WaterBodyRiverComponent.h
+++ b/Engine/Plugins/Experimental/Water/Source/Runtime/Public/WaterBodyRiverComponent.h
@@ -27,9 +27,15 @@ public:
 	virtual TArray<UPrimitiveComponent*> GetBrushRenderableComponents() const override;
 #endif //WITH_EDITOR
 
+	UFUNCTION(BlueprintCallable, Category = Rendering)
 	void SetLakeTransitionMaterial(UMaterialInterface* InMat);
+
+	UFUNCTION(BlueprintCallable, Category = Rendering)
 	void SetOceanTransitionMaterial(UMaterialInterface* InMat);
 
+	UFUNCTION(BlueprintCallable, Category = Rendering)
+	void SetLakeAndOceanTransitionMaterials(UMaterialInterface* InLakeTransition, UMaterialInterface* InOceanTransition);
+
 protected:
 	/** UWaterBodyComponent Interface */
 	virtual void Reset() override;
