Water: avoid rebuilding the water mesh when rendering is disabled.

diff --git a/Engine/Plugins/Experimental/Water/Source/Runtime/Private/WaterMeshComponent.cpp b/Engine/Plugins/Experimental/Water/Source/Runtime/Private/WaterMeshComponent.cpp
--- a/Engine/Plugins/Experimental/Water/Source/Runtime/Private/WaterMeshComponent.cpp
+++ b/Engine/Plugins/Experimental/Water/Source/Runtime/Private/WaterMeshComponent.cpp
@@ -594,7 +594,7 @@ void UWaterMeshComponent::RebuildWaterMesh(float InTileSize, const FIntPoint& In
 
 void UWaterMeshComponent::Update()
 {
-	bIsEnabled = FWaterUtils::IsWaterMeshEnabled(/*bIsRenderThread = */false);
+	bIsEnabled = FWaterUtils::IsWaterMeshEnabled(/*bIsRenderThread = */false) && FApp::CanEverRender();
 
 	// Early out
 	if (!bIsEnabled)
