Water: disable water static mesh components loading on server to save memory.

diff --git a/Engine/Plugins/Experimental/Water/Source/Runtime/Public/WaterBodyMeshComponent.h b/Engine/Plugins/Experimental/Water/Source/Runtime/Public/WaterBodyMeshComponent.h
--- a/Engine/Plugins/Experimental/Water/Source/Runtime/Public/WaterBodyMeshComponent.h
+++ b/Engine/Plugins/Experimental/Water/Source/Runtime/Public/WaterBodyMeshComponent.h
@@ -26,5 +26,8 @@ protected:
 
 	// Exclude water mesh components from HLOD since we have a custom HLOD builder for water bodies.
 	virtual bool IsHLODRelevant() const override { return false; }
+
+	// Exclude water meshes from loading on the server since we don't render at all.
+	virtual bool NeedsLoadForServer() const override { return false; }
 };
 
