Water: Expose a method to set the far mesh material from blueprint on the water zone.

diff --git a/Engine/Plugins/Experimental/Water/Source/Runtime/Private/WaterZoneActor.cpp b/Engine/Plugins/Experimental/Water/Source/Runtime/Private/WaterZoneActor.cpp
--- a/Engine/Plugins/Experimental/Water/Source/Runtime/Private/WaterZoneActor.cpp
+++ b/Engine/Plugins/Experimental/Water/Source/Runtime/Private/WaterZoneActor.cpp
@@ -376,9 +376,18 @@ FBox AWaterZone::GetStreamingBounds() const
 {
 	return GetZoneBounds();
 }
-
 #endif // WITH_EDITOR
 
+void AWaterZone::SetFarMeshMaterial(UMaterialInterface* InFarDistanceMaterial)
+{
+	if (WaterMesh && InFarDistanceMaterial != WaterMesh->FarDistanceMaterial)
+	{
+		WaterMesh->FarDistanceMaterial = InFarDistanceMaterial;
+		MarkForRebuild(EWaterZoneRebuildFlags::UpdateWaterMesh);
+	}
+}
+
+
 void AWaterZone::OnExtentChanged()
 {
 	// Compute the new tile extent based on the new bounds
diff --git a/Engine/Plugins/Experimental/Water/Source/Runtime/Public/WaterZoneActor.h b/Engine/Plugins/Experimental/Water/Source/Runtime/Public/WaterZoneActor.h
--- a/Engine/Plugins/Experimental/Water/Source/Runtime/Public/WaterZoneActor.h
+++ b/Engine/Plugins/Experimental/Water/Source/Runtime/Public/WaterZoneActor.h
@@ -84,6 +84,9 @@ public:
 	virtual FBox GetStreamingBounds() const override;
 #endif //WITH_EDITOR
 
+	UFUNCTION(BlueprintCallable, Category=Rendering)
+	void SetFarMeshMaterial(UMaterialInterface* InFarMaterial);
+
 	UE_DEPRECATED(5.3, "Function renamed to GetDynamicWaterInfoCenter")
 	FVector GetTessellatedWaterMeshCenter() const { return GetDynamicWaterInfoCenter(); }
 
