Make sure there is a valid body setup before creating the water body static meshes.

diff --git a/Engine/Plugins/Experimental/Water/Source/Runtime/Private/WaterBodyMeshBuilder.cpp b/Engine/Plugins/Experimental/Water/Source/Runtime/Private/WaterBodyMeshBuilder.cpp
--- a/Engine/Plugins/Experimental/Water/Source/Runtime/Private/WaterBodyMeshBuilder.cpp
+++ b/Engine/Plugins/Experimental/Water/Source/Runtime/Private/WaterBodyMeshBuilder.cpp
@@ -570,12 +570,16 @@ UStaticMesh* FWaterBodyMeshBuilder::CreateUStaticMesh(UObject* Outer, FName Mesh
 	// Disable ray tracing as we don't want water meshes to show up in LumenScene.
 	StaticMesh->bSupportRayTracing = false;
 
+	// Always call CreateBodySetup before attempting to modify it.  The BodySetup is normally created when the mesh data is built, but can be created ahead of time to set default data
+	StaticMesh->CreateBodySetup();
+
 	if (UBodySetup* BodySetup = StaticMesh->GetBodySetup())
 	{
 		BodySetup->DefaultInstance.SetCollisionProfileName(UCollisionProfile::NoCollision_ProfileName);
 		BodySetup->CollisionTraceFlag = CTF_UseSimpleAsComplex;
 		// We won't ever enable collisions (since collisions are handled by the dedicated water body collision components), ensure we don't even cook or load any collision data on this mesh: 
 		BodySetup->bNeverNeedsCookedCollisionData = true;
+		BodySetup->bHasCookedCollisionData = false;
 	}
 	return StaticMesh;
 }
diff --git a/Engine/Plugins/Experimental/Water/Source/Runtime/Private/WaterBodyMeshComponent.cpp b/Engine/Plugins/Experimental/Water/Source/Runtime/Private/WaterBodyMeshComponent.cpp
--- a/Engine/Plugins/Experimental/Water/Source/Runtime/Private/WaterBodyMeshComponent.cpp
+++ b/Engine/Plugins/Experimental/Water/Source/Runtime/Private/WaterBodyMeshComponent.cpp
@@ -82,3 +82,25 @@ bool UWaterBodyMeshComponent::CanCreateSceneProxy() const
 
 	return true;
 }
+
+
+#if WITH_EDITOR
+void UWaterBodyMeshComponent::PostLoad()
+{
+	Super::PostLoad();
+
+	// Fix for a bug that existing where the body setup was not created with bNeverNeedsCookedCollisionData, and also address an issue where bNeverNeedsCookedCollisionData was allowing meshes to generate cooked data.
+	UBodySetup* BodySetup = GetBodySetup();
+	UStaticMesh* Mesh = GetStaticMesh();
+	bool bResetBodySetup = Mesh != nullptr && (BodySetup == nullptr || !BodySetup->bNeverNeedsCookedCollisionData || BodySetup->bHasCookedCollisionData);
+	if (bResetBodySetup)
+	{
+		Mesh->CreateBodySetup();
+		BodySetup = GetBodySetup();
+		BodySetup->bNeverNeedsCookedCollisionData = true;
+		BodySetup->bHasCookedCollisionData = false;
+		BodySetup->InvalidatePhysicsData();
+	}
+}
+#endif // WITH_EDITOR
+
diff --git a/Engine/Plugins/Experimental/Water/Source/Runtime/Public/WaterBodyMeshComponent.h b/Engine/Plugins/Experimental/Water/Source/Runtime/Public/WaterBodyMeshComponent.h
--- a/Engine/Plugins/Experimental/Water/Source/Runtime/Public/WaterBodyMeshComponent.h
+++ b/Engine/Plugins/Experimental/Water/Source/Runtime/Public/WaterBodyMeshComponent.h
@@ -29,5 +29,9 @@ protected:
 
 	// Exclude water meshes from loading on the server since we don't render at all.
 	virtual bool NeedsLoadForServer() const override { return false; }
+
+#if WITH_EDITOR
+	virtual void PostLoad() override;
+#endif
 };
 
