Fix crash on check(World) in FSolverSafeWaterBodyData::FSolverSafeWaterBodyData

diff --git a/Engine/Plugins/Experimental/Water/Source/Runtime/Private/WaterBodyTypes.cpp b/Engine/Plugins/Experimental/Water/Source/Runtime/Private/WaterBodyTypes.cpp
--- a/Engine/Plugins/Experimental/Water/Source/Runtime/Private/WaterBodyTypes.cpp
+++ b/Engine/Plugins/Experimental/Water/Source/Runtime/Private/WaterBodyTypes.cpp
@@ -118,31 +118,32 @@ FSolverSafeWaterBodyData::FSolverSafeWaterBodyData(UWaterBodyComponent* WaterBod
 	if (WaterBodyComponent)
 	{
 		World = WaterBodyComponent->GetWorld();
-		check(World); // ?
-
-		if (ALandscapeProxy* LandscapeProxyActor = WaterBodyComponent->FindLandscape())
+		if (World)
 		{
-			if (ULandscapeInfo* Info = LandscapeProxyActor->GetLandscapeInfo())
+			if (ALandscapeProxy* LandscapeProxyActor = WaterBodyComponent->FindLandscape())
 			{
-				for (auto CollisionComponent : Info->XYtoCollisionComponentMap)
+				if (ULandscapeInfo* Info = LandscapeProxyActor->GetLandscapeInfo())
 				{
-					LandscapeCollisionComponents.Add(Cast<UPrimitiveComponent>(CollisionComponent.Value));
+					for (auto CollisionComponent : Info->XYtoCollisionComponentMap)
+					{
+						LandscapeCollisionComponents.Add(Cast<UPrimitiveComponent>(CollisionComponent.Value));
+					}
 				}
 			}
+			WaterSpline = WaterBodyComponent->GetWaterSpline();
+			WaterSplineMetadata = WaterBodyComponent->GetWaterSplineMetadata();
+			Location = WaterBodyComponent->GetComponentLocation();
+			WaterBodyType = WaterBodyComponent->GetWaterBodyType();
+			OceanHeightOffset = WaterBodyComponent->MaxWaveHeightOffset;
+			if (const UGerstnerWaterWaves* WavesObject = Cast< UGerstnerWaterWaves>(WaterBodyComponent->GetWaterWaves()))
+			{
+				WaveParams = WavesObject->GetGerstnerWaves();
+			}
+			WaveSpeedFactor = 1.f;
+			TargetWaveMaskDepth = WaterBodyComponent->TargetWaveMaskDepth;
+			MaxWaveHeight = WaterBodyComponent->GetMaxWaveHeight();
+			WaterBodyIndex = WaterBodyComponent->GetWaterBodyIndex();
 		}
-		WaterSpline = WaterBodyComponent->GetWaterSpline();
-		WaterSplineMetadata = WaterBodyComponent->GetWaterSplineMetadata();
-		Location = WaterBodyComponent->GetComponentLocation();
-		WaterBodyType = WaterBodyComponent->GetWaterBodyType();
-		OceanHeightOffset = WaterBodyComponent->MaxWaveHeightOffset;
-		if (const UGerstnerWaterWaves* WavesObject = Cast< UGerstnerWaterWaves>(WaterBodyComponent->GetWaterWaves()))
-		{
-			WaveParams = WavesObject->GetGerstnerWaves();
-		}
-		WaveSpeedFactor = 1.f;
-		TargetWaveMaskDepth = WaterBodyComponent->TargetWaveMaskDepth;
-		MaxWaveHeight = WaterBodyComponent->GetMaxWaveHeight();
-		WaterBodyIndex = WaterBodyComponent->GetWaterBodyIndex();
 	}
 }
 
