Water: removed the serialized extent in tiles from the quad tree. This is now recomputed dynamically based on the size of the water zone to prevent potential mismatch and remove some abiguous authority over the size of the water zone.

diff --git a/Engine/Plugins/Experimental/Water/Source/Runtime/Private/WaterMeshComponent.cpp b/Engine/Plugins/Experimental/Water/Source/Runtime/Private/WaterMeshComponent.cpp
--- a/Engine/Plugins/Experimental/Water/Source/Runtime/Private/WaterMeshComponent.cpp
+++ b/Engine/Plugins/Experimental/Water/Source/Runtime/Private/WaterMeshComponent.cpp
@@ -171,13 +171,6 @@ FMaterialRelevance UWaterMeshComponent::GetWaterMaterialRelevance(ERHIFeatureLev
 	return Result;
 }
 
-void UWaterMeshComponent::SetExtentInTiles(FIntPoint NewExtentInTiles)
-{
-	ExtentInTiles = NewExtentInTiles;
-	MarkWaterMeshGridDirty();
-	MarkRenderStateDirty();
-}
-
 void UWaterMeshComponent::SetDynamicWaterMeshCenter(const FVector2D& NewCenter)
 {
 	if (!DynamicWaterMeshCenter.Equals(NewCenter))
@@ -194,6 +187,25 @@ void UWaterMeshComponent::SetTileSize(float NewTileSize)
 	MarkRenderStateDirty();
 }
 
+FIntPoint UWaterMeshComponent::GetExtentInTiles() const
+{
+	if (const AWaterZone* WaterZone = GetOwner<AWaterZone>(); ensureMsgf(WaterZone != nullptr, TEXT("WaterMeshComponent is owned by an actor that is not a WaterZone. This is not supported!")))
+	{
+		const float MeshTileSize = TileSize;
+		const FVector2D ZoneHalfExtent = FVector2D(WaterZone->GetDynamicWaterInfoExtent()) / 2.0;
+		const int32 HalfExtentInTilesX = FMath::RoundUpToPowerOfTwo(ZoneHalfExtent.X / MeshTileSize);
+		const int32 HalfExtentInTilesY = FMath::RoundUpToPowerOfTwo(ZoneHalfExtent.Y / MeshTileSize);
+		const FIntPoint HalfExtentInTiles = FIntPoint(HalfExtentInTilesX, HalfExtentInTilesY);
+
+		// QuadTreeResolution caches the resolution so it is clearly visible to the user in the details panel. It represents the full extent rather than the half extent
+		QuadTreeResolution = HalfExtentInTiles * 2;
+
+		return HalfExtentInTiles;
+	}
+
+	return FIntPoint(1, 1);
+}
+
 FBoxSphereBounds UWaterMeshComponent::CalcBounds(const FTransform& LocalToWorld) const
 {
 	// Always return valid bounds (tree is initialized with invalid bounds and if nothing is inserted, the tree bounds will stay invalid)
@@ -221,25 +233,13 @@ void UWaterMeshComponent::RebuildWaterMesh(float InTileSize, const FIntPoint& In
 {
 	TRACE_CPUPROFILER_EVENT_SCOPE(RebuildWaterMesh);
 
+	AWaterZone* WaterZone = CastChecked<AWaterZone>(GetOwner());
+
 	// Position snapped to the grid
-	const FVector2D GridPosition = FVector2D(FMath::GridSnap<FVector::FReal>(GetComponentLocation().X, InTileSize), FMath::GridSnap<FVector::FReal>(GetComponentLocation().Y, InTileSize));
+	const FVector2D GridPosition = WaterZone->IsLocalOnlyTessellationEnabled() ? GetDynamicWaterMeshCenter() : FVector2D(FMath::GridSnap<FVector::FReal>(GetComponentLocation().X, InTileSize), FMath::GridSnap<FVector::FReal>(GetComponentLocation().Y, InTileSize));
 	const FVector2D WorldExtent = FVector2D(InTileSize * InExtentInTiles.X, InTileSize * InExtentInTiles.Y);
 
 	FBox2D WaterWorldBox = FBox2D(-WorldExtent + GridPosition, WorldExtent + GridPosition);
-
-	AWaterZone* WaterZone = Cast<AWaterZone>(GetOwner());
-
-	// when local tessellation is enabled, only use the overlap of the dynamic mesh region with the total water mesh extent to avoid updating the entire water zone.
-	FIntPoint ExtentInTilesToUse = InExtentInTiles;
-	if (WaterZone->IsLocalOnlyTessellationEnabled())
-	{
-		const FVector2D WorldspaceExtent = FVector2D(LocalTessellationExtentInTiles * InTileSize);
-		const FVector2D MeshPosition = GetDynamicWaterMeshCenter();
-		FBox2D DynamicWaterMeshBounds(MeshPosition - WorldspaceExtent, MeshPosition + WorldspaceExtent);
-
-		WaterWorldBox = DynamicWaterMeshBounds.Overlap(WaterWorldBox);
-		ExtentInTilesToUse = LocalTessellationExtentInTiles;
-	}
 	
 	// If the dynamic bounds is outside the full bounds of the water mesh, we shouldn't regenerate the quadtree
 	if (!(WaterWorldBox.GetArea() > 0.f))
@@ -248,7 +248,7 @@ void UWaterMeshComponent::RebuildWaterMesh(float InTileSize, const FIntPoint& In
 	}
 
 	// This resets the tree to an initial state, ready for node insertion
-	WaterQuadTree.InitTree(WaterWorldBox, InTileSize, ExtentInTilesToUse);
+	WaterQuadTree.InitTree(WaterWorldBox, InTileSize, InExtentInTiles);
 
 	UsedMaterials.Empty();
 
@@ -583,8 +583,10 @@ void UWaterMeshComponent::RebuildWaterMesh(float InTileSize, const FIntPoint& In
 	{
 		UsedMaterials.Add(FarDistanceMaterial);
 
-		const FBox2D FarMeshInnerRegion = FBox2D(-WorldExtent + GridPosition, WorldExtent + GridPosition);
-		WaterQuadTree.AddFarMesh(FarDistanceMaterial, FarMeshInnerRegion, FarDistanceMeshExtent, FarMeshHeight);
+		// Far Mesh should stitch to the edge of the water zone
+		const FBox2D FarMeshBounds = WaterZone->GetZoneBounds2D();
+
+		WaterQuadTree.AddFarMesh(FarDistanceMaterial, FarMeshBounds, FarDistanceMeshExtent, FarMeshHeight);
 	}
 
 	WaterQuadTree.Unlock(true);
@@ -616,6 +618,8 @@ void UWaterMeshComponent::Update()
 		TessFactorBiasScalability = NewTessFactorBias;
 		LODScaleBiasScalability = NewLODScaleBias;
 		const float LODCountBiasFactor = FMath::Pow(2.0f, (float)LODCountBiasScalability);
+
+		FIntPoint ExtentInTiles = GetExtentInTiles();
 		RebuildWaterMesh(TileSize / LODCountBiasFactor, FIntPoint(FMath::CeilToInt(ExtentInTiles.X * LODCountBiasFactor), FMath::CeilToInt(ExtentInTiles.Y * LODCountBiasFactor)));
 		PrecachePSOs();
 		bNeedsRebuild = false;
@@ -636,8 +640,6 @@ void UWaterMeshComponent::PostEditChangeProperty(FPropertyChangedEvent& Property
 		if (PropertyName == GET_MEMBER_NAME_CHECKED(UWaterMeshComponent, LODScale)
 			|| PropertyName == GET_MEMBER_NAME_CHECKED(UWaterMeshComponent, TessellationFactor)
 			|| PropertyName == GET_MEMBER_NAME_CHECKED(UWaterMeshComponent, TileSize)
-			|| PropertyName == GET_MEMBER_NAME_CHECKED(UWaterMeshComponent, ExtentInTiles)
-			|| PropertyName == GET_MEMBER_NAME_CHECKED(UWaterMeshComponent, LocalTessellationExtentInTiles)
 			|| PropertyName == GET_MEMBER_NAME_CHECKED(UWaterMeshComponent, ForceCollapseDensityLevel)
 			|| PropertyName == GET_MEMBER_NAME_CHECKED(UWaterMeshComponent, FarDistanceMaterial)
 			|| PropertyName == GET_MEMBER_NAME_CHECKED(UWaterMeshComponent, FarDistanceMeshExtent)
diff --git a/Engine/Plugins/Experimental/Water/Source/Runtime/Private/WaterMeshSceneProxy.cpp b/Engine/Plugins/Experimental/Water/Source/Runtime/Private/WaterMeshSceneProxy.cpp
--- a/Engine/Plugins/Experimental/Water/Source/Runtime/Private/WaterMeshSceneProxy.cpp
+++ b/Engine/Plugins/Experimental/Water/Source/Runtime/Private/WaterMeshSceneProxy.cpp
@@ -3,6 +3,7 @@
 #include "WaterMeshSceneProxy.h"
 #include "MaterialShared.h"
 #include "WaterMeshComponent.h"
+#include "WaterZoneActor.h"
 #include "Materials/Material.h"
 #include "WaterUtils.h"
 #include "PrimitiveViewRelevance.h"
@@ -120,6 +121,12 @@ FWaterMeshSceneProxy::FWaterMeshSceneProxy(UWaterMeshComponent* Component)
 
 	WaterQuadTree.BuildMaterialIndices();
 
+	if (const AWaterZone* WaterZone = Component->GetOwner<AWaterZone>(); ensureMsgf(WaterZone != nullptr, TEXT("WaterMeshComponent is owned by an actor that is not a WaterZone. This is not supported!")))
+	{
+		const FBox WaterInfoBounds3D = WaterZone->GetDynamicWaterInfoBounds();
+		WaterInfoBounds = FBox2D(FVector2D(WaterInfoBounds3D.Min), FVector2D(WaterInfoBounds3D.Max));
+	}
+
 #if RHI_RAYTRACING
 	RayTracingWaterData.SetNum(DensityCount);
 #endif
@@ -268,6 +275,7 @@ void FWaterMeshSceneProxy::GetDynamicMeshElements(const TArray<const FSceneView*
 			TraversalDesc.PreViewTranslation = View->ViewMatrices.GetPreViewTranslation();
 			TraversalDesc.LODScale = LODScale;
 			TraversalDesc.bLODMorphingEnabled = !!CVarWaterMeshLODMorphEnabled.GetValueOnRenderThread();
+			TraversalDesc.WaterInfoBounds = WaterInfoBounds;
 
 #if !(UE_BUILD_SHIPPING || UE_BUILD_TEST)
 			//Debug
@@ -507,6 +515,7 @@ void FWaterMeshSceneProxy::GetDynamicRayTracingInstances(FRayTracingMaterialGath
 	TraversalDesc.Frustum = FConvexVolume(); // Default volume to disable frustum culling
 	TraversalDesc.LODScale = LODScale;
 	TraversalDesc.bLODMorphingEnabled = !!CVarWaterMeshLODMorphEnabled.GetValueOnRenderThread();
+	TraversalDesc.WaterInfoBounds = WaterInfoBounds;
 
 	WaterQuadTree.BuildWaterTileInstanceData(TraversalDesc, WaterInstanceData);
 
@@ -660,23 +669,3 @@ HHitProxy* FWaterMeshSceneProxy::CreateHitProxies(UPrimitiveComponent* Component
 	return nullptr;
 }
 #endif // WITH_WATER_SELECTION_SUPPORT
-
-
-void FWaterMeshSceneProxy::OnTessellatedWaterMeshBoundsChanged_GameThread(const FBox2D& InTessellatedWaterMeshBounds)
-{
-	check(IsInParallelGameThread() || IsInGameThread());
-
-	FWaterMeshSceneProxy* SceneProxy = this;
-	ENQUEUE_RENDER_COMMAND(OnTessellatedWaterMeshBoundsChanged)(
-		[SceneProxy, InTessellatedWaterMeshBounds](FRHICommandListImmediate& RHICmdList)
-		{
-			SceneProxy->OnTessellatedWaterMeshBoundsChanged_RenderThread(InTessellatedWaterMeshBounds);
-		});
-}
-
-void FWaterMeshSceneProxy::OnTessellatedWaterMeshBoundsChanged_RenderThread(const FBox2D& InTessellatedWaterMeshBounds)
-{
-	check(IsInRenderingThread());
-
-	TessellatedWaterMeshBounds = InTessellatedWaterMeshBounds;
-}
diff --git a/Engine/Plugins/Experimental/Water/Source/Runtime/Private/WaterQuadTree.cpp b/Engine/Plugins/Experimental/Water/Source/Runtime/Private/WaterQuadTree.cpp
--- a/Engine/Plugins/Experimental/Water/Source/Runtime/Private/WaterQuadTree.cpp
+++ b/Engine/Plugins/Experimental/Water/Source/Runtime/Private/WaterQuadTree.cpp
@@ -131,9 +131,11 @@ void FWaterQuadTree::FNode::SelectLODRefinement(const FNodeData& InNodeData, int
 	const FWaterBodyRenderData& WaterBodyRenderData = InNodeData.WaterBodyRenderData[WaterBodyIndex];
 	const FVector CenterPosition = Bounds.GetCenter();
 	const FVector Extent = Bounds.GetExtent();
+	const FBox2D Bounds2D = FBox2D(FVector2D(Bounds.Min), FVector2D(Bounds.Max));
 
 	// Early out on frustum culling 
-	if (InTraversalDesc.Frustum.IntersectBox(CenterPosition, Extent))
+	check(InTraversalDesc.WaterInfoBounds.bIsValid);
+	if (InTraversalDesc.Frustum.IntersectBox(CenterPosition, Extent) && Bounds2D.Intersect(InTraversalDesc.WaterInfoBounds))
 	{
 		// This LOD can represent all its leaf nodes, simply add node
 		if (CanRender(DensityLevel, InTraversalDesc.ForceCollapseDensityLevel, WaterBodyRenderData))
@@ -159,16 +161,17 @@ void FWaterQuadTree::FNode::SelectLOD(const FNodeData& InNodeData, int32 InLODLe
 	const FWaterBodyRenderData& WaterBodyRenderData = InNodeData.WaterBodyRenderData[WaterBodyIndex];
 	const FVector CenterPosition = Bounds.GetCenter();
 	const FVector Extent = Bounds.GetExtent();
+	const FBox2D Bounds2D = FBox2D(FVector2D(Bounds.Min), FVector2D(Bounds.Max));
 
 	// Early out on frustum culling 
-	if (!InTraversalDesc.Frustum.IntersectBox(CenterPosition, Extent))
+	check(InTraversalDesc.WaterInfoBounds.bIsValid);
+	if (!InTraversalDesc.Frustum.IntersectBox(CenterPosition, Extent) || !Bounds2D.Intersect(InTraversalDesc.WaterInfoBounds))
 	{
 		// Handled
 		return;
 	}
 
 	// Distance to tile (if 0, position is inside quad)
-	FBox2D Bounds2D(FVector2D(Bounds.Min), FVector2D(Bounds.Max));
 	const float ClosestDistanceToTile = FMath::Sqrt(Bounds2D.ComputeSquaredDistanceToPoint(FVector2D(InTraversalDesc.ObserverPosition)));
 
 	// If quad is outside this LOD range, it belongs to the LOD above, assume it fits in that LOD and drill down to find renderable nodes
@@ -263,66 +266,6 @@ void FWaterQuadTree::FNode::SelectLOD(const FNodeData& InNodeData, int32 InLODLe
 	}
 }
 
-void FWaterQuadTree::FNode::SelectLODWithinBounds(const FNodeData& InNodeData, int32 InLODLevel, const FTraversalDesc& InTraversalDesc, FTraversalOutput& Output) const
-{
-	// #todo_water [roey]: this function currently forces all nodes to render at their lowest lod size. This isn't _that_ bad considering most of the nodes are close
-	// enough to the camera to render at lowest lod level anyways but ideally we would leverage the same lod selection system as the non-bounds implementation.
-
-	const FWaterBodyRenderData& WaterBodyRenderData = InNodeData.WaterBodyRenderData[WaterBodyIndex];
-	const FVector CenterPosition = Bounds.GetCenter();
-	const FVector Extent = Bounds.GetExtent();
-
-	// Early out on frustum culling 
-	if (!InTraversalDesc.Frustum.IntersectBox(CenterPosition, Extent))
-	{
-		// Handled
-		return;
-	}
-
-	if (InLODLevel == 0)
-	{
-		if (CanRender(0, InTraversalDesc.ForceCollapseDensityLevel, WaterBodyRenderData))
-		{
-			AddNodeForRender(InNodeData, WaterBodyRenderData, 0, InLODLevel, InTraversalDesc, Output);
-		}
-	}
-	else
-	{
-		// If this node has a complete subtree it will not contain any actual children, they are implicit to save memory so we generate them here
-		if (HasCompleteSubtree && IsSubtreeSameWaterBody)
-		{
-			FNode ChildNode;
-			const FVector HalfBoundSize(Extent.X, Extent.Y, Extent.Z*2.0f);
-			const FVector HalfOffsets[] = { {0.0f, 0.0f, 0.0f}, {1.0f, 0.0f, 0.0f} , {0.0f, 1.0f, 0.0f} , {1.0f, 1.0f, 0.0f} };
-			for (int i = 0; i < 4; i++)
-			{
-				const FVector ChildMin = Bounds.Min + HalfBoundSize * HalfOffsets[i];
-				const FVector ChildMax = ChildMin + HalfBoundSize;
-				const FBox ChildBounds(ChildMin, ChildMax);
-
-				// Create a temporary node to traverse
-				ChildNode.HasCompleteSubtree = 1;
-				ChildNode.IsSubtreeSameWaterBody = 1;
-				ChildNode.TransitionWaterBodyIndex = TransitionWaterBodyIndex;
-				ChildNode.WaterBodyIndex = WaterBodyIndex;
-				ChildNode.Bounds = ChildBounds;
-
-				ChildNode.SelectLODWithinBounds(InNodeData, InLODLevel - 1, InTraversalDesc, Output);
-			}
-		}
-		else
-		{
-			for (int32 ChildIndex : Children)
-			{
-				if (ChildIndex > 0)
-				{
-					InNodeData.Nodes[ChildIndex].SelectLODWithinBounds(InNodeData, InLODLevel - 1, InTraversalDesc, Output);
-				}
-			}
-		}
-	}
-}
-
 void FWaterQuadTree::FNode::AddNodes(FNodeData& InNodeData, const FBox& InMeshBounds, const FBox& InWaterBodyBounds, uint32 InWaterBodyIndex, int32 InLODLevel, uint32 InParentIndex)
 {
 	const FWaterBodyRenderData& InWaterBody = InNodeData.WaterBodyRenderData[InWaterBodyIndex];
diff --git a/Engine/Plugins/Experimental/Water/Source/Runtime/Private/WaterViewExtension.cpp b/Engine/Plugins/Experimental/Water/Source/Runtime/Private/WaterViewExtension.cpp
--- a/Engine/Plugins/Experimental/Water/Source/Runtime/Private/WaterViewExtension.cpp
+++ b/Engine/Plugins/Experimental/Water/Source/Runtime/Private/WaterViewExtension.cpp
@@ -339,7 +339,7 @@ void FWaterViewExtension::SetupView(FSceneViewFamily& InViewFamily, FSceneView&
 
 			WaterInfoUpdateBounds.Add(WaterZone, UpdateBounds);
 
-			const FVector2D WaterQuadTreeHalfExtent = WaterMesh->GetLocalTessellationExtentInTiles() * TileSize;
+			const FVector2D WaterQuadTreeHalfExtent = WaterMesh->GetExtentInTiles() * TileSize;
 			const FVector2D WaterQuadTreeCenter = WaterMesh->GetDynamicWaterMeshCenter();
 			const FBox2D WaterQuadTreeBounds(WaterQuadTreeCenter - WaterQuadTreeHalfExtent, WaterQuadTreeCenter + WaterQuadTreeHalfExtent);
 
diff --git a/Engine/Plugins/Experimental/Water/Source/Runtime/Private/WaterZoneActor.cpp b/Engine/Plugins/Experimental/Water/Source/Runtime/Private/WaterZoneActor.cpp
--- a/Engine/Plugins/Experimental/Water/Source/Runtime/Private/WaterZoneActor.cpp
+++ b/Engine/Plugins/Experimental/Water/Source/Runtime/Private/WaterZoneActor.cpp
@@ -144,12 +144,9 @@ void AWaterZone::PostLoad()
 		LocalTessellationExtent = TessellatedWaterMeshExtent_DEPRECATED;
 		bEnableLocalOnlyTessellation = bEnableNonTesselatedLODMesh_DEPRECATED;
 	}
+#endif // WITH_EDITORONLY_DATA
 
-	// WaterMesh ExtentInTiles returns the half extent.
-	const FVector2D ExtentInTiles = 2.0 * FVector2D(WaterMesh->GetExtentInTiles());
-	ZoneExtent = FVector2D(ExtentInTiles * WaterMesh->GetTileSize());
 	OnExtentChanged();
-#endif // WITH_EDITORONLY_DATA
 }
 
 void AWaterZone::PostRegisterAllComponents()
@@ -196,6 +193,7 @@ void AWaterZone::MarkForRebuild(EWaterZoneRebuildFlags Flags, const FBox2D& Upda
 		{
 			UE_LOG(LogWater, Verbose, TEXT("AWaterZone::MarkForRebuild (UpdateWaterMesh)"));
 			WaterMesh->MarkWaterMeshGridDirty();
+			WaterMesh->MarkRenderStateDirty();
 		}
 	}
 	if (EnumHasAnyFlags(Flags, EWaterZoneRebuildFlags::UpdateWaterInfoTexture))
@@ -390,21 +388,8 @@ void AWaterZone::SetFarMeshMaterial(UMaterialInterface* InFarDistanceMaterial)
 
 void AWaterZone::OnExtentChanged()
 {
-	// Compute the new tile extent based on the new bounds
-	const float MeshTileSize = WaterMesh->GetTileSize();
-	const FVector2D ZoneHalfExtent = ZoneExtent / 2.0;
-
-	int32 NewExtentInTilesX = FMath::FloorToInt(ZoneHalfExtent.X / MeshTileSize);
-	int32 NewExtentInTilesY = FMath::FloorToInt(ZoneHalfExtent.Y / MeshTileSize);
-	
-	// We must ensure that the zone is always at least 1x1
-	NewExtentInTilesX = FMath::Max(1, NewExtentInTilesX);
-	NewExtentInTilesY = FMath::Max(1, NewExtentInTilesY);
-
-	WaterMesh->SetExtentInTiles(FIntPoint(NewExtentInTilesX, NewExtentInTilesY));
-
 #if WITH_EDITOR
-	BoundsComponent->SetBoxExtent(FVector(ZoneHalfExtent, 8192.f));
+	BoundsComponent->SetBoxExtent(FVector(GetZoneExtent() / 2., 8192.f));
 #endif // WITH_EDITOR
 
 	UpdateOverlappingWaterBodies();
diff --git a/Engine/Plugins/Experimental/Water/Source/Runtime/Public/WaterMeshComponent.h b/Engine/Plugins/Experimental/Water/Source/Runtime/Public/WaterMeshComponent.h
--- a/Engine/Plugins/Experimental/Water/Source/Runtime/Public/WaterMeshComponent.h
+++ b/Engine/Plugins/Experimental/Water/Source/Runtime/Public/WaterMeshComponent.h
@@ -59,10 +59,8 @@ public:
 
 	float GetLODScale() const { return LODScale + LODScaleBiasScalability; }
 
-	void SetExtentInTiles(FIntPoint NewExtentInTiles);
-	FIntPoint GetExtentInTiles() const { return ExtentInTiles; }
+	FIntPoint GetExtentInTiles() const;
 
-	FIntPoint GetLocalTessellationExtentInTiles() const { return LocalTessellationExtentInTiles; }
 	void SetDynamicWaterMeshCenter(const FVector2D& NewCenter);
 	FVector2D GetDynamicWaterMeshCenter() const { return DynamicWaterMeshCenter; }
 
@@ -86,6 +84,9 @@ public:
 
 	UFUNCTION(BlueprintPure, Category = Rendering)
 	bool IsEnabled() const { return bIsEnabled; }
+
+	UE_DEPRECATED(5.4, "The ExtentInTiles is now derived from the water zone extent and the tile size.")
+	void SetExtentInTiles(FIntPoint NewExtentInTiles) {}
 private:
 	//~ Begin USceneComponent Interface
 	virtual FBoxSphereBounds CalcBounds(const FTransform& LocalToWorld) const override;
@@ -98,13 +99,9 @@ private:
 	UPROPERTY(EditAnywhere, Category = Rendering, meta = (ClampMin = "100", AllowPrivateAcces = "true"))
 	float TileSize = 2400.0f;
 
-	/** The extent of the system in number of tiles. Maximum number of tiles for this system will be ExtentInTiles.X*2*ExtentInTiles.Y*2 */
-	UPROPERTY(EditAnywhere, Category = Rendering, meta = (ClampMin = "1", AllowPrivateAcces = "true"))
-	FIntPoint ExtentInTiles = FIntPoint(64, 64);
-
-	/** The extent of the system in number of tiles when local tessellation is enabled. Maximum number of tiles for this system will be LocalTessellationExtentInTiles.X * 2 * LocalTessellationExtentInTiles.Y * 2 */
-	UPROPERTY(EditAnywhere, Category = Rendering, meta = (ClampMin = "1"))
-	FIntPoint LocalTessellationExtentInTiles = FIntPoint(32, 32);
+	/** The current quad tree resolution derived from the extent of the water zone and the water mesh tile size (Extent / TileSize). */
+	UPROPERTY(Transient, VisibleAnywhere, Category = Rendering)
+	mutable FIntPoint QuadTreeResolution = FIntPoint::ZeroValue;
 
 	/** The current center of the dynamic water mesh. Updated by the water view extension whenever the view location crosses the update bounds. */
 	FVector2D DynamicWaterMeshCenter = FVector2D::ZeroVector;
@@ -148,6 +145,11 @@ private:
 	virtual void PostEditChangeProperty(FPropertyChangedEvent& PropertyChangedEvent) override;
 	//~ Begin USceneComponent Interface
 #endif
+
+#if WITH_EDITORONLY_DATA
+	UPROPERTY()
+	FIntPoint ExtentInTiles_DEPRECATED = FIntPoint(64, 64);
+#endif // WITH_EDITORONLY_DATA
 };
 
 #if UE_ENABLE_INCLUDE_ORDER_DEPRECATED_IN_5_2
diff --git a/Engine/Plugins/Experimental/Water/Source/Runtime/Public/WaterMeshSceneProxy.h b/Engine/Plugins/Experimental/Water/Source/Runtime/Public/WaterMeshSceneProxy.h
--- a/Engine/Plugins/Experimental/Water/Source/Runtime/Public/WaterMeshSceneProxy.h
+++ b/Engine/Plugins/Experimental/Water/Source/Runtime/Public/WaterMeshSceneProxy.h
@@ -63,8 +63,6 @@ public:
 	virtual bool IsRayTracingRelevant() const override { return true; }
 #endif
 
-	void OnTessellatedWaterMeshBoundsChanged_GameThread(const FBox2D& InTessellatedWaterMeshBounds);
-
 private:
 	struct FWaterLODParams
 	{
@@ -92,8 +90,6 @@ private:
 	void SetupRayTracingInstances(FRHICommandListBase& RHICmdList, int32 NumInstances, uint32 DensityIndex);
 #endif
 
-	void OnTessellatedWaterMeshBoundsChanged_RenderThread(const FBox2D& InTessellatedWaterMeshBounds);
-
 	FMaterialRelevance MaterialRelevance;
 
 	// One vertex factory per LOD
@@ -109,7 +105,7 @@ private:
 	WaterMeshUserDataBuffersType* WaterMeshUserDataBuffers;
 
 	/** The world-space bounds of the current water info texture coverage. The Water mesh should only render tiles within this bounding box. */
-	FBox2D TessellatedWaterMeshBounds = FBox2D(ForceInit);
+	FBox2D WaterInfoBounds = FBox2D(ForceInit);
 
 	/** Scale of the concentric LOD squares  */
 	float LODScale = -1.0f;
diff --git a/Engine/Plugins/Experimental/Water/Source/Runtime/Public/WaterQuadTree.h b/Engine/Plugins/Experimental/Water/Source/Runtime/Public/WaterQuadTree.h
--- a/Engine/Plugins/Experimental/Water/Source/Runtime/Public/WaterQuadTree.h
+++ b/Engine/Plugins/Experimental/Water/Source/Runtime/Public/WaterQuadTree.h
@@ -108,6 +108,7 @@ struct FWaterQuadTree
 		FVector ObserverPosition = FVector::ZeroVector;
 		FVector PreViewTranslation = FVector::ZeroVector;
 		FConvexVolume Frustum;
+		FBox2D WaterInfoBounds = FBox2D(ForceInitToZero);
 		bool bLODMorphingEnabled = true;
 
 #if !(UE_BUILD_SHIPPING || UE_BUILD_TEST)
