Water: fixed the water info ground actor collection using GetcomponentsBoundingBox default parameter which only includes collision components by default. This could lead to some landscape components not being included in the render.

diff --git a/Engine/Plugins/Experimental/Water/Source/Runtime/Private/WaterZoneActor.cpp b/Engine/Plugins/Experimental/Water/Source/Runtime/Private/WaterZoneActor.cpp
--- a/Engine/Plugins/Experimental/Water/Source/Runtime/Private/WaterZoneActor.cpp
+++ b/Engine/Plugins/Experimental/Water/Source/Runtime/Private/WaterZoneActor.cpp
@@ -501,7 +501,7 @@ bool AWaterZone::UpdateWaterInfoTexture()
 		const FBox WaterZoneBounds = GetZoneBounds();
 		for (ALandscapeProxy* LandscapeProxy : TActorRange<ALandscapeProxy>(World))
 		{
-			const FBox LandscapeBox = LandscapeProxy->GetComponentsBoundingBox();
+			const FBox LandscapeBox = LandscapeProxy->GetComponentsBoundingBox(/*bIncludeNonColliding = */ true);
 			// Only consider landscapes which this zone intersects with in XY and if the landscape volume is not zero sized
 			if (WaterZoneBounds.IntersectXY(LandscapeBox) && LandscapeBox.GetVolume() > 0.0)
 			{
