Move default water body components into the WaterBodyActor so it's more easily overriden by game based classes.
Prevents users setting Game assets on the runtime settings requiring engine classes to load game classes before the module was loaded.

diff --git a/Engine/Plugins/Experimental/Water/Source/Runtime/Private/WaterBodyActor.cpp b/Engine/Plugins/Experimental/Water/Source/Runtime/Private/WaterBodyActor.cpp
--- a/Engine/Plugins/Experimental/Water/Source/Runtime/Private/WaterBodyActor.cpp
+++ b/Engine/Plugins/Experimental/Water/Source/Runtime/Private/WaterBodyActor.cpp
@@ -17,6 +17,7 @@
 #include "WaterBodyRiverComponent.h"
 #include "WaterBodyInfoMeshComponent.h"
 #include "WaterBodyStaticMeshComponent.h"
+#include "WaterModule.h"
 #include "WaterVersion.h"
 #include "Algo/RemoveIf.h"
 
@@ -42,6 +43,10 @@ TAutoConsoleVariable<float> CVarWaterSplineResampleMaxDistance(
 
 AWaterBody::AWaterBody(const FObjectInitializer& ObjectInitializer)
 	: Super(ObjectInitializer)
+	, WaterBodyRiverComponentClass(UWaterBodyRiverComponent::StaticClass())
+	, WaterBodyLakeComponentClass(UWaterBodyLakeComponent::StaticClass())
+	, WaterBodyOceanComponentClass(UWaterBodyOceanComponent::StaticClass())
+	, WaterBodyCustomComponentClass(UWaterBodyCustomComponent::StaticClass())
 {
 	SetCanBeDamaged(false);
 	bNetLoadOnClient = true;
@@ -274,16 +279,56 @@ void AWaterBody::InitializeBody()
 		switch (GetWaterBodyType())
 		{
 		case EWaterBodyType::River:
+			PRAGMA_DISABLE_DEPRECATION_WARNINGS
 			WaterBodyComponentClass = WaterSettings->GetWaterBodyRiverComponentClass().Get();
+			PRAGMA_ENABLE_DEPRECATION_WARNINGS
+			if (WaterBodyComponentClass)
+			{
+				UE_LOG(LogWater, Error, TEXT("Found non-null deprecated setting for %s water bodies. Global WaterBodyComponent class overrides are no longer supported. Please create a new WaterBodyActor class within the same plugin as the component class to override the setting."), *UEnum::GetValueAsString(GetWaterBodyType()));
+			}
+			else
+			{
+				WaterBodyComponentClass = WaterBodyRiverComponentClass.Get();
+			}
 			break;
 		case EWaterBodyType::Lake:
+			PRAGMA_DISABLE_DEPRECATION_WARNINGS
 			WaterBodyComponentClass = WaterSettings->GetWaterBodyLakeComponentClass().Get();
+			PRAGMA_ENABLE_DEPRECATION_WARNINGS
+			if (WaterBodyComponentClass)
+			{
+				UE_LOG(LogWater, Error, TEXT("Found non-null deprecated setting for %s water bodies. Global WaterBodyComponent class overrides are no longer supported. Please create a new WaterBodyActor class within the same plugin as the component class to override the setting."), *UEnum::GetValueAsString(GetWaterBodyType()));
+			}
+			else
+			{
+				WaterBodyComponentClass = WaterBodyLakeComponentClass.Get();
+			}
 			break;
 		case EWaterBodyType::Ocean:
+			PRAGMA_DISABLE_DEPRECATION_WARNINGS
 			WaterBodyComponentClass = WaterSettings->GetWaterBodyOceanComponentClass().Get();
+			PRAGMA_ENABLE_DEPRECATION_WARNINGS
+			if (WaterBodyComponentClass)
+			{
+				UE_LOG(LogWater, Error, TEXT("Found non-null deprecated setting for %s water bodies. Global WaterBodyComponent class overrides are no longer supported. Please create a new WaterBodyActor class within the same plugin as the component class to override the setting."), *UEnum::GetValueAsString(GetWaterBodyType()));
+			}
+			else
+			{
+				WaterBodyComponentClass = WaterBodyOceanComponentClass.Get();
+			}
 			break;
 		case EWaterBodyType::Transition:
+			PRAGMA_DISABLE_DEPRECATION_WARNINGS
 			WaterBodyComponentClass = WaterSettings->GetWaterBodyCustomComponentClass().Get();
+			PRAGMA_ENABLE_DEPRECATION_WARNINGS
+			if (WaterBodyComponentClass)
+			{
+				UE_LOG(LogWater, Error, TEXT("Found non-null deprecated setting for %s water bodies. Global WaterBodyComponent class overrides are no longer supported. Please create a new WaterBodyActor class within the same plugin as the component class to override the setting."), *UEnum::GetValueAsString(GetWaterBodyType()));
+			}
+			else
+			{
+				WaterBodyComponentClass = WaterBodyCustomComponentClass.Get();
+			}
 			break;
 		default:
 			checkf(false, TEXT("Invalid Water Body Type"));
diff --git a/Engine/Plugins/Experimental/Water/Source/Runtime/Private/WaterRuntimeSettings.cpp b/Engine/Plugins/Experimental/Water/Source/Runtime/Private/WaterRuntimeSettings.cpp
--- a/Engine/Plugins/Experimental/Water/Source/Runtime/Private/WaterRuntimeSettings.cpp
+++ b/Engine/Plugins/Experimental/Water/Source/Runtime/Private/WaterRuntimeSettings.cpp
@@ -18,10 +18,6 @@ UWaterRuntimeSettings::UWaterRuntimeSettings()
 	: MaterialParameterCollection(FSoftObjectPath(TEXT("/Water/Materials/MPC/MPC_Water.MPC_Water")))
 	, DefaultWaterCollisionProfileName(TEXT("WaterBodyCollision"))
 	, DefaultWaterInfoMaterial(FSoftObjectPath(TEXT("/Water/Materials/WaterInfo/DrawWaterInfo.DrawWaterInfo")))
-	, WaterBodyRiverComponentClass(UWaterBodyRiverComponent::StaticClass())
-	, WaterBodyLakeComponentClass(UWaterBodyLakeComponent::StaticClass())
-	, WaterBodyOceanComponentClass(UWaterBodyOceanComponent::StaticClass())
-	, WaterBodyCustomComponentClass(UWaterBodyCustomComponent::StaticClass())
 {
 	
 }
@@ -43,22 +39,22 @@ void UWaterRuntimeSettings::PostInitProperties()
 
 TSubclassOf<UWaterBodyRiverComponent> UWaterRuntimeSettings::GetWaterBodyRiverComponentClass() const
 {
-	return WaterBodyRiverComponentClass;
+	return WaterBodyRiverComponentClass_DEPRECATED;
 }
 
 TSubclassOf<UWaterBodyLakeComponent> UWaterRuntimeSettings::GetWaterBodyLakeComponentClass() const
 {
-	return WaterBodyLakeComponentClass;
+	return WaterBodyLakeComponentClass_DEPRECATED;
 }
 
 TSubclassOf<UWaterBodyOceanComponent> UWaterRuntimeSettings::GetWaterBodyOceanComponentClass() const
 {
-	return WaterBodyOceanComponentClass;
+	return WaterBodyOceanComponentClass_DEPRECATED;
 }
 
 TSubclassOf<UWaterBodyCustomComponent> UWaterRuntimeSettings::GetWaterBodyCustomComponentClass() const
 {
-	return WaterBodyCustomComponentClass;
+	return WaterBodyCustomComponentClass_DEPRECATED;
 }
 
 #if WITH_EDITOR
diff --git a/Engine/Plugins/Experimental/Water/Source/Runtime/Public/WaterBodyActor.h b/Engine/Plugins/Experimental/Water/Source/Runtime/Public/WaterBodyActor.h
--- a/Engine/Plugins/Experimental/Water/Source/Runtime/Public/WaterBodyActor.h
+++ b/Engine/Plugins/Experimental/Water/Source/Runtime/Public/WaterBodyActor.h
@@ -14,6 +14,12 @@ class AWaterBodyExclusionVolume;
 class ALandscapeProxy;
 class UMaterialInstanceDynamic;
 
+class UWaterBodyComponent;
+class UWaterBodyRiverComponent;
+class UWaterBodyLakeComponent;
+class UWaterBodyOceanComponent;
+class UWaterBodyCustomComponent;
+
 // ----------------------------------------------------------------------------------
 
 // For internal use.
@@ -35,7 +41,7 @@ class WATER_API UDEPRECATED_WaterBodyGenerator : public UObject
  * The new class will automatically have a corresponding UWaterBodyComponent specific to that water body type.
  * The component class for each water body type can be defined in the Editor Settings.
  */
-UCLASS(Blueprintable, Abstract, HideCategories = (Tags, Activation, Cooking, Replication, Input, Actor, AssetUserData))
+UCLASS(Blueprintable, config = Engine, Abstract, HideCategories = (Tags, Activation, Cooking, Replication, Input, Actor, AssetUserData))
 class WATER_API AWaterBody : public AActor, public IWaterBrushActorInterface
 {
 	GENERATED_UCLASS_BODY()
@@ -150,6 +156,18 @@ protected:
 	/** Sets up a new list of water body static mesh components. */
 	void SetWaterBodyStaticMeshComponents(TArrayView<TObjectPtr<UWaterBodyStaticMeshComponent>> NewComponentList, TConstArrayView<TObjectPtr<UWaterBodyStaticMeshComponent>> ComponentsToUnregister = {});
 
+	UPROPERTY(EditDefaultsOnly, Config, Category = Water,  meta = (MetaClass = "/Script/Water.WaterBodyRiverComponent"))
+	TSubclassOf<UWaterBodyRiverComponent> WaterBodyRiverComponentClass;
+
+	UPROPERTY(EditDefaultsOnly, Config, Category = Water,  meta = (MetaClass = "/Script/Water.WaterBodyLakeComponent"))
+	TSubclassOf<UWaterBodyLakeComponent> WaterBodyLakeComponentClass;
+
+	UPROPERTY(EditDefaultsOnly, Config, Category = Water,  meta = (MetaClass = "/Script/Water.WaterBodyOceanComponent"))
+	TSubclassOf<UWaterBodyOceanComponent> WaterBodyOceanComponentClass;
+
+	UPROPERTY(EditDefaultsOnly, Config, Category = Water,  meta = (MetaClass = "/Script/Water.WaterBodyCustomComponent"))
+	TSubclassOf<UWaterBodyCustomComponent> WaterBodyCustomComponentClass;
+
 // ----------------------------------------------------------------------------------
 // Deprecated
 
diff --git a/Engine/Plugins/Experimental/Water/Source/Runtime/Public/WaterRuntimeSettings.h b/Engine/Plugins/Experimental/Water/Source/Runtime/Public/WaterRuntimeSettings.h
--- a/Engine/Plugins/Experimental/Water/Source/Runtime/Public/WaterRuntimeSettings.h
+++ b/Engine/Plugins/Experimental/Water/Source/Runtime/Public/WaterRuntimeSettings.h
@@ -33,12 +33,16 @@ public:
 
 	UMaterialInterface* GetDefaultWaterInfoMaterial() const;
 
+	UE_DEPRECATED(5.4, "Seting moved to WaterBodyActor")
 	TSubclassOf<UWaterBodyRiverComponent> GetWaterBodyRiverComponentClass() const;
 
+	UE_DEPRECATED(5.4, "Seting moved to WaterBodyActor")
 	TSubclassOf<UWaterBodyLakeComponent> GetWaterBodyLakeComponentClass() const;
 
+	UE_DEPRECATED(5.4, "Seting moved to WaterBodyActor")
 	TSubclassOf<UWaterBodyOceanComponent> GetWaterBodyOceanComponentClass() const;
 
+	UE_DEPRECATED(5.4, "Seting moved to WaterBodyActor")
 	TSubclassOf<UWaterBodyCustomComponent> GetWaterBodyCustomComponentClass() const;
 
 	//#todo_water: how can we put these settins on the editor settings but still access them from the WaterModule?
@@ -76,17 +80,17 @@ private:
 	UPROPERTY(EditAnywhere, config, Category = Water)
 	TSoftObjectPtr<UMaterialInterface> DefaultWaterInfoMaterial;
 	
-	UPROPERTY(EditAnywhere, Config, Category = Water,  meta = (MetaClass = "/Script/Water.WaterBodyRiverComponent"))
-	TSubclassOf<UWaterBodyRiverComponent> WaterBodyRiverComponentClass;
+	UPROPERTY(Config, meta=(DeprecatedProperty, DeprecationMessage = "Moved to WaterBodyActor", MetaClass = "/Script/Water.WaterBodyRiverComponent"))
+	TSubclassOf<UWaterBodyRiverComponent> WaterBodyRiverComponentClass_DEPRECATED;
 
-	UPROPERTY(EditAnywhere, Config, Category = Water,  meta = (MetaClass = "/Script/Water.WaterBodyLakeComponent"))
-	TSubclassOf<UWaterBodyLakeComponent> WaterBodyLakeComponentClass;
+	UPROPERTY(Config, meta=(DeprecatedProperty, DeprecationMessage = "Moved to WaterBodyActor", MetaClass = "/Script/Water.WaterBodyLakeComponent"))
+	TSubclassOf<UWaterBodyLakeComponent> WaterBodyLakeComponentClass_DEPRECATED;
 
-	UPROPERTY(EditAnywhere, Config, Category = Water,  meta = (MetaClass = "/Script/Water.WaterBodyOceanComponent"))
-	TSubclassOf<UWaterBodyOceanComponent> WaterBodyOceanComponentClass;
+	UPROPERTY(Config, meta=(DeprecatedProperty, DeprecationMessage = "Moved to WaterBodyActor", MetaClass = "/Script/Water.WaterBodyOceanComponent"))
+	TSubclassOf<UWaterBodyOceanComponent> WaterBodyOceanComponentClass_DEPRECATED;
 
-	UPROPERTY(EditAnywhere, Config, Category = Water,  meta = (MetaClass = "/Script/Water.WaterBodyCustomComponent"))
-	TSubclassOf<UWaterBodyCustomComponent> WaterBodyCustomComponentClass;
+	UPROPERTY(Config, meta=(DeprecatedProperty, DeprecationMessage = "Moved to WaterBodyActor", MetaClass = "/Script/Water.WaterBodyCustomComponent"))
+	TSubclassOf<UWaterBodyCustomComponent> WaterBodyCustomComponentClass_DEPRECATED;
 
 	// #todo_water: move this
 	/** If enabled, MapCheck will notify users that their ocean does not completely fill the water zone. This can be desirable in some cases but the default should be to fill completely. */
