Fix for single frame flicker of water with the new WaterInfoTexture rendering method. Fixed by moving update of View water GPU buffers from SetupViewFamily() to SetupView(). This avoids issues where changes done to the WaterZone within SetupView() weren't reflected on the GPU.

#rnx

diff --git a/Engine/Plugins/Experimental/Water/Source/Runtime/Private/WaterViewExtension.cpp b/Engine/Plugins/Experimental/Water/Source/Runtime/Private/WaterViewExtension.cpp
--- a/Engine/Plugins/Experimental/Water/Source/Runtime/Private/WaterViewExtension.cpp
+++ b/Engine/Plugins/Experimental/Water/Source/Runtime/Private/WaterViewExtension.cpp
@@ -66,7 +66,7 @@ void FWaterViewExtension::Deinitialize()
 	}
 }
 
-void FWaterViewExtension::SetupViewFamily(FSceneViewFamily& InViewFamily)
+void FWaterViewExtension::UpdateGPUBuffers()
 {
 	if (bRebuildGPUData)
 	{
@@ -287,6 +287,11 @@ void FWaterViewExtension::SetupViewFamily(FSceneViewFamily& InViewFamily)
 	}
 }
 
+void FWaterViewExtension::SetupViewFamily(FSceneViewFamily& InViewFamily)
+{
+
+}
+
 void FWaterViewExtension::SetupView(FSceneViewFamily& InViewFamily, FSceneView& InView)
 {
 	if (CVarLocalTessellationFreeze.GetValueOnGameThread())
@@ -386,6 +391,10 @@ void FWaterViewExtension::SetupView(FSceneViewFamily& InViewFamily, FSceneView&
 			}
 		}
 	}
+
+	// The logic in UpdateGPUBuffers() used to be done in SetupViewFamily(). However, SetupView() (which is responsible for water info rendering) potentially modifies the WaterZone but is called after SetupViewFamily().
+	// This can lead to visual artifacts due to outdated data in the GPU buffers.
+	UpdateGPUBuffers();
 }
 
 void FWaterViewExtension::PreRenderViewFamily_RenderThread(FRDGBuilder& GraphBuilder, FSceneViewFamily& InViewFamily)
diff --git a/Engine/Plugins/Experimental/Water/Source/Runtime/Public/WaterViewExtension.h b/Engine/Plugins/Experimental/Water/Source/Runtime/Public/WaterViewExtension.h
--- a/Engine/Plugins/Experimental/Water/Source/Runtime/Public/WaterViewExtension.h
+++ b/Engine/Plugins/Experimental/Water/Source/Runtime/Public/WaterViewExtension.h
@@ -51,6 +51,8 @@ private:
 	TSharedRef<FWaterGPUResources, ESPMode::ThreadSafe> WaterGPUData;
 
 	bool bRebuildGPUData = true;
+
+	void UpdateGPUBuffers();
 };
 
 #if UE_ENABLE_INCLUDE_ORDER_DEPRECATED_IN_5_2
