Water: fixed missing update when changing the WaterZoneOverride parameter on water bodies.

diff --git a/Engine/Plugins/Experimental/Water/Source/Runtime/Private/WaterBodyComponent.cpp b/Engine/Plugins/Experimental/Water/Source/Runtime/Private/WaterBodyComponent.cpp
--- a/Engine/Plugins/Experimental/Water/Source/Runtime/Private/WaterBodyComponent.cpp
+++ b/Engine/Plugins/Experimental/Water/Source/Runtime/Private/WaterBodyComponent.cpp
@@ -1150,6 +1150,10 @@ void UWaterBodyComponent::OnPostEditChangeProperty(FOnWaterBodyChangedParams& In
 	{
 		UpdateWaterBodyStaticMeshComponents();
 	}
+	else if (PropertyChangedEvent.GetMemberPropertyName() == GET_MEMBER_NAME_CHECKED(UWaterBodyComponent, WaterZoneOverride))
+	{
+		UpdateWaterZones();
+	}
 }
 
 TArray<TSharedRef<FTokenizedMessage>> UWaterBodyComponent::CheckWaterBodyStatus()
