Water: Fixed river mesh generating via DivideSplineIntoPolyLine which should've instead been using Convert. Divide should be a private helper that is mathematically incorrect when used across multiple spline keys.

diff --git a/Engine/Plugins/Experimental/Water/Source/Runtime/Private/WaterBodyRiverComponent.cpp b/Engine/Plugins/Experimental/Water/Source/Runtime/Private/WaterBodyRiverComponent.cpp
--- a/Engine/Plugins/Experimental/Water/Source/Runtime/Private/WaterBodyRiverComponent.cpp
+++ b/Engine/Plugins/Experimental/Water/Source/Runtime/Private/WaterBodyRiverComponent.cpp
@@ -302,7 +302,7 @@ bool UWaterBodyRiverComponent::GenerateWaterBodyMesh(UE::Geometry::FDynamicMesh3
 
 	TArray<double> Distances;
 	TArray<FVector> Points;
-	SplineComp->DivideSplineIntoPolylineRecursiveWithDistances(0.f, SplineComp->GetSplineLength(), ESplineCoordinateSpace::Local, FMath::Square(5.f), Points, Distances);
+	SplineComp->ConvertSplineToPolyLineWithDistances(ESplineCoordinateSpace::Local, FMath::Square(5.f), Points, Distances);
 	if (Distances.Num() == 0)
 	{
 		return false;
