Water: avoid updating water body render data when in PIE if dynamic data changes are not allowed (typically when the mesh is set to Static).

diff --git a/Engine/Plugins/Experimental/Water/Source/Runtime/Private/WaterBodyComponent.cpp b/Engine/Plugins/Experimental/Water/Source/Runtime/Private/WaterBodyComponent.cpp
--- a/Engine/Plugins/Experimental/Water/Source/Runtime/Private/WaterBodyComponent.cpp
+++ b/Engine/Plugins/Experimental/Water/Source/Runtime/Private/WaterBodyComponent.cpp
@@ -1933,11 +1933,18 @@ void UWaterBodyComponent::SetWaterBodyStaticMeshEnabled(bool bEnabled)
 
 void UWaterBodyComponent::UpdateWaterBodyRenderData()
 {
-	UpdateWaterInfoMeshComponents();
+	// Avoid updating any mesh data if we are in a PIE world if dynamic data changes aren't allowed.
+	if (const UWorld* World = GetWorld())
+	{
+		if (!World->HasBegunPlay() || AreDynamicDataChangesAllowed())
+		{
+			UpdateWaterInfoMeshComponents();
 
-	UpdateWaterBodyStaticMeshComponents();
+			UpdateWaterBodyStaticMeshComponents();
 
-	OnWaterBodyRenderDataUpdated();
+			OnWaterBodyRenderDataUpdated();
+		}
+	}
 }
 
 void UWaterBodyComponent::OnWaterBodyRenderDataUpdated()
