Implemented GetOcclusionQueries/AcceptOcclusionResults based occlusion culling for the CPU water quadtree. Supports setting a limit on the number of occlusion queries issued. Also added a new cvar to also draw occluded (but still rendered) water tile bounds with debug drawing enabled: r.Water.WaterMesh.ShowTileBounds.DrawForeground

r.Water.WaterMesh.OcclusionCulling enables OC. Enabled by default.
r.Water.WaterMesh.OcclusionCulling.MaxQueries defaults to 256
r.Water.WaterMesh.OcclusionCulling.IncludeFarMesh do OC for the water far mesh; independent of the MaxQueries limit. Enabled by default.

#rnx

diff --git a/Engine/Plugins/Experimental/Water/Source/Runtime/Private/WaterMeshSceneProxy.cpp b/Engine/Plugins/Experimental/Water/Source/Runtime/Private/WaterMeshSceneProxy.cpp
--- a/Engine/Plugins/Experimental/Water/Source/Runtime/Private/WaterMeshSceneProxy.cpp
+++ b/Engine/Plugins/Experimental/Water/Source/Runtime/Private/WaterMeshSceneProxy.cpp
@@ -18,9 +18,11 @@
 DECLARE_STATS_GROUP(TEXT("Water Mesh"), STATGROUP_WaterMesh, STATCAT_Advanced);
 
 DECLARE_DWORD_COUNTER_STAT(TEXT("Tiles Drawn"), STAT_WaterTilesDrawn, STATGROUP_WaterMesh);
+DECLARE_DWORD_COUNTER_STAT(TEXT("Tiles Occlusion Culled"), STAT_WaterTilesOcclusionCulled, STATGROUP_WaterMesh);
 DECLARE_DWORD_COUNTER_STAT(TEXT("Draw Calls"), STAT_WaterDrawCalls, STATGROUP_WaterMesh);
 DECLARE_DWORD_COUNTER_STAT(TEXT("Vertices Drawn"), STAT_WaterVerticesDrawn, STATGROUP_WaterMesh);
 DECLARE_DWORD_COUNTER_STAT(TEXT("Number Drawn Materials"), STAT_WaterDrawnMats, STATGROUP_WaterMesh);
+DECLARE_DWORD_COUNTER_STAT(TEXT("Number Occlusion Queries"), STAT_WaterOcclusionQueries, STATGROUP_WaterMesh);
 
 /** Scalability CVars */
 static TAutoConsoleVariable<int32> CVarWaterMeshLODMorphEnabled(
@@ -28,6 +30,16 @@ static TAutoConsoleVariable<int32> CVarWaterMeshLODMorphEnabled(
 	TEXT("If the smooth LOD morph is enabled. Turning this off may cause slight popping between LOD levels but will skip the calculations in the vertex shader, making it cheaper"),
 	ECVF_Scalability | ECVF_RenderThreadSafe);
 
+static TAutoConsoleVariable<int32> CVarWaterMeshOcclusionCullingMaxQueries(
+	TEXT("r.Water.WaterMesh.OcclusionCulling.MaxQueries"), 256,
+	TEXT("Maximum number of occlusion queries for the CPU water quadtree nodes. Using fewer queries than nodes will result in coarser culling."),
+	ECVF_Scalability);
+
+static TAutoConsoleVariable<int32> CVarWaterMeshOcclusionCullingIncludeFarMesh(
+	TEXT("r.Water.WaterMesh.OcclusionCulling.IncludeFarMesh"), 1,
+	TEXT("When occlusion culling is enabled, always do occlusion queries for the water far mesh, independent of r.Water.WaterMesh.OcclusionCulling.MaxQueries."),
+	ECVF_Scalability);
+
 /** Debug CVars */
 static TAutoConsoleVariable<int32> CVarWaterMeshShowWireframe(
 	TEXT("r.Water.WaterMesh.ShowWireframe"),
@@ -64,6 +76,13 @@ static TAutoConsoleVariable<int32> CVarWaterMeshShowTileBounds(
 	ECVF_RenderThreadSafe
 );
 
+static TAutoConsoleVariable<int32> CVarWaterMeshShowTileBoundsForeground(
+	TEXT("r.Water.WaterMesh.ShowTileBounds.DrawForeground"),
+	0,
+	TEXT("Shows all tile bounds, even occluded ones by drawing into the foreground"),
+	ECVF_RenderThreadSafe
+);
+
 static TAutoConsoleVariable<int32> CVarWaterMeshPreAllocStagingInstanceMemory(
 	TEXT("r.Water.WaterMesh.PreAllocStagingInstanceMemory"),
 	0,
@@ -77,6 +96,13 @@ static TAutoConsoleVariable<int32> CVarRayTracingGeometryWater(
 	TEXT("Include water in ray tracing effects (default = 0 (water disabled in ray tracing))"));
 #endif
 
+static TAutoConsoleVariable<int32> CVarWaterMeshOcclusionCulling(
+	TEXT("r.Water.WaterMesh.OcclusionCulling"),
+	1,
+	TEXT("Enables occlusion culling for the CPU water quadtree."),
+	ECVF_RenderThreadSafe
+);
+
 // ----------------------------------------------------------------------------------
 
 FWaterMeshSceneProxy::FWaterMeshSceneProxy(UWaterMeshComponent* Component)
@@ -130,6 +156,14 @@ FWaterMeshSceneProxy::FWaterMeshSceneProxy(UWaterMeshComponent* Component)
 #if RHI_RAYTRACING
 	RayTracingWaterData.SetNum(DensityCount);
 #endif
+
+	const int32 MaxQueries = CVarWaterMeshOcclusionCullingMaxQueries.GetValueOnGameThread();
+	const bool bIncludeFarMeshOcclusionQueries = CVarWaterMeshOcclusionCullingIncludeFarMesh.GetValueOnGameThread() != 0;
+	OcclusionCullingBounds = WaterQuadTree.ComputeNodeBounds(MaxQueries, bIncludeFarMeshOcclusionQueries, &OcclusionResultsFarMeshOffset);
+	if (!OcclusionCullingBounds.IsEmpty())
+	{
+		EmptyOcclusionCullingBounds.Add(OcclusionCullingBounds[0]);
+	}
 }
 
 FWaterMeshSceneProxy::~FWaterMeshSceneProxy()
@@ -156,6 +190,11 @@ FWaterMeshSceneProxy::~FWaterMeshSceneProxy()
 #endif
 }
 
+void FWaterMeshSceneProxy::CreateRenderThreadResources(FRHICommandListBase& RHICmdList)
+{
+	SceneProxyCreatedFrameNumberRenderThread = GFrameNumberRenderThread;
+}
+
 FWaterMeshSceneProxy::FWaterLODParams FWaterMeshSceneProxy::GetWaterLODParams(const FVector& Position) const
 {
 	float WaterHeightForLOD = 0.0f;
@@ -276,11 +315,22 @@ void FWaterMeshSceneProxy::GetDynamicMeshElements(const TArray<const FSceneView*
 			TraversalDesc.LODScale = LODScale;
 			TraversalDesc.bLODMorphingEnabled = !!CVarWaterMeshLODMorphEnabled.GetValueOnRenderThread();
 			TraversalDesc.WaterInfoBounds = WaterInfoBounds;
+			TraversalDesc.OcclusionCullingResults = nullptr;
+			TraversalDesc.OcclusionCullingFarMeshOffset = OcclusionResultsFarMeshOffset;
+			if (CVarWaterMeshOcclusionCulling.GetValueOnRenderThread() != 0)
+			{
+				const FOcclusionCullingResults* CullingResults = OcclusionResults.Find(View->GetViewKey());
+				if (CullingResults && !CullingResults->Results.IsEmpty())
+				{
+					TraversalDesc.OcclusionCullingResults = &CullingResults->Results;
+				}
+			}
 
 #if !(UE_BUILD_SHIPPING || UE_BUILD_TEST)
 			//Debug
 			TraversalDesc.DebugPDI = Collector.GetPDI(ViewIndex);
 			TraversalDesc.DebugShowTile = CVarWaterMeshShowTileBounds.GetValueOnRenderThread();
+			TraversalDesc.bDebugDrawIntoForeground = CVarWaterMeshShowTileBoundsForeground.GetValueOnRenderThread() != 0;
 #endif
 			WaterQuadTree.BuildWaterTileInstanceData(TraversalDesc, WaterInstanceData);
 
@@ -438,6 +488,61 @@ void FWaterMeshSceneProxy::GetDynamicMeshElements(const TArray<const FSceneView*
 	WaterInstanceDataBuffers->Unlock(RHICmdList);
 }
 
+const TArray<FBoxSphereBounds>* FWaterMeshSceneProxy::GetOcclusionQueries(const FSceneView* View) const
+{
+	if (CVarWaterMeshOcclusionCulling.GetValueOnRenderThread() != 0)
+	{
+		return &OcclusionCullingBounds;
+	}
+	return &EmptyOcclusionCullingBounds;
+}
+
+void FWaterMeshSceneProxy::AcceptOcclusionResults(const FSceneView* View, TArray<bool>* Results, int32 ResultsStart, int32 NumResults)
+{
+	// Don't accept subprimitive occlusion results from a previously-created sceneproxy - the tree may have been different
+	if (OcclusionCullingBounds.Num() == NumResults && SceneProxyCreatedFrameNumberRenderThread < GFrameNumberRenderThread)
+	{
+		// This lock is necessary to guard against access from multiple views.
+		OcclusionResultsMutex.Lock();
+
+		uint32 ViewId = View->GetViewKey();
+		FOcclusionCullingResults* OldResults = OcclusionResults.Find(ViewId);
+		if (OldResults)
+		{
+			OldResults->FrameNumber = GFrameNumberRenderThread;
+			OldResults->Results.Reset();
+			OldResults->Results.Append(Results->GetData() + ResultsStart, NumResults);
+		}
+		else
+		{
+			// now is a good time to clean up any stale entries
+			for (auto Iter = OcclusionResults.CreateIterator(); Iter; ++Iter)
+			{
+				if (Iter.Value().FrameNumber != GFrameNumberRenderThread)
+				{
+					Iter.RemoveCurrent();
+				}
+			}
+
+			FOcclusionCullingResults NewResults;
+			NewResults.FrameNumber = GFrameNumberRenderThread;
+			NewResults.Results.Append(Results->GetData() + ResultsStart, NumResults);
+
+			OcclusionResults.Add(ViewId, MoveTemp(NewResults));
+		}
+
+		OcclusionResultsMutex.Unlock();
+
+		int32 NumCulled = 0;
+		for (int32 ResultIndex = 0; ResultIndex < NumResults; ++ResultIndex)
+		{
+			NumCulled += (*Results)[ResultsStart + ResultIndex] ? 1 : 0;
+		}
+		INC_DWORD_STAT_BY(STAT_WaterTilesOcclusionCulled, NumCulled);
+		INC_DWORD_STAT_BY(STAT_WaterOcclusionQueries, NumResults);
+	}
+}
+
 #if RHI_RAYTRACING
 void FWaterMeshSceneProxy::SetupRayTracingInstances(FRHICommandListBase& RHICmdList, int32 NumInstances, uint32 DensityIndex)
 {
@@ -660,6 +765,11 @@ FPrimitiveViewRelevance FWaterMeshSceneProxy::GetViewRelevance(const FSceneView*
 	return Result;
 }
 
+bool FWaterMeshSceneProxy::HasSubprimitiveOcclusionQueries() const
+{
+	return true;
+}
+
 #if WITH_WATER_SELECTION_SUPPORT
 HHitProxy* FWaterMeshSceneProxy::CreateHitProxies(UPrimitiveComponent* Component, TArray<TRefCountPtr<HHitProxy> >& OutHitProxies)
 {
diff --git a/Engine/Plugins/Experimental/Water/Source/Runtime/Private/WaterQuadTree.cpp b/Engine/Plugins/Experimental/Water/Source/Runtime/Private/WaterQuadTree.cpp
--- a/Engine/Plugins/Experimental/Water/Source/Runtime/Private/WaterQuadTree.cpp
+++ b/Engine/Plugins/Experimental/Water/Source/Runtime/Private/WaterQuadTree.cpp
@@ -112,7 +112,7 @@ void FWaterQuadTree::FNode::AddNodeForRender(const FNodeData& InNodeData, const
 			Color = GColorList.GetFColorByIndex(DensityIndex + 1);
 		}
 
-		DrawWireBox(InTraversalDesc.DebugPDI, Bounds.ExpandBy(FVector(-20.0f, -20.0f, 0.0f)), Color, SDPG_World);
+		DrawWireBox(InTraversalDesc.DebugPDI, Bounds.ExpandBy(FVector(-20.0f, -20.0f, 0.0f)), Color, InTraversalDesc.bDebugDrawIntoForeground ? SDPG_Foreground : SDPG_World);
 	}
 #endif
 }
@@ -137,6 +137,15 @@ void FWaterQuadTree::FNode::SelectLODRefinement(const FNodeData& InNodeData, int
 	check(InTraversalDesc.WaterInfoBounds.bIsValid);
 	if (InTraversalDesc.Frustum.IntersectBox(CenterPosition, Extent) && Bounds2D.Intersect(InTraversalDesc.WaterInfoBounds))
 	{
+		// Occlusion culling
+		if (InTraversalDesc.OcclusionCullingResults 
+			&& (BreadthFirstIndex < (uint32)InTraversalDesc.OcclusionCullingFarMeshOffset) 
+			&& InTraversalDesc.OcclusionCullingResults->IsValidIndex(BreadthFirstIndex)
+			&& (*InTraversalDesc.OcclusionCullingResults)[BreadthFirstIndex])
+		{
+			return;
+		}
+
 		// This LOD can represent all its leaf nodes, simply add node
 		if (CanRender(DensityLevel, InTraversalDesc.ForceCollapseDensityLevel, WaterBodyRenderData))
 		{
@@ -171,6 +180,15 @@ void FWaterQuadTree::FNode::SelectLOD(const FNodeData& InNodeData, int32 InLODLe
 		return;
 	}
 
+	// Occlusion culling
+	if (InTraversalDesc.OcclusionCullingResults
+		&& (BreadthFirstIndex < (uint32)InTraversalDesc.OcclusionCullingFarMeshOffset)
+		&& InTraversalDesc.OcclusionCullingResults->IsValidIndex(BreadthFirstIndex)
+		&& (*InTraversalDesc.OcclusionCullingResults)[BreadthFirstIndex])
+	{
+		return;
+	}
+
 	// Distance to tile (if 0, position is inside quad)
 	const float ClosestDistanceToTile = FMath::Sqrt(Bounds2D.ComputeSquaredDistanceToPoint(FVector2D(InTraversalDesc.ObserverPosition)));
 
@@ -454,6 +472,8 @@ void FWaterQuadTree::InitTree(const FBox2D& InBounds, float InTileSize, FIntPoin
 
 	TileRegion = InBounds;
 
+	BreadthFirstOrder.Reset();
+
 	// Allocate theoretical max, shrink later in Lock()
 	// This is so that the node array doesn't move in memory while inserting
 	NodeData.Nodes.Empty((float)(FMath::Square(RootDim) * 4) / 3.0f);
@@ -549,6 +569,66 @@ void FWaterQuadTree::Unlock(bool bPruneRedundantNodes)
 		NodeData.Nodes.SetNum(EndIndex + 1);
 	}
 
+	// Compute breadth first indices
+	{
+		class FQueue
+		{
+		public:
+			explicit FQueue(int32 InCapacity) 
+			{ 
+				Capacity = InCapacity;
+				Queue.SetNumUninitialized(Capacity);
+			}
+
+			void Add(int32 Element)
+			{
+				check(NumElements < Capacity);
+				Queue[(Front + NumElements) % Capacity] = Element;
+				++NumElements;
+			}
+
+			int32 Pop() 
+			{ 
+				check(NumElements > 0);
+				const int32 Result = Queue[Front]; 
+				Front = (Front + 1) % Capacity;
+				--NumElements;
+				return Result;
+			}
+
+			bool IsEmpty() const { return NumElements == 0; }
+		private:
+			TArray<int32> Queue;
+			int32 Capacity = 0;
+			int32 Front = 0;
+			int32 NumElements = 0;
+		};
+		
+		FQueue Queue(GetMaxLeafCount());
+		uint32 NextBreadthFirstIndex = 0;
+		
+		check(BreadthFirstOrder.IsEmpty());
+		BreadthFirstOrder.Reserve(NodeData.Nodes.Num());
+
+		Queue.Add(0);
+
+		while (!Queue.IsEmpty())
+		{
+			const int32 NodeIndex = Queue.Pop();
+			FNode& Node = NodeData.Nodes[NodeIndex];
+			Node.BreadthFirstIndex = NextBreadthFirstIndex++;
+			BreadthFirstOrder.Add(NodeIndex);
+
+			for (uint32 ChildIndex : Node.Children)
+			{
+				if (ChildIndex > 0)
+				{
+					Queue.Add(ChildIndex);
+				}
+			}
+		}
+	}
+
 	bIsReadOnly = true;
 }
 
@@ -665,23 +745,37 @@ void FWaterQuadTree::BuildWaterTileInstanceData(const FTraversalDesc& InTraversa
 
 		// Bucket index calculation is MaterialIndex*DensityCount+CurrentDensity. Since far mesh doesn't have any Density(aka LOD) steps and should render only using a 2 triangle quad, we enter it only into the last Density bucket (this always corresponds to a 2 triangle quad). 
 		const int32 BucketIndex = FarMeshData.MaterialIndex * InTraversalDesc.DensityCount + (InTraversalDesc.DensityCount - 1);
-		Output.BucketInstanceCounts[BucketIndex] += FarMeshTileCount;
-		Output.InstanceCount += FarMeshTileCount;
-
-		const int32 StartIndex = Output.StagingInstanceData.AddUninitialized(FarMeshTileCount);
-
+		
 		for (int32 i = 0; i < FarMeshTileCount; i++)
 		{
+			// Frustum culling
+			if (!InTraversalDesc.Frustum.IntersectBox(FarMeshData.InstanceData[i].WorldPosition, FVector(FarMeshData.InstanceData[i].Scale.X * 0.5, FarMeshData.InstanceData[i].Scale.Y * 0.5, 1.0)))
+			{
+				continue;
+			}
+
+			// Occlusion culling
+			if (InTraversalDesc.OcclusionCullingResults 
+				&& InTraversalDesc.OcclusionCullingResults->IsValidIndex(InTraversalDesc.OcclusionCullingFarMeshOffset + i)
+				&& (*InTraversalDesc.OcclusionCullingResults)[InTraversalDesc.OcclusionCullingFarMeshOffset + i])
+			{
+				continue;
+			}
+
+			++Output.BucketInstanceCounts[BucketIndex];
+			++Output.InstanceCount;
+
 			// Build instance data
 			// Transform worldposition to Translated World Position
 			const FVector TranslatedWorldPosition(FarMeshData.InstanceData[i].WorldPosition + InTraversalDesc.PreViewTranslation);
-			Output.StagingInstanceData[StartIndex + i].Data[0] = FVector4f(FVector4(TranslatedWorldPosition, 0.0));
-			Output.StagingInstanceData[StartIndex + i].Data[1] = FVector4f(FVector2f::ZeroVector, FarMeshData.InstanceData[i].Scale);
+
+			FStagingInstanceData& StagingInstanceData = Output.StagingInstanceData[Output.StagingInstanceData.AddUninitialized()];
+			StagingInstanceData.BucketIndex = BucketIndex;
+			StagingInstanceData.Data[0] = FVector4f(FVector4(TranslatedWorldPosition, 0.0));
+			StagingInstanceData.Data[1] = FVector4f(FVector2f::ZeroVector, FarMeshData.InstanceData[i].Scale);
 #if WITH_WATER_SELECTION_SUPPORT
-			Output.StagingInstanceData[StartIndex + i].Data[2] = FHitProxyId::InvisibleHitProxyId.GetColor().ReinterpretAsLinear();
+			StagingInstanceData.Data[2] = FHitProxyId::InvisibleHitProxyId.GetColor().ReinterpretAsLinear();
 #endif // WITH_WATER_SELECTION_SUPPORT
-
-			Output.StagingInstanceData[StartIndex + i].BucketIndex = BucketIndex;
 		}
 	}
 }
@@ -759,6 +853,39 @@ void FWaterQuadTree::GatherHitProxies(TArray<TRefCountPtr<HHitProxy> >& OutHitPr
 }
 #endif //WITH_WATER_SELECTION_SUPPORT
 
+TArray<FBoxSphereBounds> FWaterQuadTree::ComputeNodeBounds(int32 MaxNumBounds, bool bIncludeFarMeshTiles, int32* OutFarMeshOffset) const
+{
+	const int32 NumNodes = NodeData.Nodes.Num();
+	const int32 ClampedNumNodes = (MaxNumBounds > 0) ? FMath::Min(MaxNumBounds, NumNodes) : NumNodes;
+	const int32 NumFarMeshTiles = FarMeshData.InstanceData.Num();
+	const int32 NumToReserve = ClampedNumNodes + (bIncludeFarMeshTiles ? NumFarMeshTiles : 0);
+
+	TArray<FBoxSphereBounds> Result;
+	Result.Reserve(NumToReserve);
+
+	// Quadtree tiles
+	for (int32 i = 0; i < ClampedNumNodes; ++i)
+	{
+		Result.Add(NodeData.Nodes[BreadthFirstOrder[i]].Bounds);
+	}
+
+	*OutFarMeshOffset = Result.Num();
+
+	// Far mesh tiles
+	if (bIncludeFarMeshTiles)
+	{
+		for (int32 i = 0; i < NumFarMeshTiles; ++i)
+		{
+			const FFarMeshData::FFarMeshInstanceData& InstanceData = FarMeshData.InstanceData[i];
+			const FVector Extent = FVector(InstanceData.Scale.X * 0.5, InstanceData.Scale.Y * 0.5, 1.0);
+			const FBox Bounds = FBox(InstanceData.WorldPosition - Extent, InstanceData.WorldPosition + Extent);
+			Result.Add(Bounds);
+		}
+	}
+
+	return Result;
+}
+
 /** Split a 2D polygon with a 2D line. Return both polygons */
 static void SplitPolyWithLine(const TArray<FVector2D>& InPoly, const FVector2D& LinePoint, const FVector2D& LineNormal, TArray<FVector2D>& OutPoly0, TArray<FVector2D>& OutPoly1)
 {
diff --git a/Engine/Plugins/Experimental/Water/Source/Runtime/Public/WaterMeshSceneProxy.h b/Engine/Plugins/Experimental/Water/Source/Runtime/Public/WaterMeshSceneProxy.h
--- a/Engine/Plugins/Experimental/Water/Source/Runtime/Public/WaterMeshSceneProxy.h
+++ b/Engine/Plugins/Experimental/Water/Source/Runtime/Public/WaterMeshSceneProxy.h
@@ -29,16 +29,24 @@ public:
 
 	virtual ~FWaterMeshSceneProxy();
 
+	virtual void CreateRenderThreadResources(FRHICommandListBase& RHICmdList) override;
+
 	virtual void GetDynamicMeshElements(const TArray<const FSceneView*>& Views, const FSceneViewFamily& ViewFamily, uint32 VisibilityMap, FMeshElementCollector& Collector) const override;
 
+	virtual const TArray<FBoxSphereBounds>* GetOcclusionQueries(const FSceneView* View) const override;
+
+	virtual void AcceptOcclusionResults(const FSceneView* View, TArray<bool>* Results, int32 ResultsStart, int32 NumResults) override;
+
 	virtual FPrimitiveViewRelevance GetViewRelevance(const FSceneView* View) const override;
 
+	virtual bool HasSubprimitiveOcclusionQueries() const override;
+
 	virtual bool CanBeOccluded() const override
 	{
 		return !MaterialRelevance.bDisableDepthTest;
 	}
 
-	virtual uint32 GetMemoryFootprint() const
+	virtual uint32 GetMemoryFootprint() const override
 	{
 		return(sizeof(*this) + GetAllocatedSize());
 	}
@@ -121,6 +129,19 @@ private:
 	// Per density array of ray tracing geometries.
 	TArray<TArray<FRayTracingWaterData>> RayTracingWaterData;	
 #endif
+
+	struct FOcclusionCullingResults
+	{
+		uint32 FrameNumber;
+		TArray<bool> Results;
+	};
+
+	TArray<FBoxSphereBounds> OcclusionCullingBounds;
+	TArray<FBoxSphereBounds> EmptyOcclusionCullingBounds;
+	TMap<uint32, FOcclusionCullingResults> OcclusionResults;
+	UE::FMutex OcclusionResultsMutex;
+	int32 OcclusionResultsFarMeshOffset = INT32_MAX;
+	uint32 SceneProxyCreatedFrameNumberRenderThread = INDEX_NONE;
 };
 
 #if UE_ENABLE_INCLUDE_ORDER_DEPRECATED_IN_5_2
diff --git a/Engine/Plugins/Experimental/Water/Source/Runtime/Public/WaterQuadTree.h b/Engine/Plugins/Experimental/Water/Source/Runtime/Public/WaterQuadTree.h
--- a/Engine/Plugins/Experimental/Water/Source/Runtime/Public/WaterQuadTree.h
+++ b/Engine/Plugins/Experimental/Water/Source/Runtime/Public/WaterQuadTree.h
@@ -99,6 +99,8 @@ struct FWaterQuadTree
 	/** Output of the quadtree when asking to traverse it for visible water tiles */
 	struct FTraversalDesc
 	{
+		const TArray<bool>* OcclusionCullingResults = nullptr;
+		int32 OcclusionCullingFarMeshOffset = 0;
 		int32 LowestLOD = 0;
 		int32 LODCount = 0;
 		int32 DensityCount = 0;
@@ -115,6 +117,7 @@ struct FWaterQuadTree
 		// Debug
 		int32 DebugShowTile = 0;
 		class FPrimitiveDrawInterface* DebugPDI = nullptr;
+		bool bDebugDrawIntoForeground = false;
 #endif
 	};
 
@@ -181,13 +184,15 @@ struct FWaterQuadTree
 	static float GetLODDistance(int32 InLODLevel, float InLODScale) { return FMath::Pow(2.0f, (float)(InLODLevel + 1)) * InLODScale; }
 
 	/** Total memory dynamically allocated by this object */
-	uint32 GetAllocatedSize() const { return NodeData.GetAllocatedSize() + WaterMaterials.GetAllocatedSize() + FarMeshData.GetAllocatedSize(); }
+	uint32 GetAllocatedSize() const { return NodeData.GetAllocatedSize() + WaterMaterials.GetAllocatedSize() + BreadthFirstOrder.GetAllocatedSize() + FarMeshData.GetAllocatedSize(); }
 
 #if WITH_WATER_SELECTION_SUPPORT
 	/** Obtain all possible hit proxies (proxies of all the water bodies) */
 	void GatherHitProxies(TArray<TRefCountPtr<HHitProxy> >& OutHitProxies) const;
 #endif // WITH_WATER_SELECTION_SUPPORT
 
+	TArray<FBoxSphereBounds> ComputeNodeBounds(int32 MaxNumBounds, bool bIncludeFarMeshTiles, int32* OutFarMeshOffset) const;
+
 private:
 	struct FNodeData;
 
@@ -254,6 +259,7 @@ private:
 
 		/** Children, 0 means invalid */
 		uint32 Children[4] = { 0, 0, 0, 0 };
+		uint32 BreadthFirstIndex = INDEX_NONE;
 	};
 
 	int32 TreeDepth = 0;
@@ -283,6 +289,9 @@ private:
 
 	TArray<FMaterialRenderProxy*> WaterMaterials;
 
+	/** Node indices ordered by breadth first ordering */
+	TArray<int32> BreadthFirstOrder;
+
 	/** Contains everything needed to render the far mesh. This data lives outside the quadtree structure itself */
 	struct FFarMeshData
 	{
