Water: Expose an event on the water zone which is fired whenever the water info texture is created or recreated. Additionally exposed the water zone index to blueprint as well.

diff --git a/Engine/Plugins/Experimental/Water/Source/Runtime/Private/WaterZoneActor.cpp b/Engine/Plugins/Experimental/Water/Source/Runtime/Private/WaterZoneActor.cpp
--- a/Engine/Plugins/Experimental/Water/Source/Runtime/Private/WaterZoneActor.cpp
+++ b/Engine/Plugins/Experimental/Water/Source/Runtime/Private/WaterZoneActor.cpp
@@ -572,7 +572,9 @@ bool AWaterZone::UpdateWaterInfoTexture()
 		// The water info texture is different, we need to bind the newly created texture to all registered water bodies
 		if (WaterInfoTexture != OldTexture)
 		{
-			ForEachWaterBodyComponent([WaterInfoTexture = WaterInfoTexture](UWaterBodyComponent* WaterBodyComponent)
+			OnWaterInfoTextureCreated.Broadcast(WaterInfoTexture);
+
+			ForEachWaterBodyComponent([](UWaterBodyComponent* WaterBodyComponent)
 			{
 				WaterBodyComponent->UpdateMaterialInstances();
 				return true;
diff --git a/Engine/Plugins/Experimental/Water/Source/Runtime/Public/WaterZoneActor.h b/Engine/Plugins/Experimental/Water/Source/Runtime/Public/WaterZoneActor.h
--- a/Engine/Plugins/Experimental/Water/Source/Runtime/Public/WaterZoneActor.h
+++ b/Engine/Plugins/Experimental/Water/Source/Runtime/Public/WaterZoneActor.h
@@ -24,6 +24,9 @@ enum class EWaterZoneRebuildFlags
 };
 ENUM_CLASS_FLAGS(EWaterZoneRebuildFlags);
 
+
+DECLARE_DYNAMIC_MULTICAST_DELEGATE_OneParam(FOnWaterInfoTextureCreated, const UTextureRenderTarget2D*, WaterInfoTexture);
+
 UCLASS(Blueprintable, HideCategories=(Physics, Replication, Input, Collision))
 class WATER_API AWaterZone : public AActor
 {
@@ -74,11 +77,14 @@ public:
 
 	int32 GetOverlapPriority() const { return OverlapPriority; }
 
+	UFUNCTION(BlueprintCallable, Category=Water)
 	int32 GetWaterZoneIndex() const { return WaterZoneIndex; }
 
 	UPROPERTY(Transient, DuplicateTransient, VisibleAnywhere, BlueprintReadOnly, Category = Water)
 	TObjectPtr<UTextureRenderTarget2D> WaterInfoTexture;
 
+	FOnWaterInfoTextureCreated& GetOnWaterInfoTextureCreated() { return OnWaterInfoTextureCreated; }
+	
 #if WITH_EDITOR
 	virtual TUniquePtr<class FWorldPartitionActorDesc> CreateClassActorDesc() const override;
 	virtual FBox GetStreamingBounds() const override;
@@ -200,6 +206,9 @@ private:
 	UPROPERTY(Transient, DuplicateTransient, NonTransactional, VisibleAnywhere, Category = Water)
 	int32 WaterZoneIndex = INDEX_NONE;
 
+	UPROPERTY(BlueprintAssignable, Category=Water)
+	FOnWaterInfoTextureCreated OnWaterInfoTextureCreated;
+
 #if WITH_EDITORONLY_DATA
 	/** A manipulatable box for visualizing/editing the water zone bounds */
 	UPROPERTY(Transient)
