Including far mesh tiles in the water mesh debug drawing

#rnx

diff --git a/Engine/Plugins/Experimental/Water/Source/Runtime/Private/WaterQuadTree.cpp b/Engine/Plugins/Experimental/Water/Source/Runtime/Private/WaterQuadTree.cpp
--- a/Engine/Plugins/Experimental/Water/Source/Runtime/Private/WaterQuadTree.cpp
+++ b/Engine/Plugins/Experimental/Water/Source/Runtime/Private/WaterQuadTree.cpp
@@ -776,6 +776,15 @@ void FWaterQuadTree::BuildWaterTileInstanceData(const FTraversalDesc& InTraversa
 #if WITH_WATER_SELECTION_SUPPORT
 			StagingInstanceData.Data[2] = FHitProxyId::InvisibleHitProxyId.GetColor().ReinterpretAsLinear();
 #endif // WITH_WATER_SELECTION_SUPPORT
+
+#if !(UE_BUILD_SHIPPING || UE_BUILD_TEST)
+			// Debug drawing
+			if (InTraversalDesc.DebugShowTile != 0)
+			{
+				const FBox DebugBounds = FBox(FarMeshData.InstanceData[i].WorldPosition - FVector(FarMeshData.InstanceData[i].Scale.X * 0.5, FarMeshData.InstanceData[i].Scale.Y * 0.5, 1.0), FarMeshData.InstanceData[i].WorldPosition + FVector(FarMeshData.InstanceData[i].Scale.X * 0.5, FarMeshData.InstanceData[i].Scale.Y * 0.5, 1.0));
+				DrawWireBox(InTraversalDesc.DebugPDI, DebugBounds.ExpandBy(FVector(-20.0f, -20.0f, 0.0f)), FColor::Orange, InTraversalDesc.bDebugDrawIntoForeground ? SDPG_Foreground : SDPG_World);
+			}
+#endif
 		}
 	}
 }
