WaterBody: Add ShallowWater intergration WIP

- WaterDataFunctions.ush: Add #if VM_SIMULATION==0 wrapper around WaterDataFunctions.ush so that Niagara VM using this file can compile

#virtualized

diff --git a/Engine/Plugins/Experimental/Water/Shaders/Private/WaterDataFunctions.ush b/Engine/Plugins/Experimental/Water/Shaders/Private/WaterDataFunctions.ush
--- a/Engine/Plugins/Experimental/Water/Shaders/Private/WaterDataFunctions.ush
+++ b/Engine/Plugins/Experimental/Water/Shaders/Private/WaterDataFunctions.ush
@@ -1,5 +1,7 @@
 #pragma once
 
+#if VM_SIMULATION==0
+
 #define PER_WAVE_DATA_SIZE 2
 #define PER_WATERZONE_DATA_SIZE 2
 #define PER_WATERBODY_DATA_SIZE 1
@@ -174,3 +176,4 @@ float2 DecodeWaterInfoVelocity(float4 WaterInfoSample, float MaxVelocity)
 	return (NormalizedFlow - 0.5) * 2.0 * MaxVelocity;
 }
 
+#endif
\ No newline at end of file
diff --git a/Engine/Plugins/Experimental/Water/Source/Runtime/Private/ShallowWaterComponent.cpp b/Engine/Plugins/Experimental/Water/Source/Runtime/Private/ShallowWaterComponent.cpp
new file mode 100644
--- /dev/null
+++ b/Engine/Plugins/Experimental/Water/Source/Runtime/Private/ShallowWaterComponent.cpp
@@ -0,0 +1,76 @@
+﻿// Copyright Epic Games, Inc. All Rights Reserved.
+
+#include "ShallowWaterComponent.h"
+
+#include "ShallowWaterSettings.h"
+#include "NiagaraComponent.h"
+#include "NiagaraFunctionLibrary.h"
+#include "WaterBodyActor.h"
+#include "WaterSubsystem.h"
+#include "Components/SkeletalMeshComponent.h"
+#include "Engine/TextureRenderTarget2D.h"
+
+UShallowWaterComponent::UShallowWaterComponent()
+{
+	PrimaryComponentTick.bCanEverTick = true;
+	SetComponentTickEnabled(true);
+	SetUsingAbsoluteRotation(true);
+}
+
+void UShallowWaterComponent::BeginPlay()
+{
+	Super::BeginPlay();
+
+	GetOwner()->OnActorBeginOverlap.AddDynamic(this, &UShallowWaterComponent::OnOwnerBeginOverlap);
+	TArray<USkeletalMeshComponent*> SKMs;
+	GetOwner()->GetComponents<USkeletalMeshComponent>(SKMs);
+	for(USkeletalMeshComponent* SKM : SKMs)
+	{
+		SKM->ComponentTags.Add(FName("RigidMesh_ShallowWaterCollider"));
+	}
+
+	ShallowWaterNiagaraSimulation = UNiagaraFunctionLibrary::SpawnSystemAttached(GetDefault<UShallowWaterSettings>()->DefaultShallowWaterNiagaraSimulation.LoadSynchronous(), this, FName(""),
+		FVector::Zero(), FRotator::ZeroRotator, EAttachLocation::SnapToTarget, false);
+	ShallowWaterNiagaraSimulation->SetGpuComputeDebug(GetDefault<UShallowWaterSettings>()->bGPUComputeDebug);
+}
+
+void UShallowWaterComponent::TickComponent(float DeltaTime, ELevelTick TickType,
+                                           FActorComponentTickFunction* ThisTickFunction)
+{
+	Super::TickComponent(DeltaTime, TickType, ThisTickFunction);
+
+	if(CurrentWaterBody == nullptr)
+	{
+		return;
+	}
+
+	// #todo this query gets the closest water surface in 3D space, not a XY-plane projection as we instinctively assumed
+	// Could be better and could be worse. We are already ignoring the accurate character-capsule-in-water-plane collision volume
+	const FWaterBodyQueryResult WaterInfo = CurrentWaterBody->GetWaterBodyComponent()->QueryWaterInfoClosestToWorldLocation(
+		GetOwner()->GetActorLocation(), EWaterBodyQueryFlags::ComputeLocation);
+	const FVector WaterLocation = WaterInfo.GetWaterSurfaceLocation();
+	SetWorldLocation(FVector(GetOwner()->GetActorLocation().X, GetOwner()->GetActorLocation().Y, WaterLocation.Z));
+}
+
+void UShallowWaterComponent::OnRegister()
+{
+	Super::OnRegister();
+}
+
+void UShallowWaterComponent::OnOwnerBeginOverlap(AActor* Owner, AActor* OtherActor)
+{
+	AWaterBody* OtherWaterBody = Cast<AWaterBody>(OtherActor);
+	if(OtherWaterBody != nullptr && OtherWaterBody != CurrentWaterBody)
+	{
+		CurrentWaterBody = OtherWaterBody;
+		UpdateNiagaraVariables();
+	}
+}
+
+void UShallowWaterComponent::UpdateNiagaraVariables()
+{
+	TObjectPtr<UTextureRenderTarget2D> WaterInfoTexture = CurrentWaterBody->GetWaterBodyComponent()->GetWaterZone()->
+	                                                                        WaterInfoTexture;
+	ShallowWaterNiagaraSimulation->SetVariableTexture(FName("WaterInfoTexture"), Cast<UTexture>(WaterInfoTexture));
+}
+
diff --git a/Engine/Plugins/Experimental/Water/Source/Runtime/Private/ShallowWaterSettings.cpp b/Engine/Plugins/Experimental/Water/Source/Runtime/Private/ShallowWaterSettings.cpp
new file mode 100644
--- /dev/null
+++ b/Engine/Plugins/Experimental/Water/Source/Runtime/Private/ShallowWaterSettings.cpp
@@ -0,0 +1,16 @@
+﻿// Copyright Epic Games, Inc. All Rights Reserved.
+
+#include "ShallowWaterSettings.h"
+
+UShallowWaterSettings::UShallowWaterSettings()
+{
+
+	
+	DefaultShallowWaterNiagaraSimulation = TSoftObjectPtr<UNiagaraSystem>(FSoftObjectPath(TEXT("/NiagaraFluids/Templates/Liquid/2D/Systems/ShallowWater/Grid2D_SW_FollowPlayer.Grid2D_SW_FollowPlayer")));
+
+}
+
+FName UShallowWaterSettings::GetCategoryName() const
+{
+	return FName("Plugins");
+}
diff --git a/Engine/Plugins/Experimental/Water/Source/Runtime/Public/ShallowWaterComponent.h b/Engine/Plugins/Experimental/Water/Source/Runtime/Public/ShallowWaterComponent.h
new file mode 100644
--- /dev/null
+++ b/Engine/Plugins/Experimental/Water/Source/Runtime/Public/ShallowWaterComponent.h
@@ -0,0 +1,35 @@
+﻿// Copyright Epic Games, Inc. All Rights Reserved.
+
+#pragma once
+
+#include "CoreMinimal.h"
+#include "Components/SceneComponent.h"
+#include "ShallowWaterComponent.generated.h"
+
+class AWaterBody;
+class UNiagaraComponent;
+
+
+UCLASS(ClassGroup=(Custom), meta=(BlueprintSpawnableComponent))
+class WATER_API UShallowWaterComponent : public USceneComponent
+{
+	GENERATED_BODY()
+
+public:
+	UShallowWaterComponent();
+	virtual void BeginPlay() override;
+	virtual void TickComponent(float DeltaTime, ELevelTick TickType,
+							   FActorComponentTickFunction* ThisTickFunction) override;
+	virtual void OnRegister() override;
+	
+protected:
+	// Asset can be set in Project Settings - Plugins - Water ShallowWaterSimulation
+	UPROPERTY(BlueprintReadOnly, VisibleAnywhere, Category="Shallow Water")
+	TObjectPtr<UNiagaraComponent> ShallowWaterNiagaraSimulation;
+
+	UFUNCTION()
+	void OnOwnerBeginOverlap(AActor* Owner, AActor* OtherActor);
+	void UpdateNiagaraVariables();
+	
+	AWaterBody* CurrentWaterBody = nullptr; 
+};
diff --git a/Engine/Plugins/Experimental/Water/Source/Runtime/Public/ShallowWaterSettings.h b/Engine/Plugins/Experimental/Water/Source/Runtime/Public/ShallowWaterSettings.h
new file mode 100644
--- /dev/null
+++ b/Engine/Plugins/Experimental/Water/Source/Runtime/Public/ShallowWaterSettings.h
@@ -0,0 +1,24 @@
+﻿// Copyright Epic Games, Inc. All Rights Reserved.
+
+#pragma once
+
+#include "CoreMinimal.h"
+#include "Engine/DeveloperSettings.h"
+#include "NiagaraSystem.h"
+#include "ShallowWaterSettings.generated.h"
+
+
+UCLASS(config = Engine, defaultconfig, meta=(DisplayName="Water Shallow Water"))
+class WATER_API UShallowWaterSettings : public UDeveloperSettings
+{
+	GENERATED_BODY()
+	UShallowWaterSettings();
+	
+public:
+	virtual FName GetCategoryName() const override;
+	
+	UPROPERTY(Config, BlueprintReadOnly, EditAnywhere, Category="Shallow Water Simulation")
+	TSoftObjectPtr<UNiagaraSystem> DefaultShallowWaterNiagaraSimulation;
+	UPROPERTY(Config, BlueprintReadOnly, EditAnywhere, Category="Shallow Water Simulation")
+	bool bGPUComputeDebug = false;
+};
diff --git a/Engine/Plugins/Experimental/Water/Water.uplugin b/Engine/Plugins/Experimental/Water/Water.uplugin
--- a/Engine/Plugins/Experimental/Water/Water.uplugin
+++ b/Engine/Plugins/Experimental/Water/Water.uplugin
@@ -1,5 +1,5 @@
 {
-	"FileVersion": 1,
+	"FileVersion": 3,
 	"Version": 1,
 	"VersionName": "0.1",
 	"FriendlyName": "Water",
@@ -42,6 +42,10 @@
 		{
 			"Name": "BlueprintMaterialTextureNodes",
 			"Enabled": true
+		},
+		{
+			"Name": "NiagaraFluids",
+			"Enabled": true
 		}
 	]
 }
\ No newline at end of file
