[Backout] - CL29554507
Original CL Desc
-----------------------------------------------------------------
WaterBody: Add ShallowWater intergration WIP

- WaterDataFunctions.ush: Add #if VM_SIMULATION==0 wrapper around WaterDataFunctions.ush so that Niagara VM using this file can compile

#virtualized

diff --git a/Engine/Plugins/Experimental/Water/Source/Runtime/Private/ShallowWaterComponent.cpp b/Engine/Plugins/Experimental/Water/Source/Runtime/Private/ShallowWaterComponent.cpp
deleted file mode 100644
diff --git a/Engine/Plugins/Experimental/Water/Source/Runtime/Private/ShallowWaterSettings.cpp b/Engine/Plugins/Experimental/Water/Source/Runtime/Private/ShallowWaterSettings.cpp
deleted file mode 100644
diff --git a/Engine/Plugins/Experimental/Water/Source/Runtime/Public/ShallowWaterComponent.h b/Engine/Plugins/Experimental/Water/Source/Runtime/Public/ShallowWaterComponent.h
deleted file mode 100644
diff --git a/Engine/Plugins/Experimental/Water/Source/Runtime/Public/ShallowWaterSettings.h b/Engine/Plugins/Experimental/Water/Source/Runtime/Public/ShallowWaterSettings.h
deleted file mode 100644
diff --git a/Engine/Plugins/Experimental/Water/Water.uplugin b/Engine/Plugins/Experimental/Water/Water.uplugin
--- a/Engine/Plugins/Experimental/Water/Water.uplugin
+++ b/Engine/Plugins/Experimental/Water/Water.uplugin
@@ -1,5 +1,5 @@
 {
-	"FileVersion": 3,
+	"FileVersion": 1,
 	"Version": 1,
 	"VersionName": "0.1",
 	"FriendlyName": "Water",
@@ -42,10 +42,6 @@
 		{
 			"Name": "BlueprintMaterialTextureNodes",
 			"Enabled": true
-		},
-		{
-			"Name": "NiagaraFluids",
-			"Enabled": true
 		}
 	]
 }
\ No newline at end of file
