[Water]
*Adding the farmesh bounds to the empty occlusion cull bounds to make sure the far mesh isn't culled when occlusion culling is turned off
*Disabling water occlusion culling by default
*Adding cvar for controlling bounds extension for occlusion culling. This only applies to XY bounds to reduce the disocclusion artifacts when viewing water from afar while keeping the occlusion culling efficiency when close to the ground.

diff --git a/Engine/Plugins/Experimental/Water/Source/Runtime/Private/WaterMeshSceneProxy.cpp b/Engine/Plugins/Experimental/Water/Source/Runtime/Private/WaterMeshSceneProxy.cpp
--- a/Engine/Plugins/Experimental/Water/Source/Runtime/Private/WaterMeshSceneProxy.cpp
+++ b/Engine/Plugins/Experimental/Water/Source/Runtime/Private/WaterMeshSceneProxy.cpp
@@ -98,11 +98,16 @@ static TAutoConsoleVariable<int32> CVarRayTracingGeometryWater(
 
 static TAutoConsoleVariable<int32> CVarWaterMeshOcclusionCulling(
 	TEXT("r.Water.WaterMesh.OcclusionCulling"),
-	1,
+	0,
 	TEXT("Enables occlusion culling for the CPU water quadtree."),
-	ECVF_RenderThreadSafe
+	ECVF_Scalability | ECVF_RenderThreadSafe
 );
 
+static TAutoConsoleVariable<float> CVarWaterMeshOcclusionCullExpandBoundsAmountXY(
+	TEXT("r.Water.WaterMesh.OcclusionCullExpandBoundsAmountXY"), 4800.0f,
+	TEXT("Expand the water tile bounds in XY for the purpose of occlusion culling"),
+	ECVF_Scalability);
+
 // ----------------------------------------------------------------------------------
 
 FWaterMeshSceneProxy::FWaterMeshSceneProxy(UWaterMeshComponent* Component)
@@ -159,11 +164,8 @@ FWaterMeshSceneProxy::FWaterMeshSceneProxy(UWaterMeshComponent* Component)
 
 	const int32 MaxQueries = CVarWaterMeshOcclusionCullingMaxQueries.GetValueOnGameThread();
 	const bool bIncludeFarMeshOcclusionQueries = CVarWaterMeshOcclusionCullingIncludeFarMesh.GetValueOnGameThread() != 0;
-	OcclusionCullingBounds = WaterQuadTree.ComputeNodeBounds(MaxQueries, bIncludeFarMeshOcclusionQueries, &OcclusionResultsFarMeshOffset);
-	if (!OcclusionCullingBounds.IsEmpty())
-	{
-		EmptyOcclusionCullingBounds.Add(OcclusionCullingBounds[0]);
-	}
+	OcclusionCullingBounds = WaterQuadTree.ComputeNodeBounds(MaxQueries, CVarWaterMeshOcclusionCullExpandBoundsAmountXY.GetValueOnGameThread(), bIncludeFarMeshOcclusionQueries, &OcclusionResultsFarMeshOffset);
+	EmptyOcclusionCullingBounds.Add(WaterQuadTree.GetBoundsIncludingFarMesh());
 }
 
 FWaterMeshSceneProxy::~FWaterMeshSceneProxy()
diff --git a/Engine/Plugins/Experimental/Water/Source/Runtime/Private/WaterQuadTree.cpp b/Engine/Plugins/Experimental/Water/Source/Runtime/Private/WaterQuadTree.cpp
--- a/Engine/Plugins/Experimental/Water/Source/Runtime/Private/WaterQuadTree.cpp
+++ b/Engine/Plugins/Experimental/Water/Source/Runtime/Private/WaterQuadTree.cpp
@@ -671,6 +671,14 @@ void FWaterQuadTree::AddFarMesh(const UMaterialInterface* InFarMeshMaterial, con
 	// |_|_|_|
 	FarMeshData.InstanceData.SetNum(8);
 	FarMeshData.Material = InFarMeshMaterial;
+	
+	{
+		// Center position of the total far mesh bounds
+		const FVector FarMeshWorldCenter = FVector(InInnerRegion.GetCenter(), InFarDistanceMeshHeight);
+		// 3D Extents from the center of the far mesh bounds to the edge. Bounds are 2cm tall (1cm above, 1cm below plane)
+		const FVector FarMeshTotalExtent = FVector(InInnerRegion.GetExtent() + InFarDistanceMeshExtent, 1.0);
+		FarMeshData.FarMeshBounds = FBox(FarMeshWorldCenter - FarMeshTotalExtent, FarMeshWorldCenter + FarMeshTotalExtent);
+	}
 
 	const FVector2D WaterCenter = InInnerRegion.GetCenter();
 	const FVector2D WaterExtents = InInnerRegion.GetExtent();
@@ -862,7 +870,7 @@ void FWaterQuadTree::GatherHitProxies(TArray<TRefCountPtr<HHitProxy> >& OutHitPr
 }
 #endif //WITH_WATER_SELECTION_SUPPORT
 
-TArray<FBoxSphereBounds> FWaterQuadTree::ComputeNodeBounds(int32 MaxNumBounds, bool bIncludeFarMeshTiles, int32* OutFarMeshOffset) const
+TArray<FBoxSphereBounds> FWaterQuadTree::ComputeNodeBounds(int32 MaxNumBounds, float OcclusionCullExpandBoundsAmountXY, bool bIncludeFarMeshTiles, int32* OutFarMeshOffset) const
 {
 	const int32 NumNodes = NodeData.Nodes.Num();
 	const int32 ClampedNumNodes = (MaxNumBounds > 0) ? FMath::Min(MaxNumBounds, NumNodes) : NumNodes;
@@ -875,7 +883,9 @@ TArray<FBoxSphereBounds> FWaterQuadTree::ComputeNodeBounds(int32 MaxNumBounds, b
 	// Quadtree tiles
 	for (int32 i = 0; i < ClampedNumNodes; ++i)
 	{
-		Result.Add(NodeData.Nodes[BreadthFirstOrder[i]].Bounds);
+		// Expand bounds by given amount, only in XY to reduce occlusion 
+		const FBox Bounds = NodeData.Nodes[BreadthFirstOrder[i]].Bounds;
+		Result.Add(Bounds.ExpandBy(FVector(OcclusionCullExpandBoundsAmountXY, OcclusionCullExpandBoundsAmountXY, 0.0)));
 	}
 
 	*OutFarMeshOffset = Result.Num();
@@ -887,7 +897,7 @@ TArray<FBoxSphereBounds> FWaterQuadTree::ComputeNodeBounds(int32 MaxNumBounds, b
 		{
 			const FFarMeshData::FFarMeshInstanceData& InstanceData = FarMeshData.InstanceData[i];
 			const FVector Extent = FVector(InstanceData.Scale.X * 0.5, InstanceData.Scale.Y * 0.5, 1.0);
-			const FBox Bounds = FBox(InstanceData.WorldPosition - Extent, InstanceData.WorldPosition + Extent);
+			const FBox Bounds = FBox(InstanceData.WorldPosition - Extent, InstanceData.WorldPosition + Extent).ExpandBy(FVector(OcclusionCullExpandBoundsAmountXY, OcclusionCullExpandBoundsAmountXY, 0.0));
 			Result.Add(Bounds);
 		}
 	}
diff --git a/Engine/Plugins/Experimental/Water/Source/Runtime/Public/WaterQuadTree.h b/Engine/Plugins/Experimental/Water/Source/Runtime/Public/WaterQuadTree.h
--- a/Engine/Plugins/Experimental/Water/Source/Runtime/Public/WaterQuadTree.h
+++ b/Engine/Plugins/Experimental/Water/Source/Runtime/Public/WaterQuadTree.h
@@ -150,6 +150,9 @@ struct FWaterQuadTree
 
 	/** Get bounds of the root node if there is one, otherwise some default box */
 	FBox GetBounds() const { return NodeData.Nodes.Num() > 0 ? NodeData.Nodes[0].Bounds : FBox(-FVector::OneVector, FVector::OneVector); }
+
+	/** Get bounds of the root node if there is one (including far mesh), otherwise some default box */
+	FBox GetBoundsIncludingFarMesh() const { return GetBounds() + FarMeshData.FarMeshBounds; }
 	
 	/** Return the 2D region containing water tiles. Tiles can not be generated outside of this region */
 	FBox2D GetTileRegion() const { return TileRegion; }
@@ -191,7 +194,7 @@ struct FWaterQuadTree
 	void GatherHitProxies(TArray<TRefCountPtr<HHitProxy> >& OutHitProxies) const;
 #endif // WITH_WATER_SELECTION_SUPPORT
 
-	TArray<FBoxSphereBounds> ComputeNodeBounds(int32 MaxNumBounds, bool bIncludeFarMeshTiles, int32* OutFarMeshOffset) const;
+	TArray<FBoxSphereBounds> ComputeNodeBounds(int32 MaxNumBounds, float OcclusionCullExpandBoundsAmountXY, bool bIncludeFarMeshTiles, int32* OutFarMeshOffset) const;
 
 private:
 	struct FNodeData;
@@ -312,6 +315,7 @@ private:
 			InstanceData.Empty();
 			Material = nullptr;
 			MaterialIndex = INDEX_NONE;
+			FarMeshBounds.Init();
 		}
 
 		/** Total memory dynamically allocated by this object */
@@ -320,6 +324,9 @@ private:
 		/** Cached material index */
 		int16 MaterialIndex = INDEX_NONE;
 
+		/** FarMesh bounds */
+		FBox FarMeshBounds;
+
 	} FarMeshData;
 
 	/** If true, the tree may not change */
