Implemented a path for generating and traversing the water quadtree on the GPU. Traversal creates indirect draw arguments which are used to draw the tessellated water mesh tiles. The quadtree is created by rasterizing water meshes into a supersampled and multisampled texture, downsampling that texture to quadtree LOD0 resolution and finally building a mip chain. Every frame when water needs to be rendered, all nodes (texels) in that quadtree texture are visited by a compute shader to determine which ones need to be drawn from a given view. The generated list of tiles is sorted into draw buckets (quad mesh tessellation and material) and finally a list of per-instance data is generated that is passed to the modified water vertex factory. The VF can now handle indirect draws with instance data read from a compacted buffer using another indirection to get to the correct offset.

The indirect draw pipeline also supports HZB based occlusion culling.

#rnx

diff --git a/Engine/Plugins/Experimental/Water/Shaders/Private/WaterMeshVertexFactory.ush b/Engine/Plugins/Experimental/Water/Shaders/Private/WaterMeshVertexFactory.ush
--- a/Engine/Plugins/Experimental/Water/Shaders/Private/WaterMeshVertexFactory.ush
+++ b/Engine/Plugins/Experimental/Water/Shaders/Private/WaterMeshVertexFactory.ush
@@ -2,6 +2,23 @@
 
 #include "/Engine/Private/VertexFactoryCommon.ush"
 
+#ifndef WATER_MESH_DRAW_INDIRECT
+#define WATER_MESH_DRAW_INDIRECT 0
+#endif
+
+#if WATER_MESH_DRAW_INDIRECT
+
+Buffer<uint> InstanceDataOffsetsBuffer;
+Buffer<float4> InstanceData0Buffer;
+Buffer<float4> InstanceData1Buffer;
+#if WITH_WATER_SELECTION_SUPPORT_VF
+Buffer<float4> InstanceData2Buffer;
+#endif
+
+uint DrawBucketIndex;
+
+#endif
+
 struct FVertexFactoryInterpolantsVSToPS
 {
 #if NUM_TEX_COORD_INTERPOLATORS
@@ -26,11 +43,15 @@ struct FVertexFactoryInput
 {
 	float4	Position	: ATTRIBUTE0;
 
+#if WATER_MESH_DRAW_INDIRECT && !RAYHITGROUPSHADER && !COMPUTESHADER
+	uint InstanceId : SV_InstanceID;
+#else
 	float4 InstanceData0 : ATTRIBUTE8;
 	float4 InstanceData1 : ATTRIBUTE9; 
 #if HIT_PROXY_SHADER || WITH_WATER_SELECTION_SUPPORT_VF
 	float4 InstanceData2 : ATTRIBUTE10; // Store an HitProxyID per instance
 #endif
+#endif // WATER_MESH_DRAW_INDIRECT && !RAYHITGROUPSHADER && !COMPUTESHADER
 
 	VF_GPUSCENE_DECLARE_INPUT_BLOCK(13)
 	VF_INSTANCED_STEREO_DECLARE_INPUT_BLOCK()
@@ -43,8 +64,12 @@ struct FPositionOnlyVertexFactoryInput
 {
 	float4	Position	: ATTRIBUTE0;
 
+#if WATER_MESH_DRAW_INDIRECT && !RAYHITGROUPSHADER && !COMPUTESHADER
+	uint InstanceId : SV_InstanceID;
+#else
 	float4 InstanceData0 : ATTRIBUTE8;
 	float4 InstanceData1 : ATTRIBUTE9; 
+#endif // WATER_MESH_DRAW_INDIRECT && !RAYHITGROUPSHADER && !COMPUTESHADER
 
 	VF_GPUSCENE_DECLARE_INPUT_BLOCK(1)
 	VF_INSTANCED_STEREO_DECLARE_INPUT_BLOCK()
@@ -58,8 +83,12 @@ struct FPositionAndNormalOnlyVertexFactoryInput
 	float4	Position	: ATTRIBUTE0;
 	float4	Normal		: ATTRIBUTE2;
 
+#if WATER_MESH_DRAW_INDIRECT && !RAYHITGROUPSHADER && !COMPUTESHADER
+	uint InstanceId : SV_InstanceID;
+#else
 	float4 InstanceData0 : ATTRIBUTE8;
 	float4 InstanceData1 : ATTRIBUTE9; 
+#endif // WATER_MESH_DRAW_INDIRECT && !RAYHITGROUPSHADER && !COMPUTESHADER
 
 	VF_GPUSCENE_DECLARE_INPUT_BLOCK(1)
 	VF_INSTANCED_STEREO_DECLARE_INPUT_BLOCK()
@@ -209,8 +238,14 @@ FVertexFactoryIntermediates GetVertexFactoryIntermediates(FVertexFactoryInput In
 	FVertexFactoryIntermediates Intermediates;
 
 	// Get the packed instance data
+#if WATER_MESH_DRAW_INDIRECT && !RAYHITGROUPSHADER && !COMPUTESHADER
+	const uint InstanceDataIndex = InstanceDataOffsetsBuffer.Load(DrawBucketIndex) + Input.InstanceId;
+	float4 Data0 = InstanceData0Buffer.Load(InstanceDataIndex);
+	float4 Data1 = InstanceData1Buffer.Load(InstanceDataIndex);
+#else
 	float4 Data0 = Input.InstanceData0;
 	float4 Data1 = Input.InstanceData1;
+#endif // WATER_MESH_DRAW_INDIRECT && !RAYHITGROUPSHADER && !COMPUTESHADER
 
 	const float3 Translation = Data0.xyz;
 	Intermediates.WaveParamIndex = asuint(Data0.w);
@@ -254,7 +289,11 @@ FVertexFactoryIntermediates GetVertexFactoryIntermediates(FVertexFactoryInput In
 	}
 
 #if WITH_WATER_SELECTION_SUPPORT_VF
+#if WATER_MESH_DRAW_INDIRECT && !RAYHITGROUPSHADER && !COMPUTESHADER
+	float SelectedValue = InstanceData2Buffer.Load(InstanceDataIndex).w;
+#else
 	float SelectedValue = Input.InstanceData2.w;
+#endif
 	float IsVisible = WaterVF.bRenderSelected * SelectedValue + WaterVF.bRenderUnselected * (1-SelectedValue);
 	Intermediates.MorphedTranslatedWorldPos *= IsVisible;
 #endif // WITH_WATER_SELECTION_SUPPORT_VF
@@ -268,7 +307,12 @@ FVertexFactoryIntermediates GetVertexFactoryIntermediates(FVertexFactoryInput In
 float4 VertexFactoryGetInstanceHitProxyId(FVertexFactoryInput Input, FVertexFactoryIntermediates Intermediates)
 {
 #if WITH_WATER_SELECTION_SUPPORT_VF
+#if WATER_MESH_DRAW_INDIRECT && !RAYHITGROUPSHADER && !COMPUTESHADER
+	const uint InstanceDataIndex = InstanceDataOffsetsBuffer.Load(DrawBucketIndex) + Input.InstanceId;
+	return float4(InstanceData2Buffer.Load(InstanceDataIndex).rgb, 0);
+#else
 	return float4(Input.InstanceData2.rgb, 0);
+#endif
 #else // WITH_WATER_SELECTION_SUPPORT_VF
 	return 0;
 #endif // !WITH_WATER_SELECTION_SUPPORT_VF
diff --git a/Engine/Plugins/Experimental/Water/Shaders/Private/WaterQuadTreeBuild.usf b/Engine/Plugins/Experimental/Water/Shaders/Private/WaterQuadTreeBuild.usf
new file mode 100644
--- /dev/null
+++ b/Engine/Plugins/Experimental/Water/Shaders/Private/WaterQuadTreeBuild.usf
@@ -0,0 +1,54 @@
+// Copyright Epic Games, Inc. All Rights Reserved.
+
+#include "/Engine/Private/Common.ush"
+#include "/Engine/Private/PackUnpack.ush"
+#include "WaterQuadTreeCommon.ush"
+
+Texture2D<float4> QuadTreeTexture;
+Texture2D<float4> WaterZBoundsTexture;
+
+void Main(in float4 Position : SV_Position, out float4 OutColor : SV_Target0, out float4 OutWaterZBounds : SV_Target1)
+{
+	const int2 BaseCoord = int2(Position.xy) * 2;
+
+	uint MinWaterBodyRenderDataIndex = 0xFFFFFFFFu;
+	uint MaxWaterBodyRenderDataIndex = 0x0u;
+	uint MinTransitionWaterBodyRenderDataIndex = 0xFFFFFFFFu;
+	uint MaxTransitionWaterBodyRenderDataIndex = 0x0u;
+	bool bHasCompleteSubtree = true;
+	float MinZ = 1.0f;
+	float MaxZ = 0.0f;
+	float SurfaceBaseHeightSum = 0.0f;
+
+	// Downsample from lower mip
+	for (int y = 0; y < 2; ++y)
+	{
+		for (int x = 0; x < 2; ++x)
+		{
+			const float4 QuadTreeSample = QuadTreeTexture.Load(int3(BaseCoord + int2(x, y), 0));
+			const FWaterQuadTreeNode ChildNode = WaterQuadTreeUnpackNodeRGBA8(QuadTreeSample);
+			
+			bHasCompleteSubtree &= ChildNode.bHasCompleteSubtree;
+			MinWaterBodyRenderDataIndex = min(MinWaterBodyRenderDataIndex, ChildNode.WaterBodyRenderDataIndex);
+			MaxWaterBodyRenderDataIndex = max(MaxWaterBodyRenderDataIndex, ChildNode.WaterBodyRenderDataIndex);
+			MinTransitionWaterBodyRenderDataIndex = min(MinTransitionWaterBodyRenderDataIndex, ChildNode.TransitionWaterBodyRenderDataIndex);
+			MaxTransitionWaterBodyRenderDataIndex = max(MaxTransitionWaterBodyRenderDataIndex, ChildNode.TransitionWaterBodyRenderDataIndex);
+
+			const float3 WaterZBoundsSample = WaterZBoundsTexture.Load(int3(BaseCoord + int2(x, y), 0)).xyz;
+			MinZ = min(MinZ, WaterZBoundsSample.x);
+			MaxZ = max(MaxZ, WaterZBoundsSample.y);
+			SurfaceBaseHeightSum += WaterZBoundsSample.z;
+		}
+	}
+
+	const bool bCanMerge = (MinWaterBodyRenderDataIndex == MaxWaterBodyRenderDataIndex) && (MinTransitionWaterBodyRenderDataIndex == MaxTransitionWaterBodyRenderDataIndex) && MaxWaterBodyRenderDataIndex > 0;
+
+	FWaterQuadTreeNode ResultNode = (FWaterQuadTreeNode)0;
+	ResultNode.WaterBodyRenderDataIndex = bCanMerge ? MinWaterBodyRenderDataIndex : 0;
+	ResultNode.TransitionWaterBodyRenderDataIndex = bCanMerge ? MinTransitionWaterBodyRenderDataIndex : 0;
+	ResultNode.bHasCompleteSubtree = bHasCompleteSubtree;
+	ResultNode.bIsSubtreeSameWaterBody = bCanMerge;
+
+	OutColor = WaterQuadTreePackNodeRGBA8(ResultNode);
+	OutWaterZBounds = bCanMerge ? float4(MinZ, MaxZ, SurfaceBaseHeightSum * 0.25f, 1.0f) : float4(0.0f, 0.0f, 0.0f, 1.0f);
+}
\ No newline at end of file
diff --git a/Engine/Plugins/Experimental/Water/Shaders/Private/WaterQuadTreeCommon.ush b/Engine/Plugins/Experimental/Water/Shaders/Private/WaterQuadTreeCommon.ush
new file mode 100644
--- /dev/null
+++ b/Engine/Plugins/Experimental/Water/Shaders/Private/WaterQuadTreeCommon.ush
@@ -0,0 +1,72 @@
+// Copyright Epic Games, Inc. All Rights Reserved.
+
+#include "/Engine/Private/PackUnpack.ush"
+
+#define WATER_BODY_TYPE_RIVER 0
+#define WATER_BODY_TYPE_LAKE 1
+#define WATER_BODY_TYPE_OCEAN 2
+
+struct FWaterBodyRenderData
+{
+	uint WaterBodyIndex;
+	uint MaterialIndex;
+	uint RiverToLakeMaterialIndex;
+	uint RiverToOceanMaterialIndex;
+	uint WaterBodyType;
+	uint HitProxyColorAndIsSelected;
+	float SurfaceBaseHeight;
+	float MinZ;
+	float MaxZ;
+	float MaxWaveHeight;
+};
+
+struct FWaterQuadTreeNode
+{
+	uint WaterBodyRenderDataIndex;
+	uint TransitionWaterBodyRenderDataIndex;
+	bool bHasCompleteSubtree;
+	bool bIsSubtreeSameWaterBody;
+};
+
+uint WaterQuadTreePackNode(FWaterQuadTreeNode Node)
+{
+	uint Result = 0;
+	Result |= Node.WaterBodyRenderDataIndex & 0x7FFFu; // 15 bits
+	Result |= (Node.TransitionWaterBodyRenderDataIndex & 0x7FFFu) << 15u;
+	Result |= uint(Node.bHasCompleteSubtree ? 1 : 0) << 30u;
+	Result |= uint(Node.bIsSubtreeSameWaterBody ? 1 : 0) << 31u;
+
+	return Result;
+}
+
+float4 WaterQuadTreePackNodeRGBA8(FWaterQuadTreeNode Node)
+{
+	uint PackedUI = WaterQuadTreePackNode(Node);
+	float4 Packed4;
+	Packed4.x = ((PackedUI >> 0u) & 0xFFu) / 255.0f;
+	Packed4.y = ((PackedUI >> 8u) & 0xFFu) / 255.0f;
+	Packed4.z = ((PackedUI >> 16u) & 0xFFu) / 255.0f;
+	Packed4.w = ((PackedUI >> 24u) & 0xFFu) / 255.0f;
+	return Packed4;
+}
+
+FWaterQuadTreeNode WaterQuadTreeUnpackNode(uint Packed)
+{
+	FWaterQuadTreeNode Result;
+	Result.WaterBodyRenderDataIndex = Packed & 0x7FFFu;
+	Result.TransitionWaterBodyRenderDataIndex = (Packed >> 15u) & 0x7FFFu;
+	Result.bHasCompleteSubtree = (Packed & (1u << 30u)) != 0;
+	Result.bIsSubtreeSameWaterBody = (Packed & (1u << 31u)) != 0;
+
+	return Result;
+}
+
+FWaterQuadTreeNode WaterQuadTreeUnpackNodeRGBA8(float4 Packed)
+{
+	uint PackedUI = 0;
+	PackedUI |= uint(saturate(Packed.x) * 255.0f + 0.5f) << 0u;
+	PackedUI |= uint(saturate(Packed.y) * 255.0f + 0.5f) << 8u;
+	PackedUI |= uint(saturate(Packed.z) * 255.0f + 0.5f) << 16u;
+	PackedUI |= uint(saturate(Packed.w) * 255.0f + 0.5f) << 24u;
+	return WaterQuadTreeUnpackNode(PackedUI);
+}
diff --git a/Engine/Plugins/Experimental/Water/Shaders/Private/WaterQuadTreeDraws.usf b/Engine/Plugins/Experimental/Water/Shaders/Private/WaterQuadTreeDraws.usf
new file mode 100644
--- /dev/null
+++ b/Engine/Plugins/Experimental/Water/Shaders/Private/WaterQuadTreeDraws.usf
@@ -0,0 +1,675 @@
+// Copyright Epic Games, Inc. All Rights Reserved.
+
+#include "/Engine/Private/Common.ush"
+#include "/Engine/Private/PackUnpack.ush"
+#include "WaterQuadTreeCommon.ush"
+
+struct FWaterTileDrawIndexedIndirectArgs
+{
+	uint IndexCountPerInstance;
+	uint InstanceCount;
+	uint StartIndexLocation;
+	int  BaseVertexLocation;
+	uint StartInstanceLocation;
+};
+
+uint PackNodeCoord(uint4 InNode)
+{
+	uint Result = 0;
+	Result |= uint(InNode.x & 0x7FFu);
+	Result |= uint(InNode.y & 0x7FFu) << 11u;
+	Result |= uint(InNode.z & 0x1Fu) << 22u;
+	Result |= uint(InNode.w & 0x1Fu) << 27u;
+	return Result;
+}
+
+uint4 UnpackNodeCoord(uint InPacked)
+{
+	uint4 Result;
+	Result.x = InPacked & 0x7FFu;
+	Result.y = (InPacked >> 11u) & 0x7FFu;
+	Result.z = (InPacked >> 22u) & 0x1Fu;
+	Result.w = (InPacked >> 27u) & 0x1Fu;
+	return Result;
+}
+
+struct FWaterLODParams
+{
+	int LowestLOD;
+	float HeightLODFactor;
+};
+
+FWaterLODParams GetWaterLODParams(float InObserverHeight, float InWaterHeightForLOD, float InLODScale, uint InTreeDepth)
+{
+	float DistToWater = abs(InObserverHeight - InWaterHeightForLOD) / InLODScale;
+	DistToWater = max(DistToWater - 2.0f, 0.0f);
+	DistToWater *= 2.0f;
+
+	const float FloatLOD = clamp(log2(DistToWater), 0.0f, InTreeDepth - 1.0f);
+
+	FWaterLODParams WaterLODParams;
+	WaterLODParams.LowestLOD = clamp(floor(FloatLOD), 0, InTreeDepth - 1);
+	WaterLODParams.HeightLODFactor = frac(FloatLOD);
+
+	return WaterLODParams;
+}
+
+#ifdef INITIALIZE_INDIRECT_ARGS
+
+RWStructuredBuffer<FWaterTileDrawIndexedIndirectArgs> IndirectArgs;
+uint NumDensities;
+uint NumMaterials;
+uint NumViews;
+uint NumQuadsLOD0;
+
+[numthreads(64, 1, 1)]
+void MainCS(uint3 DispatchThreadID : SV_DispatchThreadID)
+{
+	const uint NumBuckets = NumDensities * NumMaterials;
+	const uint BucketIndex = DispatchThreadID.x;
+
+	if (BucketIndex >= NumBuckets)
+	{
+		return;
+	}
+
+	const uint DensityIndex = BucketIndex % NumDensities;
+	const uint MaterialIndex = BucketIndex / NumDensities;
+	const uint NumQuads = NumQuadsLOD0 >> DensityIndex;
+
+	FWaterTileDrawIndexedIndirectArgs BucketArgs;
+	BucketArgs.IndexCountPerInstance = NumQuads * NumQuads * 6;
+	BucketArgs.InstanceCount = 0;
+	BucketArgs.StartIndexLocation = 0;
+	BucketArgs.BaseVertexLocation = 0;
+	BucketArgs.StartInstanceLocation = 0;
+
+	for (uint ViewIndex = 0; ViewIndex < NumViews; ++ViewIndex)
+	{
+		IndirectArgs[NumBuckets * ViewIndex + BucketIndex] = BucketArgs;
+	}
+}
+
+#endif // INITIALIZE_INDIRECT_ARGS
+
+#ifdef QUAD_TREE_TRAVERSE
+
+#ifndef DEBUG_SHADER_PRINT
+#define DEBUG_SHADER_PRINT 0
+#endif
+
+#include "/Engine/Private/Nanite/NaniteHZBCull.ush" // BoxCullFrustum()
+#if DEBUG_SHADER_PRINT
+#include "/Engine/Private/ShaderPrint.ush"
+#endif
+
+RWByteAddressBuffer PackedNodes;
+Texture2D<float4> QuadTreeTexture;
+Texture2D<float4> WaterZBoundsTexture;
+StructuredBuffer<FWaterBodyRenderData> WaterBodyRenderData;
+float4x4 TranslatedWorldToClip;
+float4x4 ViewToClip;
+float4 CullingBoundsAABB;
+float3 QuadTreePosition;
+float3 ObserverPosition;
+uint QuadTreeResolutionX;
+uint QuadTreeResolutionY;
+uint NumDensities;
+uint NumMaterials;
+float LeafSize;
+float LODScale;
+float CaptureDepthRange;
+int ForceCollapseDensityLevel;
+uint NumLODs;
+uint NumDispatchedThreads;
+int DebugShowTile;
+uint bHZBOcclusionCullingEnabled;
+
+
+void AddNodeForRender(uint InNodeCoordPacked)
+{
+	uint WritePos = 0;
+	PackedNodes.InterlockedAdd(0, 1, WritePos);
+	WritePos += 1; // Num is stored at index 0
+	PackedNodes.Store(WritePos << 2u, InNodeCoordPacked);
+}
+
+float ComputeSquaredDistanceToPoint(float2 InBox2DMin, float2 InBox2DMax, float2 InPoint)
+{
+	// Accumulates the distance as we iterate axis
+	float DistSquared = 0.0f;
+
+	if (InPoint.x < InBox2DMin.x)
+	{
+		DistSquared += Square(InPoint.x - InBox2DMin.x);
+	}
+	else if (InPoint.x > InBox2DMax.x)
+	{
+		DistSquared += Square(InPoint.x - InBox2DMax.x);
+	}
+
+	if (InPoint.y < InBox2DMin.y)
+	{
+		DistSquared += Square(InPoint.y - InBox2DMin.y);
+	}
+	else if (InPoint.y > InBox2DMax.y)
+	{
+		DistSquared += Square(InPoint.y - InBox2DMax.y);
+	}
+
+	return DistSquared;
+}
+
+float GetLODDistance(int InLODLevel, float InLODScale)
+{
+	return pow(2.0f, (float)(InLODLevel + 1)) * InLODScale;
+}
+
+bool CanRender(FWaterQuadTreeNode Node, int InDensityLevel, int InForceCollapseDensityLevel)
+{
+	int MaterialIndex = -1;
+	// There is a dummy entry at index 0 but we know that it has no valid MaterialIndex, so we can skip this case
+	if (Node.WaterBodyRenderDataIndex > 0)
+	{
+		const FWaterBodyRenderData WBRenderData = WaterBodyRenderData[Node.WaterBodyRenderDataIndex];
+		MaterialIndex = WBRenderData.MaterialIndex;
+	}
+	// Can render if the density level is (in addition to same water bodies in all descendants) either above the force collapse level or if the subtree is complete
+	return MaterialIndex >= 0 && Node.bIsSubtreeSameWaterBody && ((InDensityLevel > InForceCollapseDensityLevel) || Node.bHasCompleteSubtree);
+}
+
+float4 GetNodeAABB2D(uint3 InNodeCoord)
+{
+	const float Scale = (1u << InNodeCoord.z) * LeafSize;
+	return float4(InNodeCoord.xy, InNodeCoord.xy + 1.0f) * Scale + QuadTreePosition.xyxy;
+}
+
+[numthreads(64, 1, 1)]
+void MainCS(uint3 DispatchThreadID : SV_DispatchThreadID)
+{
+	if (DispatchThreadID.x >= NumDispatchedThreads)
+	{
+		return;
+	}
+
+	// Compute which LOD level this thread belongs to and then find the node index relative to that LOD level
+	uint3 NodeCoord = 0;
+	{
+		uint LODLevelIndex = 0;
+		uint2 Resolution = uint2(QuadTreeResolutionX, QuadTreeResolutionY);
+		uint NumPreviousNodes = 0;
+		while (LODLevelIndex < NumLODs && DispatchThreadID.x >= (Resolution.x * Resolution.y + NumPreviousNodes))
+		{
+			NumPreviousNodes += Resolution.x * Resolution.y;
+			Resolution = uint2(max(1u, Resolution.x / 2), max(1u, Resolution.y / 2));
+			++LODLevelIndex;
+		}
+
+		const uint LocalNodeIndex = DispatchThreadID.x - NumPreviousNodes;
+		NodeCoord = uint3(LocalNodeIndex % Resolution.x, LocalNodeIndex / Resolution.x, LODLevelIndex);
+	}
+
+	const uint LODLevel = NodeCoord.z;
+	const float4 NodeAABB2D = GetNodeAABB2D(NodeCoord);
+	
+	// Bounds culling
+	{
+		const float4 CullingBoundsAABBTWS = CullingBoundsAABB + QuadTreePosition.xyxy;
+		if (any(NodeAABB2D.zw < CullingBoundsAABBTWS.xy) || any(NodeAABB2D.xy > CullingBoundsAABBTWS.zw))
+		{
+			return;
+		}
+	}
+
+	// Frustum culling
+	{
+		const float2 WaterZBounds = WaterZBoundsTexture.Load(NodeCoord).xy * CaptureDepthRange + QuadTreePosition.z;
+		const float3 NodeAABBMin = float3(NodeAABB2D.xy, WaterZBounds.x);
+		const float3 NodeAABBMax = float3(NodeAABB2D.zw, WaterZBounds.y);
+		const float3 NodeCenter = (NodeAABBMin + NodeAABBMax) * 0.5f;
+		const float3 NodeExtent = float3((NodeAABBMax.xy - NodeAABBMin.xy) * 0.5f, (NodeAABBMax.z - NodeAABBMin.z) * 0.5f + 1e-5f);
+	
+		const bool bIsOrtho = IsOrthoProjection(ViewToClip);
+		FFrustumCullData Cull = BoxCullFrustum(NodeCenter, NodeExtent, TranslatedWorldToClip, ViewToClip, bIsOrtho, true /* near clip */, false /* skip culling */);
+
+		if (bHZBOcclusionCullingEnabled && !Cull.bCrossesNearPlane)
+		{
+			const FScreenRect Rect = GetScreenRect(int4(0, 0, HZBViewSize), Cull, 4);
+			Cull.bIsVisible = IsVisibleHZB(Rect, true);
+		}
+
+		if (!Cull.bIsVisible)
+		{
+			return;
+		}
+	}
+
+	// Compute the 2D distance between the observer and the tile. Will be zero if the observer is within the tile
+	const float ClosestDistanceToTile = sqrt(ComputeSquaredDistanceToPoint(NodeAABB2D.xy, NodeAABB2D.zw, ObserverPosition.xy));
+
+	// Compute the lowest LOD we want to render, based on the observer height above the water surface
+	const float3 ObserverPos = ObserverPosition - QuadTreePosition;
+	const int2 TextureLoadCoord = (int2)clamp(ObserverPos.xy / LeafSize, 0.0f, float2(QuadTreeResolutionX, QuadTreeResolutionY) - 0.5f);
+	const float WaterHeightForLOD = WaterZBoundsTexture.Load(int3(TextureLoadCoord, 0)).z * CaptureDepthRange;
+	const FWaterLODParams WaterLODParams = GetWaterLODParams(ObserverPos.z, WaterHeightForLOD, LODScale, NumLODs - 1);
+
+	uint DensityLevel = 0;
+	bool bShouldRender = false;
+
+	// We might need to draw this tile even if it is outside the LOD range for tiles in this level of the quadtree.
+	// This is the case if tiles at a higher LOD can't be rendered, so we need to emulate a higher LOD tile with
+	// multiple lower LOD tiles, but using a lower density/tessellation.
+	const bool bOutsideLODRange = ClosestDistanceToTile > GetLODDistance(LODLevel, LODScale) || LODLevel < WaterLODParams.LowestLOD;
+	if ((LODLevel < (NumLODs - 1)) && bOutsideLODRange)
+	{
+		// Does the parent tile intersect the LOD range of tiles at this LOD?
+		uint3 ParentNodeCoord = uint3(NodeCoord.xy / 2, NodeCoord.z + 1);
+		float4 ParentNodeAABB2D = GetNodeAABB2D(ParentNodeCoord);
+		const float ClosestDistanceToParentTile = sqrt(ComputeSquaredDistanceToPoint(ParentNodeAABB2D.xy, ParentNodeAABB2D.zw, ObserverPosition.xy));
+		const bool bParentTileIntersectsThisLOD = ClosestDistanceToParentTile <= GetLODDistance(LODLevel, LODScale) && LODLevel >= WaterLODParams.LowestLOD;
+
+		// Is the parent tile renderable?
+		const FWaterQuadTreeNode ParentNode = WaterQuadTreeUnpackNodeRGBA8(QuadTreeTexture.Load(ParentNodeCoord));
+		const bool bParentTileCanRender = CanRender(ParentNode, DensityLevel, ForceCollapseDensityLevel);
+
+		// If the parent tile intersects this LOD, we don't render the parent tile and instead render its children.
+		// If the parent tile can't render at all, then we need to emulate it by rendering its children at the highest renderable LOD.
+		if (bParentTileIntersectsThisLOD || !bParentTileCanRender)
+		{
+			// In both cases we need to compute the density we need to use to make it appear like we rendered the higher LOD tile.
+			for (uint ParentLOD = LODLevel + 1; ParentLOD < NumLODs; ++ParentLOD)
+			{
+				DensityLevel = ParentLOD - LODLevel;
+
+				ParentNodeCoord = uint3(NodeCoord.xy >> (ParentLOD - LODLevel), ParentLOD);
+				ParentNodeAABB2D = GetNodeAABB2D(ParentNodeCoord);
+
+				const float Dist = sqrt(ComputeSquaredDistanceToPoint(ParentNodeAABB2D.xy, ParentNodeAABB2D.zw, ObserverPosition.xy));
+				if (Dist <= GetLODDistance(ParentLOD, LODScale) && ParentLOD >= WaterLODParams.LowestLOD)
+				{
+					break;
+				}
+			}
+			bShouldRender = true;
+		}
+	}
+	// The tile is fully within its LOD range and does not intersect the LOD range of its children (Distance <= GetLODDistance(LODLevel) && Distance > GetLODDistance(LODLevel - 1)), 
+	// so we can render it as is.
+	else if (LODLevel == WaterLODParams.LowestLOD || ClosestDistanceToTile > GetLODDistance(LODLevel - 1, LODScale))
+	{
+		bShouldRender = true;
+	}
+
+	if (bShouldRender)
+	{
+		const FWaterQuadTreeNode Node = WaterQuadTreeUnpackNodeRGBA8(QuadTreeTexture.Load(NodeCoord));
+		if (CanRender(Node, DensityLevel, ForceCollapseDensityLevel))
+		{
+			AddNodeForRender(PackNodeCoord(uint4(NodeCoord, DensityLevel)));
+
+#if DEBUG_SHADER_PRINT
+			if (DebugShowTile != 0)
+			{
+				FWaterBodyRenderData WBRenderData = WaterBodyRenderData[Node.WaterBodyRenderDataIndex];
+
+				// TODO: Support the two missing modes for visualizing the LOD level and the density
+				float4 DebugColor;
+				switch (WBRenderData.WaterBodyType)
+				{
+				case WATER_BODY_TYPE_RIVER: DebugColor = ColorRed; break;
+				case WATER_BODY_TYPE_LAKE: DebugColor = ColorGreen; break;
+				case WATER_BODY_TYPE_OCEAN: DebugColor = ColorBlue; break;
+				default: DebugColor = ColorWhite;
+				}
+
+				if (WBRenderData.WaterBodyType == WATER_BODY_TYPE_RIVER && Node.TransitionWaterBodyRenderDataIndex > 0)
+				{
+					const FWaterBodyRenderData TransitionWBRenderData = WaterBodyRenderData[Node.TransitionWaterBodyRenderDataIndex];
+					if (TransitionWBRenderData.WaterBodyType == WATER_BODY_TYPE_LAKE)
+					{
+						DebugColor = ColorYellow;
+					}
+					else if (TransitionWBRenderData.WaterBodyType == WATER_BODY_TYPE_OCEAN)
+					{
+						DebugColor = ColorPurple;
+					}
+				}
+
+				const float2 WaterZBounds = WaterZBoundsTexture.Load(NodeCoord).xy * CaptureDepthRange + QuadTreePosition.z;
+				const float3 AABBMin = float3(NodeAABB2D.xy + 20.0f, WaterZBounds.x);
+				const float3 AABBMax = float3(NodeAABB2D.zw - 20.0f, WaterZBounds.y);
+				
+				AddOBBTWS(InitShaderPrintContext(), AABBMin, AABBMax, DebugColor, float4x4(
+					1, 0, 0, 0, 
+					0, 1, 0, 0, 
+					0, 0, 1, 0, 
+					0, 0, 0, 1));
+			}
+#endif // DEBUG_SHADER_PRINT
+		}
+	}
+}
+
+#endif // QUAD_TREE_TRAVERSE
+
+#ifdef COMPUTE_BUCKET_COUNTS
+
+RWByteAddressBuffer BucketCounts;
+Texture2D<float4> QuadTreeTexture;
+StructuredBuffer<FWaterBodyRenderData> WaterBodyRenderData;
+ByteAddressBuffer PackedNodes;
+uint NumDensities;
+uint NumMaterials;
+uint NumDispatchedThreads;
+
+[numthreads(64, 1, 1)]
+void MainCS(uint3 DispatchThreadID : SV_DispatchThreadID)
+{
+	const uint NumDraws = PackedNodes.Load(0);
+
+	// We are limited to 65535 threads, but might end up with more draws than that
+	for (uint DrawIndex = DispatchThreadID.x; DrawIndex < NumDraws; DrawIndex += NumDispatchedThreads)
+	{
+		// Load packed node coord and density index
+		uint4 NodeCoordAndDensity = UnpackNodeCoord(PackedNodes.Load((DrawIndex + 1) << 2));
+		const uint DensityIndex = min(NodeCoordAndDensity.w, NumDensities - 1);
+
+		// Sample the quadtree to access the water body render data
+		const float4 QuadTreeSample = QuadTreeTexture.Load(NodeCoordAndDensity.xyz);
+		const FWaterQuadTreeNode Node = WaterQuadTreeUnpackNodeRGBA8(QuadTreeSample);
+		const FWaterBodyRenderData WBRenderData = WaterBodyRenderData[Node.WaterBodyRenderDataIndex];
+
+		// Determine the material
+		uint MaterialIndex = WBRenderData.MaterialIndex;
+		// Rivers can have transitions to other water bodies
+		if (WBRenderData.WaterBodyType == WATER_BODY_TYPE_RIVER && Node.TransitionWaterBodyRenderDataIndex > 0)
+		{
+			const FWaterBodyRenderData TransitionWBRenderData = WaterBodyRenderData[Node.TransitionWaterBodyRenderDataIndex];
+			if (TransitionWBRenderData.WaterBodyType == WATER_BODY_TYPE_LAKE)
+			{
+				MaterialIndex = WBRenderData.RiverToLakeMaterialIndex;
+			}
+			else if (TransitionWBRenderData.WaterBodyType == WATER_BODY_TYPE_OCEAN)
+			{
+				MaterialIndex = WBRenderData.RiverToOceanMaterialIndex;
+			}
+		}
+
+		// Increment bucket counter
+		const uint BucketIndex = MaterialIndex * NumDensities + DensityIndex;
+		uint Dummy;
+		BucketCounts.InterlockedAdd(BucketIndex << 2, 1, Dummy);
+	}
+}
+
+#endif // COMPUTE_BUCKET_COUNTS
+
+#ifdef COMPUTE_BUCKET_PREFIX_SUM
+
+#ifndef PARALLEL_PREFIX_SUM
+#define PARALLEL_PREFIX_SUM 0
+#endif
+
+RWBuffer<uint> BucketPrefixSums;
+ByteAddressBuffer BucketCounts;
+uint NumBuckets;
+uint OutputOffset;
+
+#if PARALLEL_PREFIX_SUM
+
+#include "/Engine/Private/WaveOpUtil.ush"
+
+#define GROUP_SIZE 128
+#define ARRAY_SIZE (GROUP_SIZE * 2)
+groupshared uint SharedData[ARRAY_SIZE];
+groupshared uint BlockPrefixSum;
+
+uint LoadFromBuffer(uint InIndex)
+{
+	if (InIndex < NumBuckets)
+	{
+		return BucketCounts.Load(InIndex << 2);
+	}
+	else
+	{
+		return 0;
+	}
+}
+
+void WriteToBuffer(uint InIndex, uint InWriteOffset, uint InGlobalPrefixSum)
+{
+	if (InIndex < NumBuckets)
+	{
+		BucketPrefixSums[InIndex + InWriteOffset] = SharedData[InIndex] + InGlobalPrefixSum;
+	}
+}
+
+#endif
+
+#if PARALLEL_PREFIX_SUM
+[numthreads(GROUP_SIZE, 1, 1)]
+#else
+[numthreads(1, 1, 1)]
+#endif
+void MainCS(uint3 DispatchThreadID : SV_DispatchThreadID)
+{
+#if PARALLEL_PREFIX_SUM
+
+	const uint ThreadID = DispatchThreadID.x;
+	const uint NumBlocks = (NumBuckets + (ARRAY_SIZE - 1)) / ARRAY_SIZE;
+
+	if (ThreadID == 0)
+	{
+		BlockPrefixSum = 0;
+	}
+
+	// Process array in blocks of fixed size, keeping track of the sum of all elements in previous blocks
+	for (uint BlockIndex = 0; BlockIndex < NumBlocks; ++BlockIndex)
+	{
+		GroupMemoryBarrierWithGroupSync();
+
+		const uint BlockOffset = BlockIndex * ARRAY_SIZE;
+		uint Offset = 1;
+
+		// Load into LDS
+		const uint Value0 = LoadFromBuffer(2 * ThreadID + 0 + BlockOffset);
+		const uint Value1 = LoadFromBuffer(2 * ThreadID + 1 + BlockOffset);
+		const uint LocalSum = Value0 + Value1;
+		SharedData[2 * ThreadID + 0] = Value0;
+		SharedData[2 * ThreadID + 1] = Value1;
+
+		// Up-sweep
+		for (uint d = ARRAY_SIZE >> 1; d > 0; d >>= 1)
+		{
+			GroupMemoryBarrierWithGroupSync();
+
+			if (ThreadID < d)
+			{
+				const uint IndexA = Offset * (2 * ThreadID + 1) - 1;
+				const uint IndexB = Offset * (2 * ThreadID + 2) - 1;
+
+				SharedData[IndexB] += SharedData[IndexA];
+			}
+			Offset <<= 1;
+		}
+
+		// Clear the last element
+		if (ThreadID == 0)
+		{
+			SharedData[ARRAY_SIZE - 1] = 0;
+		}
+
+		// Down-sweep
+		for (uint d = 1; d < ARRAY_SIZE; d <<= 1)
+		{
+			Offset >>= 1;
+			GroupMemoryBarrierWithGroupSync();
+
+			if (ThreadID < d)
+			{
+				const uint IndexA = Offset * (2 * ThreadID + 1) - 1;
+				const uint IndexB = Offset * (2 * ThreadID + 2) - 1;
+
+				const uint Temp = SharedData[IndexA];
+				SharedData[IndexA] = SharedData[IndexB];
+				SharedData[IndexB] += Temp;
+			}
+		}
+
+		const uint GlobalPrefixSum = BlockPrefixSum;
+
+		GroupMemoryBarrierWithGroupSync();
+
+		// Write results to output buffer
+		WriteToBuffer(2 * ThreadID + 0, BlockOffset + OutputOffset, GlobalPrefixSum);
+		WriteToBuffer(2 * ThreadID + 1, BlockOffset + OutputOffset, GlobalPrefixSum);
+
+		if (NumBlocks > 1)
+		{
+			WaveInterlockedAdd(BlockPrefixSum, LocalSum);
+		}
+	}
+
+#else
+
+	uint PrefixSum = 0;
+	for (uint BucketIndex = 0; BucketIndex < NumBuckets; ++BucketIndex)
+	{
+		BucketPrefixSums[BucketIndex + OutputOffset] = PrefixSum;
+		const uint Count = BucketCounts.Load(BucketIndex << 2);
+		PrefixSum += Count;
+	}
+
+#endif // PARALLEL_PREFIX_SUM
+}
+
+#endif // COMPUTE_BUCKET_PREFIX_SUM
+
+#ifdef GENERATE_INSTANCE_DATA
+
+RWStructuredBuffer<FWaterTileDrawIndexedIndirectArgs> IndirectArgs;
+RWBuffer<float4> InstanceData0;
+RWBuffer<float4> InstanceData1;
+RWBuffer<float4> InstanceData2;
+Texture2D<float4> QuadTreeTexture;
+Texture2D<float4> WaterZBoundsTexture;
+StructuredBuffer<FWaterBodyRenderData> WaterBodyRenderData;
+ByteAddressBuffer PackedNodes;
+Buffer<uint> InstanceDataOffsets;
+
+float3 QuadTreePosition;
+float3 ObserverPosition;
+uint QuadTreeResolutionX;
+uint QuadTreeResolutionY;
+uint NumDensities;
+uint NumMaterials;
+uint NumDispatchedThreads;
+uint BucketIndexOffset;
+uint NumLODs;
+float LeafSize;
+float LODScale;
+float CaptureDepthRange;
+uint bWithWaterSelectionSupport;
+uint bLODMorphingEnabled;
+
+[numthreads(64, 1, 1)]
+void MainCS(uint3 DispatchThreadID : SV_DispatchThreadID)
+{
+	const uint NumDraws = PackedNodes.Load(0);
+	// We are limited to 65535 threads, but might end up with more draws than that
+	for (uint DrawIndex = DispatchThreadID.x; DrawIndex < NumDraws; DrawIndex += NumDispatchedThreads)
+	{
+		// Load packed node coord and density index
+		uint4 NodeCoordAndDensity = UnpackNodeCoord(PackedNodes.Load((DrawIndex + 1) << 2));
+		const uint2 NodeCoord = NodeCoordAndDensity.xy;
+		const uint LODLevel = NodeCoordAndDensity.z;
+		const uint DensityIndex = NodeCoordAndDensity.w;
+		const uint DensityIndexClamped = min(DensityIndex, NumDensities - 1);
+
+		// Sample the quadtree to access the water body render data
+		const float4 QuadTreeSample = QuadTreeTexture.Load(int3(NodeCoord, LODLevel));
+		const FWaterQuadTreeNode Node = WaterQuadTreeUnpackNodeRGBA8(QuadTreeSample);
+		const FWaterBodyRenderData WBRenderData = WaterBodyRenderData[Node.WaterBodyRenderDataIndex];
+
+		// Determine the material
+		uint MaterialIndex = WBRenderData.MaterialIndex;
+		uint WaterBodyIndex = WBRenderData.WaterBodyIndex;
+		// Rivers can have transitions to other water bodies
+		if (WBRenderData.WaterBodyType == WATER_BODY_TYPE_RIVER && Node.TransitionWaterBodyRenderDataIndex > 0)
+		{
+			const FWaterBodyRenderData TransitionWBRenderData = WaterBodyRenderData[Node.TransitionWaterBodyRenderDataIndex];
+			if (TransitionWBRenderData.WaterBodyType == WATER_BODY_TYPE_LAKE)
+			{
+				MaterialIndex = WBRenderData.RiverToLakeMaterialIndex;
+				WaterBodyIndex = TransitionWBRenderData.WaterBodyIndex;
+			}
+			else if (TransitionWBRenderData.WaterBodyType == WATER_BODY_TYPE_OCEAN)
+			{
+				MaterialIndex = WBRenderData.RiverToOceanMaterialIndex;
+				WaterBodyIndex = TransitionWBRenderData.WaterBodyIndex;
+			}
+		}
+
+		const uint BucketIndex = MaterialIndex * NumDensities + DensityIndexClamped + BucketIndexOffset;
+
+		const uint BucketInstanceDataOffset = InstanceDataOffsets.Load(BucketIndex);
+		uint LocalInstanceDataOffset;
+		InterlockedAdd(IndirectArgs[BucketIndex].InstanceCount, 1, LocalInstanceDataOffset);
+		const uint InstanceDataIndex = BucketInstanceDataOffset + LocalInstanceDataOffset;
+		
+		const float Scale = (1u << LODLevel) * LeafSize;
+		const float WaterSurfaceBaseHeight = WaterZBoundsTexture.Load(int3(NodeCoord, LODLevel)).z * CaptureDepthRange;
+		const float3 TranslatedWorldPosition = float3((NodeCoord + 0.5f) * Scale, WaterSurfaceBaseHeight) + QuadTreePosition;
+
+		// Add the data to the bucket
+		{
+			float4 Data0;
+			Data0.x = TranslatedWorldPosition.x;
+			Data0.y = TranslatedWorldPosition.y;
+			Data0.z = TranslatedWorldPosition.z;
+			Data0.w = asfloat(WaterBodyIndex);
+
+			InstanceData0[InstanceDataIndex] = Data0;
+		}
+		
+		// Should morph
+		{
+			const float3 ObserverPos = ObserverPosition - QuadTreePosition;
+			const int2 TextureLoadCoord = (int2)clamp(ObserverPos.xy / LeafSize, 0.0f, float2(QuadTreeResolutionX, QuadTreeResolutionY) - 0.5f);
+			const float WaterHeightForLOD = WaterZBoundsTexture.Load(int3(TextureLoadCoord, 0)).z * CaptureDepthRange;
+			const FWaterLODParams WaterLODParams = GetWaterLODParams(ObserverPos.z, WaterHeightForLOD, LODScale, NumLODs - 1);
+
+			const uint LogicalLODLevel = LODLevel + DensityIndex;
+
+			// Lowest LOD isn't always 0, this increases with the height distance 
+			const bool bIsLowestLOD = (LogicalLODLevel == WaterLODParams.LowestLOD);
+
+			// Only allow a tile to morph if it's not the last density level and not the last LOD level, sicne there is no next level to morph to
+			const uint bShouldMorph = (bLODMorphingEnabled && (DensityIndexClamped != NumDensities - 1)) ? 1 : 0;
+			// Tiles can morph twice to be able to morph between 3 LOD levels. Next to last density level can only morph once
+			const uint bCanMorphTwice = (DensityIndexClamped < NumDensities - 2) ? 1 : 0;
+
+			// Pack some of the data to save space. LOD level in the lower 8 bits and then bShouldMorph in the 9th bit and bCanMorphTwice in the 10th bit
+			const uint BitPackedChannel = ((uint)LogicalLODLevel & 0xFFu) | (bShouldMorph << 8u) | (bCanMorphTwice << 9u);
+
+			float4 Data1;
+			Data1.x = asfloat(BitPackedChannel);
+			Data1.y = bIsLowestLOD ? WaterLODParams.HeightLODFactor : 0.0f;
+			Data1.z = Scale;
+			Data1.w = Scale;
+
+			InstanceData1[InstanceDataIndex] = Data1;
+		}
+
+		// Instance Hit Proxy ID
+		if (bWithWaterSelectionSupport)
+		{
+			float4 Data2 = UnpackRGBA8(WBRenderData.HitProxyColorAndIsSelected).bgra;
+			Data2.w = Data2.w != 0.0f ? 1.0f : 0.0f;
+
+			InstanceData2[InstanceDataIndex] = Data2;
+		}
+	}
+}
+
+#endif // GENERATE_INSTANCE_DATA
\ No newline at end of file
diff --git a/Engine/Plugins/Experimental/Water/Shaders/Private/WaterQuadTreeMerge.usf b/Engine/Plugins/Experimental/Water/Shaders/Private/WaterQuadTreeMerge.usf
new file mode 100644
--- /dev/null
+++ b/Engine/Plugins/Experimental/Water/Shaders/Private/WaterQuadTreeMerge.usf
@@ -0,0 +1,122 @@
+// Copyright Epic Games, Inc. All Rights Reserved.
+
+#include "/Engine/Private/Common.ush"
+#include "/Engine/Private/PackUnpack.ush"
+#include "WaterQuadTreeCommon.ush"
+
+#ifndef NUM_MSAA_SAMPLES
+#define NUM_MSAA_SAMPLES 1
+#endif
+
+#if NUM_MSAA_SAMPLES > 1
+Texture2DMS<float4, NUM_MSAA_SAMPLES> WaterBodyRasterTextureMS;
+Texture2DMS<float4, NUM_MSAA_SAMPLES> ZBoundsRasterTextureMS;
+#else
+Texture2D<float4> WaterBodyRasterTexture;
+Texture2D<float4> ZBoundsRasterTexture;
+#endif
+StructuredBuffer<FWaterBodyRenderData> WaterBodyRenderData;
+int SuperSamplingFactor;
+float RcpCaptureDepthRange;
+float ZBoundsPadding; // We base our tile Z bounds on rasterization, however with sloped rivers we will likely miss part of the Z extent. Add some padding to account for that.
+
+struct FWaterBodyInfo
+{
+	int Index;
+	int Priority;
+};
+
+FWaterBodyInfo UnpackWaterBodyInfo(float Packed)
+{
+	const uint PackedUI = saturate(Packed) * 1024.0f;
+	
+	FWaterBodyInfo Result;
+	Result.Index = PackedUI & 0x7Fu;
+	Result.Priority = PackedUI >> 7u;
+
+	return Result;
+}
+
+bool IsValidWaterBody(FWaterBodyInfo WaterBodyInfo)
+{
+	return WaterBodyInfo.Index > 0;
+}
+
+void Main(
+	in float4 Position : SV_Position, 
+	out float4 OutPackedNode : SV_Target0, 
+	out float4 OutWaterZBounds : SV_Target1)
+{
+	// Downsample super- and multisampled input texture. The data is set up to correctly resolve with MAX blending.
+	float NonRiverWaterBodyPacked = 0.0f;
+	float RiverWaterBodyPacked = 0.0f;
+	float MinZ = 0.0f;
+	float MaxZ = 0.0f;
+	for (int y = 0; y < SuperSamplingFactor; ++y)
+	{
+		for (int x = 0; x < SuperSamplingFactor; ++x)
+		{
+#if NUM_MSAA_SAMPLES > 1
+			for (int SampleIdx = 0; SampleIdx < NUM_MSAA_SAMPLES; ++SampleIdx)
+#endif
+			{
+				const int2 Coord = int2(Position.xy) * SuperSamplingFactor + int2(x, y);
+
+#if NUM_MSAA_SAMPLES > 1
+				const float2 WaterBodySample = WaterBodyRasterTextureMS.Load(Coord, SampleIdx).xy;
+				const float2 ZBoundsSample = ZBoundsRasterTextureMS.Load(Coord, SampleIdx).xy;
+#else
+				const float2 WaterBodySample = WaterBodyRasterTexture.Load(int3(Coord, 0)).xy;
+				const float2 ZBoundsSample = ZBoundsRasterTexture.Load(int3(Coord, 0)).xy;
+#endif
+				NonRiverWaterBodyPacked = max(NonRiverWaterBodyPacked, WaterBodySample.x);
+				RiverWaterBodyPacked = max(RiverWaterBodyPacked, WaterBodySample.y);
+				
+				MinZ = max(MinZ, ZBoundsSample.x); // MinZ is actually 1 - MinZ
+				MaxZ = max(MaxZ, ZBoundsSample.y);
+			}
+		}
+	}
+
+	const FWaterBodyInfo NonRiverWaterBodyInfo = UnpackWaterBodyInfo(NonRiverWaterBodyPacked);
+	const FWaterBodyInfo RiverWaterBodyInfo = UnpackWaterBodyInfo(RiverWaterBodyPacked);
+
+	const bool bIsNonRiverWaterBodyValid = IsValidWaterBody(NonRiverWaterBodyInfo);
+	const bool bIsRiverWaterBodyValid = IsValidWaterBody(RiverWaterBodyInfo);
+	const bool bIsAnyValid = bIsRiverWaterBodyValid || bIsNonRiverWaterBodyValid;
+
+	FWaterQuadTreeNode ResultNode;
+	ResultNode.WaterBodyRenderDataIndex = 0;
+	ResultNode.TransitionWaterBodyRenderDataIndex = 0;
+	ResultNode.bHasCompleteSubtree = bIsAnyValid;
+	ResultNode.bIsSubtreeSameWaterBody = bIsAnyValid;
+
+	// Rivers always have priority and can have transitions to non-river water bodies
+	if (bIsRiverWaterBodyValid)
+	{
+		ResultNode.WaterBodyRenderDataIndex = RiverWaterBodyInfo.Index;
+		ResultNode.TransitionWaterBodyRenderDataIndex = bIsNonRiverWaterBodyValid ? NonRiverWaterBodyInfo.Index : 0;
+	}
+	else if (bIsNonRiverWaterBodyValid)
+	{
+		ResultNode.WaterBodyRenderDataIndex = NonRiverWaterBodyInfo.Index;
+		ResultNode.TransitionWaterBodyRenderDataIndex = 0;
+	}
+
+	OutPackedNode = WaterQuadTreePackNodeRGBA8(ResultNode);
+
+	const float2 MinMaxZ = bIsAnyValid ? saturate(float2((1.0f - MinZ) - ZBoundsPadding, MaxZ + ZBoundsPadding)) : float2(0.0f, 0.0f);
+	float SurfaceBaseHeight = 0.0f;
+	// In case of rivers transitioning to other water bodies, we want to use the SurfaceBaseHeight of that other water body
+	if (bIsNonRiverWaterBodyValid)
+	{
+		const FWaterBodyRenderData WBRenderData = WaterBodyRenderData[NonRiverWaterBodyInfo.Index];
+		SurfaceBaseHeight = WBRenderData.SurfaceBaseHeight * RcpCaptureDepthRange;
+	}
+	else if (bIsRiverWaterBodyValid)
+	{
+		SurfaceBaseHeight = MinMaxZ.y;
+	}
+
+	OutWaterZBounds = float4(MinMaxZ, SurfaceBaseHeight, 1.0f);
+}
\ No newline at end of file
diff --git a/Engine/Plugins/Experimental/Water/Shaders/Private/WaterQuadTreePixelShader.usf b/Engine/Plugins/Experimental/Water/Shaders/Private/WaterQuadTreePixelShader.usf
new file mode 100644
--- /dev/null
+++ b/Engine/Plugins/Experimental/Water/Shaders/Private/WaterQuadTreePixelShader.usf
@@ -0,0 +1,41 @@
+// Copyright Epic Games, Inc. All Rights Reserved.
+
+#include "/Engine/Private/Common.ush"
+
+uint WaterBodyRenderDataIndex;
+uint Priority;
+float WaterBodyMinZ;
+float WaterBodyMaxZ;
+float MaxWaveHeight;
+uint bIsRiver;
+
+void Main(
+	in float4 InPosition : SV_Position, 
+	out float4 OutWaterBody : SV_Target0,
+	out float4 OutZBounds : SV_Target1)
+{
+	uint PackedPriorityAndWaterBodyRenderDataIndex = 0;
+	PackedPriorityAndWaterBodyRenderDataIndex |= WaterBodyRenderDataIndex & 0x7Fu; // 7bits
+	PackedPriorityAndWaterBodyRenderDataIndex |= min(Priority, 7u) << 7u; // 3bits
+
+	// Pack into 10bit unorm
+	const float PackedF = (PackedPriorityAndWaterBodyRenderDataIndex + 0.5f) / 1024.0f;
+	
+	// RT0 R: Packed non-river water body priority and render data index
+	// RT0 G: Packed river water body priority and render data index
+	// RT1 R: Inverse normalized minimum bounds Z value
+	// RT1 G: Normalized maximum bounds Z value
+	if (bIsRiver != 0)
+	{
+		const float RiverDepth = InPosition.z;
+		const float RiverMinZ = saturate(RiverDepth - MaxWaveHeight);
+		const float RiverMaxZ = saturate(RiverDepth + MaxWaveHeight);
+		OutWaterBody = float4(0.0f, PackedF, 0.0f, 1.0f);
+		OutZBounds = float4(1.0f - RiverMinZ, RiverMaxZ, 0.0f, 1.0f);
+	}
+	else
+	{
+		OutWaterBody = float4(PackedF, 0.0f, 0.0f, 1.0f);
+		OutZBounds = float4(1.0f - WaterBodyMinZ, WaterBodyMaxZ, 0.0f, 1.0f);
+	}
+}
diff --git a/Engine/Plugins/Experimental/Water/Shaders/Private/WaterQuadTreeVertexShader.usf b/Engine/Plugins/Experimental/Water/Shaders/Private/WaterQuadTreeVertexShader.usf
new file mode 100644
--- /dev/null
+++ b/Engine/Plugins/Experimental/Water/Shaders/Private/WaterQuadTreeVertexShader.usf
@@ -0,0 +1,13 @@
+// Copyright Epic Games, Inc. All Rights Reserved.
+
+#include "/Engine/Private/Common.ush"
+
+float4x4 Transform;
+
+void Main(
+	in float3 InPosition : ATTRIBUTE0,
+	out float4 OutPosition : SV_Position
+)
+{
+	OutPosition = mul(float4(InPosition, 1.0f), Transform);
+}
diff --git a/Engine/Plugins/Experimental/Water/Source/Runtime/Private/WaterMeshComponent.cpp b/Engine/Plugins/Experimental/Water/Source/Runtime/Private/WaterMeshComponent.cpp
--- a/Engine/Plugins/Experimental/Water/Source/Runtime/Private/WaterMeshComponent.cpp
+++ b/Engine/Plugins/Experimental/Water/Source/Runtime/Private/WaterMeshComponent.cpp
@@ -6,12 +6,14 @@
 #include "Materials/MaterialInstanceDynamic.h"
 #include "Materials/Material.h"
 #include "PhysicsEngine/BodySetup.h"
+#include "Engine/StaticMesh.h"
 #include "WaterBodyComponent.h"
 #include "WaterMeshSceneProxy.h"
 #include "WaterModule.h"
 #include "WaterSplineComponent.h"
 #include "WaterSubsystem.h"
 #include "WaterUtils.h"
+#include "WaterBodyInfoMeshComponent.h"
 
 #include UE_INLINE_GENERATED_CPP_BY_NAME(WaterMeshComponent)
 
@@ -31,6 +33,13 @@ static TAutoConsoleVariable<float> CVarWaterMeshLODScaleBias(
 	TEXT("This value is added to the LOD Scale of each Mesh Component. Negative values will lower the overall density/resolution or the vertex grid and make the LODs smaller, higher values will increase the density/resolution and make the LODs larger. Smallest value is -0.5. That will make the inner LOD as tight and optimized as possible"),
 	ECVF_Scalability);
 
+static TAutoConsoleVariable<int32> CVarWaterMeshGPUQuadTree(
+	TEXT("r.Water.WaterMesh.GPUQuadTree"),
+	0,
+	TEXT("Builds the water quadtree on the GPU and does indirect draws of water tiles, driven by the GPU."),
+	ECVF_RenderThreadSafe
+);
+
 /** Debug CVars */ 
 static TAutoConsoleVariable<int32> CVarWaterMeshShowTileGenerationGeometry(
 	TEXT("r.Water.WaterMesh.ShowTileGenerationGeometry"),
@@ -73,7 +82,7 @@ void UWaterMeshComponent::PostLoad()
 
 void UWaterMeshComponent::CollectPSOPrecacheData(const FPSOPrecacheParams& BasePrecachePSOParams, FComponentPSOPrecacheParamsList& OutParams)
 {
-	const FVertexFactoryType* WaterVertexFactoryType = GetWaterVertexFactoryType(/*bWithWaterSelectionSupport = */ false);
+	const FVertexFactoryType* WaterVertexFactoryType = GetWaterVertexFactoryType(/*bWithWaterSelectionSupport = */ false, /*bIndirectDraws = */ false);
 	if (FarDistanceMaterial)
 	{
 		FComponentPSOPrecacheParams& ComponentParams = OutParams[OutParams.AddDefaulted()];
@@ -247,8 +256,10 @@ void UWaterMeshComponent::RebuildWaterMesh(float InTileSize, const FIntPoint& In
 		return;
 	}
 
+	const bool bIsGPUQuadTree = CVarWaterMeshGPUQuadTree.GetValueOnGameThread() != 0;
+
 	// This resets the tree to an initial state, ready for node insertion
-	WaterQuadTree.InitTree(WaterWorldBox, InTileSize, InExtentInTiles);
+	WaterQuadTree.InitTree(WaterWorldBox, InTileSize, InExtentInTiles, bIsGPUQuadTree);
 
 	UsedMaterials.Empty();
 
@@ -266,25 +277,61 @@ void UWaterMeshComponent::RebuildWaterMesh(float InTileSize, const FIntPoint& In
 	// Go through all water body actors to figure out bounds and water tiles
 	AWaterZone* OwningZone = GetOwner<AWaterZone>();
 	check(OwningZone);
-	OwningZone->ForEachWaterBodyComponent([this, WaterWorldBox, bIsFlooded, GlobalOceanHeight, OceanFlood, &FarMeshHeight, &bHasOcean](UWaterBodyComponent* WaterBodyComponent)
-	{
-		check(WaterBodyComponent);
-		AActor* Actor = WaterBodyComponent->GetOwner();
-		check(Actor);
 
-		const FBox WaterBodyBounds = WaterBodyComponent->Bounds.GetBox();
+	// Min and max user defined priority range. (Input also clamped on OverlapMaterialPriority in AWaterBody)
+	constexpr int32 MinWaterBodyPriority = -8192;
+	constexpr int32 MaxWaterBodyPriority = 8191;
+	constexpr int32 GPUQuadTreeMaxNumPriorities = 8;
 
-		// Don't process water bodies that has their spline outside of this water mesh
-		if (!WaterBodyBounds.IntersectXY(FBox(FVector(WaterWorldBox.Min, 0.0f), FVector(WaterWorldBox.Max, 0.0f))))
+	// The GPU quadtree only supports 8 different priority values, so we need to remap priorities into that space.
+	// Fortunately, rivers and non-river water bodies are rendered into their own "priority space", so we don't need to worry about
+	// moving river priorities into their own range.
+	TArray<int16> SortedPriorities;
+	if (bIsGPUQuadTree)
+	{
+		OwningZone->ForEachWaterBodyComponent([this, WaterWorldBox, &SortedPriorities](UWaterBodyComponent* WaterBodyComponent)
 		{
+			check(WaterBodyComponent);
+			AActor* Actor = WaterBodyComponent->GetOwner();
+			check(Actor);
+
+			// Skip invisible water bodies
+			if (!WaterBodyComponent->ShouldRender() || !WaterBodyComponent->ShouldGenerateWaterMeshTile())
+			{
+				return true;
+			}
+
+			// Don't process water bodies which have their spline outside of this water mesh
+			const FBox WaterBodyBounds = WaterBodyComponent->Bounds.GetBox();
+			if (!WaterBodyBounds.IntersectXY(FBox(FVector(WaterWorldBox.Min, 0.0f), FVector(WaterWorldBox.Max, 0.0f))))
+			{
+				return true;
+			}
+
+			const int16 Priority = static_cast<int16>(FMath::Clamp(WaterBodyComponent->GetOverlapMaterialPriority(), MinWaterBodyPriority, MaxWaterBodyPriority));
+			SortedPriorities.AddUnique(Priority);
+
 			return true;
+		});
+
+		SortedPriorities.Sort();
+
+		if (SortedPriorities.Num() > GPUQuadTreeMaxNumPriorities)
+		{
+			UE_LOG(LogWater, Warning, TEXT("WaterZone has more unique water body priorities (%i) than can be supported with GPU driven water quadtree rendering (%i)!"), SortedPriorities.Num(), GPUQuadTreeMaxNumPriorities);
 		}
+	}
+	
 
-		FWaterBodyRenderData RenderData;
+	bool bAnyWaterMeshesNotReady = false;
 
-		const EWaterBodyType WaterBodyType = WaterBodyComponent->GetWaterBodyType();
+	OwningZone->ForEachWaterBodyComponent([this, bIsGPUQuadTree, WaterWorldBox, bIsFlooded, GlobalOceanHeight, OceanFlood, &FarMeshHeight, &bHasOcean, &SortedPriorities, &bAnyWaterMeshesNotReady](UWaterBodyComponent* WaterBodyComponent)
+	{
+		check(WaterBodyComponent);
+		AActor* Actor = WaterBodyComponent->GetOwner();
+		check(Actor);
 
-		// Skip invisible water bodies : 
+		// Skip invisible water bodies
 		if (!WaterBodyComponent->ShouldRender())
 		{
 			return true;
@@ -296,6 +343,17 @@ void UWaterMeshComponent::RebuildWaterMesh(float InTileSize, const FIntPoint& In
 			return true;
 		}
 
+		// Don't process water bodies that has their spline outside of this water mesh
+		const FBox WaterBodyBounds = WaterBodyComponent->Bounds.GetBox();
+		if (!WaterBodyBounds.IntersectXY(FBox(FVector(WaterWorldBox.Min, 0.0f), FVector(WaterWorldBox.Max, 0.0f))))
+		{
+			return true;
+		}
+
+		FWaterBodyRenderData RenderData;
+
+		const EWaterBodyType WaterBodyType = WaterBodyComponent->GetWaterBodyType();
+
 		if (WaterBodyType != EWaterBodyType::Ocean)
 		{
 			if (bIsFlooded)
@@ -329,12 +387,12 @@ void UWaterMeshComponent::RebuildWaterMesh(float InTileSize, const FIntPoint& In
 			UsedMaterials.Add(RenderData.Material);
 		}
 
-		// Min and max user defined priority range. (Input also clamped on OverlapMaterialPriority in AWaterBody)
-		const int32 MinPriority = -8192;
-		const int32 MaxPriority = 8191;
-		RenderData.Priority = static_cast<int16>(FMath::Clamp(WaterBodyComponent->GetOverlapMaterialPriority(), MinPriority, MaxPriority));
+		RenderData.Priority = static_cast<int16>(FMath::Clamp(WaterBodyComponent->GetOverlapMaterialPriority(), MinWaterBodyPriority, MaxWaterBodyPriority));
 		RenderData.WaterBodyIndex = static_cast<int16>(WaterBodyComponent->GetWaterBodyIndex());
 		RenderData.SurfaceBaseHeight = WaterBodyComponent->GetComponentLocation().Z;
+		RenderData.MaxWaveHeight = WaterBodyComponent->GetMaxWaveHeight();
+		RenderData.BoundsMinZ = WaterBodyComponent->Bounds.GetBox().Min.Z;
+		RenderData.BoundsMaxZ = WaterBodyComponent->Bounds.GetBox().Max.Z;
 		RenderData.WaterBodyType = static_cast<int8>(WaterBodyType);
 #if WITH_WATER_SELECTION_SUPPORT
 		RenderData.HitProxy = new HActor(/*InActor = */Actor, /*InPrimComponent = */nullptr);
@@ -372,207 +430,248 @@ void UWaterMeshComponent::RebuildWaterMesh(float InTileSize, const FIntPoint& In
 		if (RenderData.RiverToLakeMaterial || RenderData.RiverToOceanMaterial)
 		{
 			// Move rivers up to it's own priority space, so that they always have precedence if they have transitions and that they only compare agains other rivers with transitions
-			RenderData.Priority += (MaxPriority-MinPriority)+1;
+			RenderData.Priority += (MaxWaterBodyPriority-MinWaterBodyPriority)+1;
 		}
 
 		uint32 WaterBodyRenderDataIndex = WaterQuadTree.AddWaterBodyRenderData(RenderData);
 
-		switch (WaterBodyType)
+		if (bIsGPUQuadTree)
 		{
-		case EWaterBodyType::River:
-		{
-			TRACE_CPUPROFILER_EVENT_SCOPE(River);
+			// On the GPU path, we only submit FWaterBodyQuadTreeRasterInfo for the water quadtree to be rasterized on the GPU. In particular, we need the static mesh, a transform and a priority/WaterBodyRenderData index.
+
+			// TODO: Add support for flooded ocean water heights. This is a legacy feature and currently not used in any project (?), so it's probably not urgent.
+			// In order to implement this, we would need to render a fullscreen quad/triangle for the ocean instead of using the spline derived mesh.
 
-			TArray<FBox, TInlineAllocator<16>> Boxes;
-			TArray<UPrimitiveComponent*> CollisionComponents = WaterBodyComponent->GetCollisionComponents();
-			for (UPrimitiveComponent* Comp : CollisionComponents)
+			UWaterBodyInfoMeshComponent* WaterBodyInfoMeshComponent = WaterBodyComponent->GetWaterInfoMeshComponent();
+			if (ensure(WaterBodyInfoMeshComponent))
 			{
-				if (UBodySetup* BodySetup = (Comp != nullptr) ? Comp->GetBodySetup() : nullptr)
+				UStaticMesh* StaticMesh = WaterBodyInfoMeshComponent->GetStaticMesh();
+				if (ensure(StaticMesh))
 				{
-					// Go through all sub shapes on the bodysetup to get a tight fit along water body
-					for (const FKConvexElem& ConvElem : BodySetup->AggGeom.ConvexElems)
+					bAnyWaterMeshesNotReady |= StaticMesh->IsCompiling();
+					FStaticMeshRenderData* StaticMeshRenderData = StaticMesh->GetRenderData();
+					if (StaticMeshRenderData)
 					{
-						TRACE_CPUPROFILER_EVENT_SCOPE(Add);
+						const int16 ClampedPriority = static_cast<int16>(FMath::Clamp(WaterBodyComponent->GetOverlapMaterialPriority(), MinWaterBodyPriority, MaxWaterBodyPriority));
 
-						FBox SubBox = ConvElem.ElemBox.TransformBy(Comp->GetComponentTransform().ToMatrixWithScale());
-						SubBox.Max.Z += WaterBodyComponent->GetMaxWaveHeight();
+						FWaterBodyQuadTreeRasterInfo RasterInfo;
+						RasterInfo.LocalToWorld = WaterBodyComponent->GetComponentTransform();
+						RasterInfo.RenderData = StaticMeshRenderData;
+						RasterInfo.WaterBodyRenderDataIndex = WaterBodyRenderDataIndex;
+						RasterInfo.Priority = FMath::Clamp(SortedPriorities.IndexOfByKey(ClampedPriority), 0, GPUQuadTreeMaxNumPriorities - 1);
+						RasterInfo.bIsRiver = WaterBodyType == EWaterBodyType::River;
 
-						Boxes.Add(SubBox);
+						WaterQuadTree.AddWaterBodyRasterInfo(RasterInfo);
 					}
 				}
-				else
-				{
-					// fallback on global AABB: 
-					FVector Center;
-					FVector Extent;
-					Actor->GetActorBounds(false, Center, Extent);
-					FBox Box(FBox::BuildAABB(Center, Extent));
-					Box.Max.Z += WaterBodyComponent->GetMaxWaveHeight();
-					Boxes.Add(Box);
-				}
 			}
 
-			for (const FBox& Box : Boxes)
+			if (WaterBodyType == EWaterBodyType::Ocean)
 			{
-				WaterQuadTree.AddWaterTilesInsideBounds(Box, WaterBodyRenderDataIndex);
-
-	#if !(UE_BUILD_SHIPPING || UE_BUILD_TEST)
-				if (!!CVarWaterMeshShowTileGenerationGeometry.GetValueOnGameThread())
-				{
-					DrawDebugBox(GetWorld(), Box.GetCenter(), Box.GetExtent(), FColor::Red);
-				}
-	#endif
+				// Place far mesh height just below the ocean level
+				FarMeshHeight = RenderData.SurfaceBaseHeight - RenderData.MaxWaveHeight;
+				bHasOcean = true;
 			}
-			break;
 		}
-		case EWaterBodyType::Lake:
+		else
 		{
-			TRACE_CPUPROFILER_EVENT_SCOPE(Lake);
+			switch (WaterBodyType)
+			{
+			case EWaterBodyType::River:
+			{
+				TRACE_CPUPROFILER_EVENT_SCOPE(River);
 
-			const UWaterSplineComponent* SplineComp = WaterBodyComponent->GetWaterSpline();
-			const int32 NumOriginalSplinePoints = SplineComp->GetNumberOfSplinePoints();
-			float ConstantZ = SplineComp->GetLocationAtDistanceAlongSpline(0.0f, ESplineCoordinateSpace::World).Z;
+				TArray<FBox, TInlineAllocator<16>> Boxes;
+				TArray<UPrimitiveComponent*> CollisionComponents = WaterBodyComponent->GetCollisionComponents();
+				for (UPrimitiveComponent* Comp : CollisionComponents)
+				{
+					if (UBodySetup* BodySetup = (Comp != nullptr) ? Comp->GetBodySetup() : nullptr)
+					{
+						// Go through all sub shapes on the bodysetup to get a tight fit along water body
+						for (const FKConvexElem& ConvElem : BodySetup->AggGeom.ConvexElems)
+						{
+							TRACE_CPUPROFILER_EVENT_SCOPE(Add);
 
-			FBox LakeBounds = Actor->GetComponentsBoundingBox(/* bNonColliding = */true);
-			LakeBounds.Max.Z += WaterBodyComponent->GetMaxWaveHeight();
+							FBox SubBox = ConvElem.ElemBox.TransformBy(Comp->GetComponentTransform().ToMatrixWithScale());
+							SubBox.Max.Z += WaterBodyComponent->GetMaxWaveHeight();
 
-			// Skip lakes with less than 3 spline points
-			if (NumOriginalSplinePoints < 3)
-			{
+							Boxes.Add(SubBox);
+						}
+					}
+					else
+					{
+						// fallback on global AABB: 
+						FVector Center;
+						FVector Extent;
+						Actor->GetActorBounds(false, Center, Extent);
+						FBox Box(FBox::BuildAABB(Center, Extent));
+						Box.Max.Z += WaterBodyComponent->GetMaxWaveHeight();
+						Boxes.Add(Box);
+					}
+				}
+
+				for (const FBox& Box : Boxes)
+				{
+					WaterQuadTree.AddWaterTilesInsideBounds(Box, WaterBodyRenderDataIndex);
+
+#if !(UE_BUILD_SHIPPING || UE_BUILD_TEST)
+					if (!!CVarWaterMeshShowTileGenerationGeometry.GetValueOnGameThread())
+					{
+						DrawDebugBox(GetWorld(), Box.GetCenter(), Box.GetExtent(), FColor::Red);
+					}
+#endif
+				}
 				break;
 			}
+			case EWaterBodyType::Lake:
+			{
+				TRACE_CPUPROFILER_EVENT_SCOPE(Lake);
 
-			TArray<TArray<FVector2D>> PolygonBatches;
-			// Reuse the convex hulls generated for the physics shape because the work has already been done, but we can fallback to a simple spline evaluation method in case there's no physics:
-			bool bUseFallbackMethod = true;
+				const UWaterSplineComponent* SplineComp = WaterBodyComponent->GetWaterSpline();
+				const int32 NumOriginalSplinePoints = SplineComp->GetNumberOfSplinePoints();
+				float ConstantZ = SplineComp->GetLocationAtDistanceAlongSpline(0.0f, ESplineCoordinateSpace::World).Z;
 
-			TArray<UPrimitiveComponent*> CollisionComponents = WaterBodyComponent->GetCollisionComponents();
-			if (CollisionComponents.Num() > 0)
-			{
-				UPrimitiveComponent* Comp = CollisionComponents[0];
-				if (UBodySetup* BodySetup = (Comp != nullptr) ? Comp->GetBodySetup() : nullptr)
+				FBox LakeBounds = Actor->GetComponentsBoundingBox(/* bNonColliding = */true);
+				LakeBounds.Max.Z += WaterBodyComponent->GetMaxWaveHeight();
+
+				// Skip lakes with less than 3 spline points
+				if (NumOriginalSplinePoints < 3)
 				{
-					FTransform CompTransform = Comp->GetComponentTransform();
-					// Go through all sub shapes on the bodysetup to get a tight fit along water body
-					for (const FKConvexElem& ConvElem : BodySetup->AggGeom.ConvexElems)
+					break;
+				}
+
+				TArray<TArray<FVector2D>> PolygonBatches;
+				// Reuse the convex hulls generated for the physics shape because the work has already been done, but we can fallback to a simple spline evaluation method in case there's no physics:
+				bool bUseFallbackMethod = true;
+
+				TArray<UPrimitiveComponent*> CollisionComponents = WaterBodyComponent->GetCollisionComponents();
+				if (CollisionComponents.Num() > 0)
+				{
+					UPrimitiveComponent* Comp = CollisionComponents[0];
+					if (UBodySetup* BodySetup = (Comp != nullptr) ? Comp->GetBodySetup() : nullptr)
 					{
-						TRACE_CPUPROFILER_EVENT_SCOPE(Add);
-
-						// Vertex data contains the bottom vertices first, then the top ones :
-						int32 TotalNumVertices = ConvElem.VertexData.Num();
-						check(TotalNumVertices % 2 == 0);
-						int32 NumVertices = TotalNumVertices / 2;
-						if (NumVertices > 0)
-						{ 
-							bUseFallbackMethod = false;
-
-							// Because the physics shape is made of multiple convex hulls, we cannot simply add their vertices to one big list but have to have 1 batch of polygons per convex hull
-							TArray<FVector2D>& Polygon = PolygonBatches.Emplace_GetRef();
-							Polygon.SetNum(NumVertices);
-							for (int32 i = 0; i < NumVertices; ++i)
+						FTransform CompTransform = Comp->GetComponentTransform();
+						// Go through all sub shapes on the bodysetup to get a tight fit along water body
+						for (const FKConvexElem& ConvElem : BodySetup->AggGeom.ConvexElems)
+						{
+							TRACE_CPUPROFILER_EVENT_SCOPE(Add);
+
+							// Vertex data contains the bottom vertices first, then the top ones :
+							int32 TotalNumVertices = ConvElem.VertexData.Num();
+							check(TotalNumVertices % 2 == 0);
+							int32 NumVertices = TotalNumVertices / 2;
+							if (NumVertices > 0)
 							{
-								// Gather the top vertices :
-								Polygon[i] = FVector2D(CompTransform.TransformPosition(ConvElem.VertexData[NumVertices + i]));
+								bUseFallbackMethod = false;
+
+								// Because the physics shape is made of multiple convex hulls, we cannot simply add their vertices to one big list but have to have 1 batch of polygons per convex hull
+								TArray<FVector2D>& Polygon = PolygonBatches.Emplace_GetRef();
+								Polygon.SetNum(NumVertices);
+								for (int32 i = 0; i < NumVertices; ++i)
+								{
+									// Gather the top vertices :
+									Polygon[i] = FVector2D(CompTransform.TransformPosition(ConvElem.VertexData[NumVertices + i]));
+								}
 							}
 						}
 					}
 				}
-			}
 
-			if (bUseFallbackMethod)
-			{
-				TArray<FVector> PolyLineVertices;
-				SplineComp->ConvertSplineToPolyLine(ESplineCoordinateSpace::World, FMath::Square(CVarWaterSplineResampleMaxDistance.GetValueOnGameThread()), PolyLineVertices);
+				if (bUseFallbackMethod)
+				{
+					TArray<FVector> PolyLineVertices;
+					SplineComp->ConvertSplineToPolyLine(ESplineCoordinateSpace::World, FMath::Square(CVarWaterSplineResampleMaxDistance.GetValueOnGameThread()), PolyLineVertices);
 
-				TArray<FVector2D>& Polygon = PolygonBatches.Emplace_GetRef();
-				Polygon.Reserve(PolyLineVertices.Num());
-				Algo::Transform(PolyLineVertices, Polygon, [](const FVector& Vertex) { return FVector2D(Vertex); });
-			}
+					TArray<FVector2D>& Polygon = PolygonBatches.Emplace_GetRef();
+					Polygon.Reserve(PolyLineVertices.Num());
+					Algo::Transform(PolyLineVertices, Polygon, [](const FVector& Vertex) { return FVector2D(Vertex); });
+				}
 
-			for (const TArray<FVector2D>& Polygon : PolygonBatches)
-			{
-				WaterQuadTree.AddLake(Polygon, LakeBounds, WaterBodyRenderDataIndex);
+				for (const TArray<FVector2D>& Polygon : PolygonBatches)
+				{
+					WaterQuadTree.AddLake(Polygon, LakeBounds, WaterBodyRenderDataIndex);
 
 #if !(UE_BUILD_SHIPPING || UE_BUILD_TEST)
-				if (!!CVarWaterMeshShowTileGenerationGeometry.GetValueOnGameThread())
-				{
-					float Z = SplineComp->GetLocationAtDistanceAlongSpline(0.0f, ESplineCoordinateSpace::World).Z;
-					int32 NumVertices = Polygon.Num();
-					for (int32 i = 0; i < NumVertices; i++)
+					if (!!CVarWaterMeshShowTileGenerationGeometry.GetValueOnGameThread())
 					{
-						const FVector2D& Point0 = Polygon[i];
-						const FVector2D& Point1 = Polygon[(i + 1) % NumVertices];
-						DrawDebugLine(GetWorld(), FVector(Point0.X, Point0.Y, Z), FVector(Point1.X, Point1.Y, Z), FColor::Green);
+						float Z = SplineComp->GetLocationAtDistanceAlongSpline(0.0f, ESplineCoordinateSpace::World).Z;
+						int32 NumVertices = Polygon.Num();
+						for (int32 i = 0; i < NumVertices; i++)
+						{
+							const FVector2D& Point0 = Polygon[i];
+							const FVector2D& Point1 = Polygon[(i + 1) % NumVertices];
+							DrawDebugLine(GetWorld(), FVector(Point0.X, Point0.Y, Z), FVector(Point1.X, Point1.Y, Z), FColor::Green);
+						}
 					}
-				}
 #endif // !(UE_BUILD_SHIPPING || UE_BUILD_TEST)
-			}
-
-			break;
-		}
-		case EWaterBodyType::Ocean:
-		{
-			TRACE_CPUPROFILER_EVENT_SCOPE(Ocean);
+				}
 
-			// Add ocean based on the ocean spline when there is no flood. Otherwise add ocean everywhere
-			if (bIsFlooded)
-			{
-				FBox OceanBounds = Actor->GetComponentsBoundingBox();
-				OceanBounds.Max.Z += WaterBodyComponent->GetMaxWaveHeight() + OceanFlood;
-				WaterQuadTree.AddWaterTilesInsideBounds(OceanBounds, WaterBodyRenderDataIndex);
+				break;
 			}
-			else
+			case EWaterBodyType::Ocean:
 			{
-				const UWaterSplineComponent* SplineComp = WaterBodyComponent->GetWaterSpline();
-				const int32 NumOriginalSplinePoints = SplineComp->GetNumberOfSplinePoints();
+				TRACE_CPUPROFILER_EVENT_SCOPE(Ocean);
 
-				// Skip oceans with less than 3 spline points
-				if (NumOriginalSplinePoints < 3)
+				// Add ocean based on the ocean spline when there is no flood. Otherwise add ocean everywhere
+				if (bIsFlooded)
 				{
-					break;
+					FBox OceanBounds = Actor->GetComponentsBoundingBox();
+					OceanBounds.Max.Z += WaterBodyComponent->GetMaxWaveHeight() + OceanFlood;
+					WaterQuadTree.AddWaterTilesInsideBounds(OceanBounds, WaterBodyRenderDataIndex);
 				}
+				else
+				{
+					const UWaterSplineComponent* SplineComp = WaterBodyComponent->GetWaterSpline();
+					const int32 NumOriginalSplinePoints = SplineComp->GetNumberOfSplinePoints();
+
+					// Skip oceans with less than 3 spline points
+					if (NumOriginalSplinePoints < 3)
+					{
+						break;
+					}
 
-				TArray<FVector2D> Polygon;
+					TArray<FVector2D> Polygon;
 
-				TArray<FVector> PolyLineVertices;
-				SplineComp->ConvertSplineToPolyLine(ESplineCoordinateSpace::World, FMath::Square(CVarWaterSplineResampleMaxDistance.GetValueOnGameThread()), PolyLineVertices);
+					TArray<FVector> PolyLineVertices;
+					SplineComp->ConvertSplineToPolyLine(ESplineCoordinateSpace::World, FMath::Square(CVarWaterSplineResampleMaxDistance.GetValueOnGameThread()), PolyLineVertices);
 
-				Polygon.Reserve(PolyLineVertices.Num());
-				Algo::Transform(PolyLineVertices, Polygon, [](const FVector& Vertex) { return FVector2D(Vertex); });
+					Polygon.Reserve(PolyLineVertices.Num());
+					Algo::Transform(PolyLineVertices, Polygon, [](const FVector& Vertex) { return FVector2D(Vertex); });
 
-				FBox OceanBounds = WaterBodyComponent->Bounds.GetBox();
-				OceanBounds.Max.Z += WaterBodyComponent->GetMaxWaveHeight();
+					FBox OceanBounds = WaterBodyComponent->Bounds.GetBox();
+					OceanBounds.Max.Z += WaterBodyComponent->GetMaxWaveHeight();
 
 #if !(UE_BUILD_SHIPPING || UE_BUILD_TEST)
-				if (!!CVarWaterMeshShowTileGenerationGeometry.GetValueOnGameThread())
-				{
-					float Z = SplineComp->GetLocationAtDistanceAlongSpline(0.0f, ESplineCoordinateSpace::World).Z;
-					int32 NumVertices = Polygon.Num();
-					for (int32 i = 0; i < NumVertices; i++)
+					if (!!CVarWaterMeshShowTileGenerationGeometry.GetValueOnGameThread())
 					{
-						const FVector2D& Point0 = Polygon[i];
-						const FVector2D& Point1 = Polygon[(i + 1) % NumVertices];
-						DrawDebugLine(GetWorld(), FVector(Point0.X, Point0.Y, Z), FVector(Point1.X, Point1.Y, Z), FColor::Blue);
-					}
+						float Z = SplineComp->GetLocationAtDistanceAlongSpline(0.0f, ESplineCoordinateSpace::World).Z;
+						int32 NumVertices = Polygon.Num();
+						for (int32 i = 0; i < NumVertices; i++)
+						{
+							const FVector2D& Point0 = Polygon[i];
+							const FVector2D& Point1 = Polygon[(i + 1) % NumVertices];
+							DrawDebugLine(GetWorld(), FVector(Point0.X, Point0.Y, Z), FVector(Point1.X, Point1.Y, Z), FColor::Blue);
+						}
 
-					DrawDebugBox(GetWorld(), OceanBounds.GetCenter(), OceanBounds.GetExtent(), FColor::Blue);
-				}
+						DrawDebugBox(GetWorld(), OceanBounds.GetCenter(), OceanBounds.GetExtent(), FColor::Blue);
+					}
 #endif //!(UE_BUILD_SHIPPING || UE_BUILD_TEST)
 
-				WaterQuadTree.AddOcean(Polygon, OceanBounds, WaterBodyRenderDataIndex);
-			}
+					WaterQuadTree.AddOcean(Polygon, OceanBounds, WaterBodyRenderDataIndex);
+				}
 
-			// Place far mesh height just below the ocean level
-			FarMeshHeight = RenderData.SurfaceBaseHeight - WaterBodyComponent->GetMaxWaveHeight();
-			bHasOcean = true;
+				// Place far mesh height just below the ocean level
+				FarMeshHeight = RenderData.SurfaceBaseHeight - WaterBodyComponent->GetMaxWaveHeight();
+				bHasOcean = true;
 
-			break;
-		}
-		case EWaterBodyType::Transition:
-			// Transitions dont require rendering
-			break;
-		default:
-			ensureMsgf(false, TEXT("This water body type is not implemented and will not produce any water tiles. "));
+				break;
+			}
+			case EWaterBodyType::Transition:
+				// Transitions dont require rendering
+				break;
+			default:
+				ensureMsgf(false, TEXT("This water body type is not implemented and will not produce any water tiles. "));
+			}
 		}
 
 		return true;
@@ -592,6 +691,12 @@ void UWaterMeshComponent::RebuildWaterMesh(float InTileSize, const FIntPoint& In
 	WaterQuadTree.Unlock(true);
 
 	MarkRenderStateDirty();
+
+	// Force another rebuild next frame if water body meshes are still compiling
+	if (bIsGPUQuadTree && bAnyWaterMeshesNotReady)
+	{
+		bNeedsRebuild = true;
+	}
 }
 
 void UWaterMeshComponent::Update()
diff --git a/Engine/Plugins/Experimental/Water/Source/Runtime/Private/WaterMeshSceneProxy.cpp b/Engine/Plugins/Experimental/Water/Source/Runtime/Private/WaterMeshSceneProxy.cpp
--- a/Engine/Plugins/Experimental/Water/Source/Runtime/Private/WaterMeshSceneProxy.cpp
+++ b/Engine/Plugins/Experimental/Water/Source/Runtime/Private/WaterMeshSceneProxy.cpp
@@ -4,16 +4,19 @@
 #include "MaterialShared.h"
 #include "WaterMeshComponent.h"
 #include "WaterZoneActor.h"
+#include "WaterViewExtension.h"
 #include "Materials/Material.h"
 #include "WaterUtils.h"
 #include "PrimitiveViewRelevance.h"
 #include "SceneManagement.h"
-#include "Engine/Engine.h"
+#include "SceneRendering.h"
 #include "Materials/MaterialRenderProxy.h"
 #include "Math/ColorList.h"
 #include "RayTracingInstance.h"
+#include "StaticMeshResources.h"
 #include "SceneInterface.h"
 #include "RenderGraphBuilder.h"
+#include "RenderGraphUtils.h"
 
 DECLARE_STATS_GROUP(TEXT("Water Mesh"), STATGROUP_WaterMesh, STATCAT_Advanced);
 
@@ -30,6 +33,16 @@ static TAutoConsoleVariable<int32> CVarWaterMeshLODMorphEnabled(
 	TEXT("If the smooth LOD morph is enabled. Turning this off may cause slight popping between LOD levels but will skip the calculations in the vertex shader, making it cheaper"),
 	ECVF_Scalability | ECVF_RenderThreadSafe);
 
+static TAutoConsoleVariable<int32> CVarWaterMeshGPUQuadTreeSuperSampling(
+	TEXT("r.Water.WaterMesh.GPUQuadTree.SuperSampling"), 2,
+	TEXT("Rasterizes water meshes into the GPU water quadtree at a higher resolution, reducing missing water tile artifacts near the edges of water bodies. Default: 2, Min: 1, Max : 8"),
+	ECVF_RenderThreadSafe);
+
+static TAutoConsoleVariable<int32> CVarWaterMeshGPUQuadTreeMultiSampling(
+	TEXT("r.Water.WaterMesh.GPUQuadTree.MultiSampling"), 4,
+	TEXT("Rasterizes water meshes into the GPU water quadtree with a multisampled rendertarget, reducing missing water tile artifacts near the edges of water bodies. Default: 4, Min: 1, Max : 8"),
+	ECVF_RenderThreadSafe);
+
 static TAutoConsoleVariable<int32> CVarWaterMeshOcclusionCullingMaxQueries(
 	TEXT("r.Water.WaterMesh.OcclusionCulling.MaxQueries"), 256,
 	TEXT("Maximum number of occlusion queries for the CPU water quadtree nodes. Using fewer queries than nodes will result in coarser culling."),
@@ -110,6 +123,13 @@ static TAutoConsoleVariable<float> CVarWaterMeshOcclusionCullExpandBoundsAmountX
 
 // ----------------------------------------------------------------------------------
 
+template <bool bWithWaterSelectionSupport>
+class TWaterVertexFactoryUserDataWrapper : public FOneFrameResource
+{
+public:
+	TWaterMeshUserData<bWithWaterSelectionSupport> UserData;
+};
+
 FWaterMeshSceneProxy::FWaterMeshSceneProxy(UWaterMeshComponent* Component)
 	: FPrimitiveSceneProxy(Component)
 	, MaterialRelevance(Component->GetWaterMaterialRelevance(GetScene().GetFeatureLevel()))
@@ -126,12 +146,16 @@ FWaterMeshSceneProxy::FWaterMeshSceneProxy(UWaterMeshComponent* Component)
 	}
 
 	int32 NumQuads = (int32)FMath::Pow(2.0f, (float)Component->GetTessellationFactor());
+	NumQuadsLOD0 = NumQuads;
 
 	WaterVertexFactories.Reserve(WaterQuadTree.GetTreeDepth());
+	WaterVertexFactoriesIndirectDraw.Reserve(WaterQuadTree.GetTreeDepth());
 	for (uint8 i = 0; i < WaterQuadTree.GetTreeDepth(); i++)
 	{
-		WaterVertexFactories.Add(new WaterVertexFactoryType(GetScene().GetFeatureLevel(), NumQuads, LODScale));
+		WaterVertexFactories.Add(new FWaterVertexFactoryType(GetScene().GetFeatureLevel(), NumQuads, LODScale));
 		BeginInitResource(WaterVertexFactories.Last());
+		WaterVertexFactoriesIndirectDraw.Add(new FWaterVertexFactoryIndirectDrawType(GetScene().GetFeatureLevel(), NumQuads, LODScale));
+		BeginInitResource(WaterVertexFactoriesIndirectDraw.Last());
 
 		NumQuads /= 2;
 
@@ -143,12 +167,13 @@ FWaterMeshSceneProxy::FWaterMeshSceneProxy(UWaterMeshComponent* Component)
 	}
 
 	WaterVertexFactories.Shrink();
+	WaterVertexFactoriesIndirectDraw.Shrink();
 	DensityCount = WaterVertexFactories.Num();
 
 	const int32 TotalLeafNodes = WaterQuadTree.GetMaxLeafCount();
-	WaterInstanceDataBuffers = new WaterInstanceDataBuffersType(TotalLeafNodes);
+	WaterInstanceDataBuffers = new FWaterInstanceDataBuffersType(TotalLeafNodes);
 
-	WaterMeshUserDataBuffers = new WaterMeshUserDataBuffersType(WaterInstanceDataBuffers);
+	WaterMeshUserDataBuffers = new FWaterMeshUserDataBuffersType(WaterInstanceDataBuffers);
 
 	WaterQuadTree.BuildMaterialIndices();
 
@@ -162,15 +187,48 @@ FWaterMeshSceneProxy::FWaterMeshSceneProxy(UWaterMeshComponent* Component)
 	RayTracingWaterData.SetNum(DensityCount);
 #endif
 
+	// Initialize Z bounds needed for GPU driven rendering
+	if (WaterQuadTree.IsGPUQuadTree())
+	{
+		WaterQuadTreeMinHeight = DBL_MAX;
+		WaterQuadTreeMaxHeight = -DBL_MAX;
+
+		const TArrayView<const FWaterBodyRenderData> WaterBodyRenderData = WaterQuadTree.GetWaterBodyRenderData();
+		for (const FWaterBodyRenderData& RenderData : WaterBodyRenderData)
+		{
+			WaterQuadTreeMinHeight = FMath::Min(WaterQuadTreeMinHeight, RenderData.BoundsMinZ - RenderData.MaxWaveHeight);
+			WaterQuadTreeMaxHeight = FMath::Max(WaterQuadTreeMaxHeight, RenderData.BoundsMaxZ + RenderData.MaxWaveHeight);
+		}
+
+		if (WaterQuadTreeMinHeight == DBL_MAX)
+		{
+			WaterQuadTreeMinHeight = 0.0;
+			WaterQuadTreeMaxHeight = 0.0;
+		}
+		WaterQuadTreeMinHeight -= 1.0;
+		WaterQuadTreeMaxHeight += 1.0;
+	}
+
+	// Always do CPU occlusion queries, even if this is a GPU quadtree. The GPU quadtree still potentially uses the far mesh which is CPU driven.
 	const int32 MaxQueries = CVarWaterMeshOcclusionCullingMaxQueries.GetValueOnGameThread();
 	const bool bIncludeFarMeshOcclusionQueries = CVarWaterMeshOcclusionCullingIncludeFarMesh.GetValueOnGameThread() != 0;
 	OcclusionCullingBounds = WaterQuadTree.ComputeNodeBounds(MaxQueries, CVarWaterMeshOcclusionCullExpandBoundsAmountXY.GetValueOnGameThread(), bIncludeFarMeshOcclusionQueries, &OcclusionResultsFarMeshOffset);
 	EmptyOcclusionCullingBounds.Add(WaterQuadTree.GetBoundsIncludingFarMesh());
+	// If this is a GPU quadtree, the CPU root node will have an invalid bounding box, so derive conservative bounds now
+	if (WaterQuadTree.IsGPUQuadTree() && !OcclusionCullingBounds.IsEmpty())
+	{
+		OcclusionCullingBounds[0] = FBox(FVector(WaterQuadTree.GetTileRegion().Min, WaterQuadTreeMinHeight), FVector(WaterQuadTree.GetTileRegion().Max, WaterQuadTreeMaxHeight));
+	}
 }
 
 FWaterMeshSceneProxy::~FWaterMeshSceneProxy()
 {
-	for (WaterVertexFactoryType* WaterFactory : WaterVertexFactories)
+	for (FWaterVertexFactoryType* WaterFactory : WaterVertexFactories)
+	{
+		WaterFactory->ReleaseResource();
+		delete WaterFactory;
+	}
+	for (FWaterVertexFactoryIndirectDrawType* WaterFactory : WaterVertexFactoriesIndirectDraw)
 	{
 		WaterFactory->ReleaseResource();
 		delete WaterFactory;
@@ -195,6 +253,33 @@ FWaterMeshSceneProxy::~FWaterMeshSceneProxy()
 void FWaterMeshSceneProxy::CreateRenderThreadResources(FRHICommandListBase& RHICmdList)
 {
 	SceneProxyCreatedFrameNumberRenderThread = GFrameNumberRenderThread;
+
+	if (WaterQuadTree.IsGPUQuadTree())
+	{
+		FWaterMeshGPUWork::FCallback Callback;
+		Callback.Proxy = this;
+		Callback.Function = [this](FRDGBuilder& GraphBuilder)
+		{
+			if (bNeedToTraverseGPUQuadTree)
+			{
+				if (!QuadTreeGPU.IsInitialized())
+				{
+					BuildGPUQuadTree(GraphBuilder);
+				}
+				QuadTreeGPU.Traverse(GraphBuilder, WaterQuadTreeGPUTraverseParams);
+				bNeedToTraverseGPUQuadTree = false;
+			}
+		};
+		GWaterMeshGPUWork.Callbacks.Add(MoveTemp(Callback));
+	}
+}
+
+void FWaterMeshSceneProxy::DestroyRenderThreadResources()
+{
+	if (WaterQuadTree.IsGPUQuadTree())
+	{
+		GWaterMeshGPUWork.Callbacks.RemoveAllSwap([this](const FWaterMeshGPUWork::FCallback& Callback) { return Callback.Proxy == this; });
+	}
 }
 
 FWaterMeshSceneProxy::FWaterLODParams FWaterMeshSceneProxy::GetWaterLODParams(const FVector& Position) const
@@ -225,7 +310,7 @@ void FWaterMeshSceneProxy::GetDynamicMeshElements(const TArray<const FSceneView*
 	FRHICommandListBase& RHICmdList = Collector.GetRHICommandList();
 
 	// The water render groups we have to render for this batch : 
-	TArray<EWaterMeshRenderGroupType, TInlineAllocator<WaterVertexFactoryType::NumRenderGroups>> BatchRenderGroups;
+	TArray<EWaterMeshRenderGroupType, TInlineAllocator<FWaterVertexFactoryType::NumRenderGroups>> BatchRenderGroups;
 	// By default, render all water tiles : 
 	BatchRenderGroups.Add(EWaterMeshRenderGroupType::RG_RenderWaterTiles);
 
@@ -259,7 +344,208 @@ void FWaterMeshSceneProxy::GetDynamicMeshElements(const TArray<const FSceneView*
 		Collector.RegisterOneFrameMaterialProxy(WireframeMaterialInstance);
 	}
 
-	const int32 NumBuckets = WaterQuadTree.GetWaterMaterials().Num() * DensityCount;
+	const int32 NumWaterMaterials = WaterQuadTree.GetWaterMaterials().Num();
+	const int32 NumBuckets = NumWaterMaterials * DensityCount;
+
+	// If this is a GPU quadtree, we are using GPU driven rendering to first create a quadtree on the GPU (once) and then traverse it every frame to build a list of indirect draws
+	const bool bIsGPUQuadTree = WaterQuadTree.IsGPUQuadTree();
+	if (bIsGPUQuadTree && NumBuckets > 0)
+	{
+		bool bEncounteredISRView = false;
+		int32 InstanceFactor = 1;
+		int32 NumVisibleViews = 0;
+
+		TArray<FWaterQuadTreeGPU::FPerViewInfo> PerViewInfoArray;
+
+		// Gather per view data for all renderable views
+		for (int32 ViewIndex = 0; ViewIndex < Views.Num(); ViewIndex++)
+		{
+			const FSceneView* View = Views[ViewIndex];
+			if (!bEncounteredISRView && View->IsInstancedStereoPass())
+			{
+				bEncounteredISRView = true;
+				InstanceFactor = View->GetStereoPassInstanceFactor();
+			}
+
+			// skip gathering visible tiles from instanced right eye views
+			if ((VisibilityMap & (1 << ViewIndex)) && (!bEncounteredISRView || View->IsPrimarySceneView()))
+			{
+				const FVector QuadTreePosition = FVector(WaterQuadTree.GetTileRegion().Min, WaterQuadTreeMinHeight);
+				const FVector PreViewTranslation = View->ViewMatrices.GetPreViewTranslation();
+				const FViewInfo* ViewInfo = View->bIsViewInfo ? reinterpret_cast<const FViewInfo*>(View) : nullptr;
+
+				FWaterQuadTreeGPU::FPerViewInfo& PerViewInfo = PerViewInfoArray.AddDefaulted_GetRef();
+				PerViewInfo.QuadTreePositionTranslatedWorldSpace = FVector3f(QuadTreePosition + PreViewTranslation);
+				PerViewInfo.ObserverPositionTranslatedWorldSpace = FVector3f(View->ViewMatrices.GetViewOrigin() + PreViewTranslation);
+				PerViewInfo.TranslatedWorldToClip = FMatrix44f(View->ViewMatrices.GetTranslatedViewProjectionMatrix());
+				PerViewInfo.ViewToClip = FMatrix44f(View->ViewMatrices.GetProjectionMatrix());
+				PerViewInfo.ViewInfo = ViewInfo;
+				
+				++NumVisibleViews;
+			}
+		}
+
+		const int32 NumIndirectDrawCalls = NumBuckets * NumVisibleViews;
+		const int32 NumInstances = WaterQuadTree.GetMaxLeafCount() * NumVisibleViews;
+
+		// Allocate buffers for the indirect draws
+		struct FIndirectDrawResources
+		{
+			TRefCountPtr<FRDGPooledBuffer> IndirectArgs;
+			TRefCountPtr<FRDGPooledBuffer> InstanceDataOffsets;
+			TRefCountPtr<FRDGPooledBuffer> InstanceData0;
+			TRefCountPtr<FRDGPooledBuffer> InstanceData1;
+			TRefCountPtr<FRDGPooledBuffer> InstanceData2;
+		};
+		FIndirectDrawResources& IndirectDrawResources = Collector.AllocateOneFrameResource<FIndirectDrawResources>();
+		IndirectDrawResources.IndirectArgs = AllocatePooledBuffer(FRDGBufferDesc::CreateIndirectDesc<FRHIDrawIndexedIndirectParameters>(FMath::Max(1, NumIndirectDrawCalls)), TEXT("WaterQuadTree.IndirectArgsBuffer"));
+		IndirectDrawResources.InstanceDataOffsets = AllocatePooledBuffer(FRDGBufferDesc::CreateBufferDesc(sizeof(uint32), FMath::Max(1, NumInstances)), TEXT("WaterQuadTree.InstanceDataOffsetsBuffer"));
+		IndirectDrawResources.InstanceData0 = AllocatePooledBuffer(FRDGBufferDesc::CreateBufferDesc(sizeof(FVector4f), FMath::Max(1, NumInstances)), TEXT("WaterQuadTree.InstanceDataBuffer0"));
+		IndirectDrawResources.InstanceData1 = AllocatePooledBuffer(FRDGBufferDesc::CreateBufferDesc(sizeof(FVector4f), FMath::Max(1, NumInstances)), TEXT("WaterQuadTree.InstanceDataBuffer1"));
+		if (WITH_WATER_SELECTION_SUPPORT != 0)
+		{
+			IndirectDrawResources.InstanceData2 = AllocatePooledBuffer(FRDGBufferDesc::CreateBufferDesc(sizeof(FVector4f), FMath::Max(1, NumInstances)), TEXT("WaterQuadTree.InstanceDataBuffer2"));
+		}
+		
+
+		// Traverse quadtree. We actually only set up the parameters here; the actual traversal is done later in the frame in a callback invoked from FWaterViewExtension::PreRenderBasePass_RenderThread.
+		// That callback accesses these params, so we need to ensure that they are still valid at that point.
+		{
+			WaterQuadTreeGPUTraverseParams = {};
+			WaterQuadTreeGPUTraverseParams.OutIndirectArgsBuffer = IndirectDrawResources.IndirectArgs;
+			WaterQuadTreeGPUTraverseParams.OutInstanceDataOffsetsBuffer = IndirectDrawResources.InstanceDataOffsets;
+			WaterQuadTreeGPUTraverseParams.OutInstanceData0Buffer = IndirectDrawResources.InstanceData0;
+			WaterQuadTreeGPUTraverseParams.OutInstanceData1Buffer = IndirectDrawResources.InstanceData1;
+			WaterQuadTreeGPUTraverseParams.OutInstanceData2Buffer = IndirectDrawResources.InstanceData2;
+			WaterQuadTreeGPUTraverseParams.PerViewInfo = MoveTemp(PerViewInfoArray);
+			WaterQuadTreeGPUTraverseParams.CullingBounds = WaterInfoBounds.ShiftBy(-WaterQuadTree.GetTileRegion().Min);
+			WaterQuadTreeGPUTraverseParams.NumDensities = DensityCount;
+			WaterQuadTreeGPUTraverseParams.NumMaterials = NumWaterMaterials;
+			WaterQuadTreeGPUTraverseParams.NumViews = NumVisibleViews;
+			WaterQuadTreeGPUTraverseParams.NumQuadsLOD0 = NumQuadsLOD0;
+			WaterQuadTreeGPUTraverseParams.ForceCollapseDensityLevel = ForceCollapseDensityLevel;
+			WaterQuadTreeGPUTraverseParams.LeafSize = WaterQuadTree.GetLeafSize();
+			WaterQuadTreeGPUTraverseParams.LODScale = LODScale;
+			WaterQuadTreeGPUTraverseParams.DebugShowTile = CVarWaterMeshShowTileBounds.GetValueOnRenderThread();
+			WaterQuadTreeGPUTraverseParams.bWithWaterSelectionSupport = WITH_WATER_SELECTION_SUPPORT != 0;
+			WaterQuadTreeGPUTraverseParams.bLODMorphingEnabled = !!CVarWaterMeshLODMorphEnabled.GetValueOnRenderThread();
+
+			bNeedToTraverseGPUQuadTree = true;
+		}
+
+
+		using FWaterVertexFactoryUserDataWrapperType = TWaterVertexFactoryUserDataWrapper<WITH_WATER_SELECTION_SUPPORT>;
+		FWaterVertexFactoryUserDataWrapperType* UserDataWrappers[3] = {};
+		for (EWaterMeshRenderGroupType RenderGroup : BatchRenderGroups)
+		{
+			FWaterVertexFactoryUserDataWrapperType& UserDataWrapper = Collector.AllocateOneFrameResource<FWaterVertexFactoryUserDataWrapperType>();
+			UserDataWrapper.UserData.RenderGroupType = RenderGroup;
+			UserDataWrapper.UserData.IndirectInstanceDataOffsets = IndirectDrawResources.InstanceDataOffsets->GetSRV(RHICmdList, FRHIBufferSRVCreateInfo(PF_R32_UINT));
+			UserDataWrapper.UserData.IndirectInstanceData0 = IndirectDrawResources.InstanceData0->GetSRV(RHICmdList, FRHIBufferSRVCreateInfo(PF_A32B32G32R32F));
+			UserDataWrapper.UserData.IndirectInstanceData1 = IndirectDrawResources.InstanceData1->GetSRV(RHICmdList, FRHIBufferSRVCreateInfo(PF_A32B32G32R32F));
+			UserDataWrapper.UserData.IndirectInstanceData2 = (WITH_WATER_SELECTION_SUPPORT != 0) ? IndirectDrawResources.InstanceData2->GetSRV(RHICmdList, FRHIBufferSRVCreateInfo(PF_A32B32G32R32F)) : nullptr;
+			
+			
+			UserDataWrappers[(int32)RenderGroup] = &UserDataWrapper;
+		}
+
+		// Go through all buckets and issue one batched draw call per LOD level per material per view
+		int32 CompactViewIndex = 0;
+		for (int32 ViewIndex = 0; ViewIndex < Views.Num(); ViewIndex++)
+		{
+			// when rendering ISR, don't process the instanced view
+			if ((VisibilityMap & (1 << ViewIndex)) && (!bEncounteredISRView || Views[ViewIndex]->IsPrimarySceneView()))
+			{
+				TRACE_CPUPROFILER_EVENT_SCOPE(BucketsPerView);
+
+				for (int32 MaterialIndex = 0; MaterialIndex < NumWaterMaterials; ++MaterialIndex)
+				{
+					TRACE_CPUPROFILER_EVENT_SCOPE(MaterialBucket);
+					bool bMaterialDrawn = false;
+
+					for (int32 DensityIndex = 0; DensityIndex < DensityCount; ++DensityIndex)
+					{
+						const int32 BucketIndex = MaterialIndex * DensityCount + DensityIndex;
+
+						TRACE_CPUPROFILER_EVENT_SCOPE(DensityBucket);
+
+						const FMaterialRenderProxy* MaterialRenderProxy = (WireframeMaterialInstance != nullptr) ? WireframeMaterialInstance : WaterQuadTree.GetWaterMaterials()[MaterialIndex];
+						check(MaterialRenderProxy != nullptr);
+
+						bool bUseForDepthPass = false;
+
+						// If there's a valid material, use that to figure out the depth pass status
+						if (const FMaterial* BucketMaterial = MaterialRenderProxy->GetMaterialNoFallback(GetScene().GetFeatureLevel()))
+						{
+							// Preemptively turn off depth rendering for this mesh batch if the material doesn't need it
+							bUseForDepthPass = !BucketMaterial->GetShadingModels().HasShadingModel(MSM_SingleLayerWater) && !IsTranslucentOnlyBlendMode(*BucketMaterial);
+						}
+
+						bMaterialDrawn = true;
+						for (EWaterMeshRenderGroupType RenderGroup : BatchRenderGroups)
+						{
+							// Set up mesh batch
+							FMeshBatch& Mesh = Collector.AllocateMesh();
+							Mesh.bWireframe = bWireframe;
+							Mesh.VertexFactory = WaterVertexFactoriesIndirectDraw[DensityIndex];
+							Mesh.MaterialRenderProxy = MaterialRenderProxy;
+							Mesh.ReverseCulling = IsLocalToWorldDeterminantNegative();
+							Mesh.Type = PT_TriangleList;
+							Mesh.DepthPriorityGroup = SDPG_World;
+							Mesh.bCanApplyViewModeOverrides = false;
+							Mesh.bUseForMaterial = true;
+							Mesh.CastShadow = false;
+							// Preemptively turn off depth rendering for this mesh batch if the material doesn't need it
+							Mesh.bUseForDepthPass = bUseForDepthPass;
+							Mesh.bUseAsOccluder = false;
+
+#if WITH_WATER_SELECTION_SUPPORT
+							Mesh.bUseSelectionOutline = (RenderGroup == EWaterMeshRenderGroupType::RG_RenderSelectedWaterTilesOnly);
+							Mesh.bUseWireframeSelectionColoring = (RenderGroup == EWaterMeshRenderGroupType::RG_RenderSelectedWaterTilesOnly);
+#endif // WITH_WATER_SELECTION_SUPPORT
+
+							Mesh.Elements.SetNumZeroed(1);
+							{
+								TRACE_CPUPROFILER_EVENT_SCOPE_STR("Setup batch element");
+
+								// Set up one mesh batch element
+								FMeshBatchElement& BatchElement = Mesh.Elements[0];
+
+								// Set up for indirect draw
+								BatchElement.FirstIndex = 0;
+								BatchElement.NumPrimitives = 0; // Must be 0 to enable usage of IndirectArgsBuffer
+								BatchElement.IndirectArgsBuffer = IndirectDrawResources.IndirectArgs->GetRHI();
+								BatchElement.IndirectArgsOffset = (CompactViewIndex * NumBuckets + BucketIndex) * sizeof(FRHIDrawIndexedIndirectParameters);
+								BatchElement.UserData = (void*)&UserDataWrappers[(int32)RenderGroup]->UserData;
+								BatchElement.UserIndex = CompactViewIndex * NumBuckets + BucketIndex;
+
+								BatchElement.IndexBuffer = WaterVertexFactoriesIndirectDraw[DensityIndex]->IndexBuffer;
+								BatchElement.PrimitiveIdMode = PrimID_ForceZero;
+
+								// We need the uniform buffer of this primitive because it stores the proper value for the bOutputVelocity flag.
+								// The identity primitive uniform buffer simply stores false for this flag which leads to missing motion vectors.
+								BatchElement.PrimitiveUniformBuffer = GetUniformBuffer();
+							}
+
+							{
+								INC_DWORD_STAT(STAT_WaterDrawCalls);
+								TRACE_CPUPROFILER_EVENT_SCOPE(Collector.AddMesh);
+
+								Collector.AddMesh(ViewIndex, Mesh);
+							}
+						}
+					}
+
+					INC_DWORD_STAT_BY(STAT_WaterDrawnMats, (int32)bMaterialDrawn);
+				}
+
+				++CompactViewIndex;
+			}
+		}
+	}
+	
+	// Even if we have a GPU quadtree, we still go through the setup for the non-GPU path because the water quadtree might have an ocean far mesh which we render traditionally.
+	// If this is a GPU quadtree, traversing it on the CPU will only return the far mesh tiles (if present).
 
 	TArray<FWaterQuadTree::FTraversalOutput, TInlineAllocator<4>> WaterInstanceDataPerView;
 
@@ -280,7 +566,7 @@ void FWaterMeshSceneProxy::GetDynamicMeshElements(const TArray<const FSceneView*
 		if ((VisibilityMap & (1 << ViewIndex)) && (!bEncounteredISRView || View->IsPrimarySceneView()))
 		{
 			const FVector ObserverPosition = View->ViewMatrices.GetViewOrigin();
-			
+
 			FWaterLODParams WaterLODParams = GetWaterLODParams(ObserverPosition);
 
 #if !(UE_BUILD_SHIPPING || UE_BUILD_TEST)
@@ -367,7 +653,6 @@ void FWaterMeshSceneProxy::GetDynamicMeshElements(const TArray<const FSceneView*
 			TRACE_CPUPROFILER_EVENT_SCOPE(BucketsPerView);
 
 			FWaterQuadTree::FTraversalOutput& WaterInstanceData = WaterInstanceDataPerView[TraversalIndex];
-			const int32 NumWaterMaterials = WaterQuadTree.GetWaterMaterials().Num();
 			TraversalIndex++;
 
 			for (int32 MaterialIndex = 0; MaterialIndex < NumWaterMaterials; ++MaterialIndex)
@@ -388,7 +673,7 @@ void FWaterMeshSceneProxy::GetDynamicMeshElements(const TArray<const FSceneView*
 					TRACE_CPUPROFILER_EVENT_SCOPE(DensityBucket);
 
 					const FMaterialRenderProxy* MaterialRenderProxy = (WireframeMaterialInstance != nullptr) ? WireframeMaterialInstance : WaterQuadTree.GetWaterMaterials()[MaterialIndex];
-					check (MaterialRenderProxy != nullptr);
+					check(MaterialRenderProxy != nullptr);
 
 					bool bUseForDepthPass = false;
 
@@ -446,7 +731,7 @@ void FWaterMeshSceneProxy::GetDynamicMeshElements(const TArray<const FSceneView*
 
 							// We need the uniform buffer of this primitive because it stores the proper value for the bOutputVelocity flag.
 							// The identity primitive uniform buffer simply stores false for this flag which leads to missing motion vectors.
-							BatchElement.PrimitiveUniformBuffer = GetUniformBuffer(); 
+							BatchElement.PrimitiveUniformBuffer = GetUniformBuffer();
 						}
 
 						{
@@ -475,7 +760,7 @@ void FWaterMeshSceneProxy::GetDynamicMeshElements(const TArray<const FSceneView*
 				const FWaterQuadTree::FStagingInstanceData& Data = WaterInstanceData.StagingInstanceData[Idx];
 				const int32 WriteIndex = WaterInstanceData.BucketInstanceCounts[Data.BucketIndex]++;
 
-				for (int32 StreamIdx = 0; StreamIdx < WaterInstanceDataBuffersType::NumBuffers; ++StreamIdx)
+				for (int32 StreamIdx = 0; StreamIdx < FWaterInstanceDataBuffersType::NumBuffers; ++StreamIdx)
 				{
 					TArrayView<FVector4f> BufferMemory = WaterInstanceDataBuffers->GetBufferMemory(StreamIdx);
 					for (int32 IdxMultipliedInstance = 0; IdxMultipliedInstance < InstanceFactor; ++IdxMultipliedInstance)
@@ -586,13 +871,6 @@ void FWaterMeshSceneProxy::SetupRayTracingInstances(FRHICommandListBase& RHICmdL
 	}
 }
 
-template <bool bWithWaterSelectionSupport>
-class TWaterVertexFactoryUserDataWrapper : public FOneFrameResource
-{
-public:
-	TWaterMeshUserData<bWithWaterSelectionSupport> UserData;
-};
-
 void FWaterMeshSceneProxy::GetDynamicRayTracingInstances(FRayTracingMaterialGatheringContext& Context, TArray<FRayTracingInstance>& OutRayTracingInstances)
 {
 	if (!HasWaterData() || !FWaterUtils::IsWaterMeshRenderingEnabled(/*bIsRenderThread = */true) || !CVarRayTracingGeometryWater.GetValueOnRenderThread())
@@ -781,3 +1059,117 @@ HHitProxy* FWaterMeshSceneProxy::CreateHitProxies(UPrimitiveComponent* Component
 	return nullptr;
 }
 #endif // WITH_WATER_SELECTION_SUPPORT
+
+void FWaterMeshSceneProxy::BuildGPUQuadTree(FRDGBuilder& GraphBuilder)
+{
+	auto BuildOrthoMatrix = [](float InOrthoWidth, float InOrthoHeight, float DepthRange)
+	{
+		check((int32)ERHIZBuffer::IsInverted);
+
+		const FMatrix::FReal OrthoWidth = InOrthoWidth / 2.0f;
+		const FMatrix::FReal OrthoHeight = InOrthoHeight / 2.0f;
+
+		const FMatrix::FReal NearPlane = 0.f;
+		const FMatrix::FReal FarPlane = DepthRange;
+
+		const FMatrix::FReal ZScale = 1.0f / (FarPlane - NearPlane);
+		const FMatrix::FReal ZOffset = 0;
+
+		return FReversedZOrthoMatrix(
+			OrthoWidth,
+			OrthoHeight,
+			ZScale,
+			ZOffset
+		);
+	};
+
+	// Create GPU array of water body render data
+	const TArrayView<const FWaterBodyRenderData> WaterBodyRenderData = WaterQuadTree.GetWaterBodyRenderData();
+	TArray<FWaterQuadTreeGPU::FWaterBodyRenderDataGPU> WaterBodyRenderDataGPU;
+	{
+		WaterBodyRenderDataGPU.Reserve(WaterBodyRenderData.Num());
+
+		for (const FWaterBodyRenderData& RenderData : WaterBodyRenderData)
+		{
+			FColor HitProxyColor = FColor(0, 0, 0, 0);
+#if WITH_WATER_SELECTION_SUPPORT
+			if (RenderData.HitProxy)
+			{
+				HitProxyColor = RenderData.HitProxy->Id.GetColor();
+				HitProxyColor.A = RenderData.bWaterBodySelected ? 255 : 0;
+			}
+#endif
+
+			FWaterQuadTreeGPU::FWaterBodyRenderDataGPU& RenderDataGPU = WaterBodyRenderDataGPU.AddDefaulted_GetRef();
+			RenderDataGPU.WaterBodyIndex = RenderData.WaterBodyIndex;
+			RenderDataGPU.MaterialIndex = RenderData.MaterialIndex;
+			RenderDataGPU.RiverToLakeMaterialIndex = RenderData.RiverToLakeMaterialIndex;
+			RenderDataGPU.RiverToOceanMaterialIndex = RenderData.RiverToOceanMaterialIndex;
+			RenderDataGPU.WaterBodyType = RenderData.WaterBodyType;
+			RenderDataGPU.HitProxyColorAndIsSelected = HitProxyColor.Bits;
+			RenderDataGPU.SurfaceBaseHeight = float(RenderData.SurfaceBaseHeight - WaterQuadTreeMinHeight);
+			RenderDataGPU.MaxWaveHeight = RenderData.MaxWaveHeight;
+		}
+	}
+
+	const FIntPoint Resolution = WaterQuadTree.GetResolution();
+	const FBox2D TileRegion = WaterQuadTree.GetTileRegion();
+	const float LeafSize = WaterQuadTree.GetLeafSize();
+
+	const FVector2D Center = TileRegion.Min + FVector2D(Resolution) * (double)LeafSize * 0.5;
+	const FVector QuadTreeCorner = FVector(TileRegion.Min, WaterQuadTreeMinHeight); // Make everything relative to this to avoid LWC precision issues
+	const float WaterQuadTreeDepthRange = WaterQuadTreeMaxHeight - WaterQuadTreeMinHeight;
+	const float RcpWaterQuadTreeDepthRange = 1.0f / WaterQuadTreeDepthRange;
+
+	// Create view/projection matrices
+	const FVector ViewLocation = FVector(Center, WaterQuadTreeMaxHeight) - QuadTreeCorner;
+	const FVector LookAt = ViewLocation - FVector(0.0, 0.0, 1.0);
+	const FMatrix44f ViewMatrix = FMatrix44f(FLookAtMatrix(ViewLocation, LookAt, FVector(0.f, -1.f, 0.f)));
+	const FMatrix44f ProjectionMatrix = FMatrix44f(BuildOrthoMatrix(Resolution.X * LeafSize, Resolution.Y * LeafSize, WaterQuadTreeDepthRange));
+	const FMatrix44f ViewProjection = ViewMatrix * ProjectionMatrix;
+
+	// Create array of water body draws
+	TArray<FWaterQuadTreeGPU::FDraw> Draws;
+	{
+		const TArray<FWaterBodyQuadTreeRasterInfo>& WaterBodyRasterInfos = WaterQuadTree.GetWaterBodyRasterInfos();
+		for (const FWaterBodyQuadTreeRasterInfo& Info : WaterBodyRasterInfos)
+		{
+			const FStaticMeshLODResources& LODResource = Info.RenderData->LODResources[0];
+			for (const FStaticMeshSection& MeshSection : LODResource.Sections)
+			{
+				if (LODResource.IndexBuffer.GetRHI())
+				{
+					// Make relative to quadtree to avoid LWC precision issues
+					FTransform LocalToQuadTreeWorld = Info.LocalToWorld;
+					LocalToQuadTreeWorld.SetTranslation(LocalToQuadTreeWorld.GetTranslation() - QuadTreeCorner);
+
+					const FWaterBodyRenderData& WBRenderData = WaterBodyRenderData[Info.WaterBodyRenderDataIndex];
+
+					FWaterQuadTreeGPU::FDraw& Draw = Draws.AddDefaulted_GetRef();
+					Draw.Transform = FMatrix44f(LocalToQuadTreeWorld.ToMatrixWithScale()) * ViewProjection;
+					Draw.VertexBuffer = LODResource.VertexBuffers.PositionVertexBuffer.GetRHI();
+					Draw.IndexBuffer = LODResource.IndexBuffer.GetRHI();
+					Draw.FirstIndex = MeshSection.FirstIndex;
+					Draw.NumPrimitives = MeshSection.NumTriangles;
+					Draw.BaseVertexIndex = 0;
+					Draw.NumVertices = MeshSection.MaxVertexIndex - MeshSection.MinVertexIndex + 1;
+					Draw.WaterBodyRenderDataIndex = Info.WaterBodyRenderDataIndex;
+					Draw.Priority = Info.Priority;
+					Draw.MinZ = float(WBRenderData.BoundsMinZ - WaterQuadTreeMinHeight) * RcpWaterQuadTreeDepthRange;
+					Draw.MaxZ = float(WBRenderData.BoundsMaxZ + WBRenderData.MaxWaveHeight - WaterQuadTreeMinHeight) * RcpWaterQuadTreeDepthRange;
+					Draw.MaxWaveHeight = WBRenderData.MaxWaveHeight * RcpWaterQuadTreeDepthRange;
+					Draw.bIsRiver = Info.bIsRiver;
+				}
+			}
+		}
+	}
+
+	FWaterQuadTreeGPU::FInitParams Params;
+	Params.WaterBodyRenderData = WaterBodyRenderDataGPU;
+	Params.RequestedQuadTreeResolution = Resolution;
+	Params.SuperSamplingFactor = FMath::Clamp(CVarWaterMeshGPUQuadTreeSuperSampling.GetValueOnRenderThread(), 1, 8);
+	Params.NumMSAASamples = FMath::Clamp(CVarWaterMeshGPUQuadTreeMultiSampling.GetValueOnRenderThread(), 1, 8);
+	Params.CaptureDepthRange = WaterQuadTreeDepthRange;
+
+	QuadTreeGPU.Init(GraphBuilder, Params, Draws);
+}
diff --git a/Engine/Plugins/Experimental/Water/Source/Runtime/Private/WaterQuadTree.cpp b/Engine/Plugins/Experimental/Water/Source/Runtime/Private/WaterQuadTree.cpp
--- a/Engine/Plugins/Experimental/Water/Source/Runtime/Private/WaterQuadTree.cpp
+++ b/Engine/Plugins/Experimental/Water/Source/Runtime/Private/WaterQuadTree.cpp
@@ -451,13 +451,15 @@ bool FWaterQuadTree::FNode::QueryBoundsAtLocation(const FNodeData& InNodeData, c
 	return ChildCount == 0;
 }
 
-void FWaterQuadTree::InitTree(const FBox2D& InBounds, float InTileSize, FIntPoint InExtentInTiles)
+void FWaterQuadTree::InitTree(const FBox2D& InBounds, float InTileSize, FIntPoint InExtentInTiles, bool bInIsGPUQuadTree)
 {
 	ensure(InBounds.GetArea() > 0.0f);
 	ensure(InTileSize > 0.0f);
 	ensure(InExtentInTiles.X > 0);
 	ensure(InExtentInTiles.Y > 0);
 
+	bIsGPUQuadTree = bInIsGPUQuadTree;
+
 	FarMeshData.Clear();
 
 	// Maximum number of allocated leaf nodes for this config
@@ -472,11 +474,19 @@ void FWaterQuadTree::InitTree(const FBox2D& InBounds, float InTileSize, FIntPoin
 
 	TileRegion = InBounds;
 
+	WaterBodyRasterInfos.Reset();
 	BreadthFirstOrder.Reset();
 
 	// Allocate theoretical max, shrink later in Lock()
 	// This is so that the node array doesn't move in memory while inserting
-	NodeData.Nodes.Empty((float)(FMath::Square(RootDim) * 4) / 3.0f);
+	if (!bIsGPUQuadTree)
+	{
+		NodeData.Nodes.Empty((float)(FMath::Square(RootDim) * 4) / 3.0f);
+	}
+	else
+	{
+		NodeData.Nodes.Empty(1);
+	}
 
 	// Add defaulted water body render data to slot 0. This is the "null" render data, pointed to by all newly created nodes. Has lowest priority so it will always be overwritten
 	NodeData.WaterBodyRenderData.Empty(1);
@@ -635,12 +645,14 @@ void FWaterQuadTree::Unlock(bool bPruneRedundantNodes)
 void FWaterQuadTree::AddWaterTilesInsideBounds(const FBox& InBounds, uint32 InWaterBodyIndex)
 {
 	check(!bIsReadOnly);
+	check(!bIsGPUQuadTree);
 	NodeData.Nodes[0].AddNodes(NodeData, FBox(FVector(TileRegion.Min, 0.0f), FVector(TileRegion.Max, 0.0f)),  InBounds, InWaterBodyIndex, TreeDepth, 0);
 }
 
 void FWaterQuadTree::AddOcean(const TArray<FVector2D>& InPoly, const FBox& InOceanBounds, uint32 InWaterBodyIndex)
 {
 	check(!bIsReadOnly);
+	check(!bIsGPUQuadTree);
 	const FBox2D OceanBounds(FVector2D(InOceanBounds.Min), FVector2D(InOceanBounds.Max));
 	AddOceanRecursive(InPoly, OceanBounds, FVector2D(InOceanBounds.Min.Z, InOceanBounds.Max.Z), true, TreeDepth * 2, InWaterBodyIndex);
 }
@@ -648,6 +660,7 @@ void FWaterQuadTree::AddOcean(const TArray<FVector2D>& InPoly, const FBox& InOce
 void FWaterQuadTree::AddLake(const TArray<FVector2D>& InPoly, const FBox& InLakeBounds, uint32 InWaterBodyIndex)
 {
 	check(!bIsReadOnly);
+	check(!bIsGPUQuadTree);
 	const FBox2D LakeBounds(FVector2D(NodeData.Nodes[0].Bounds.Min), FVector2D(NodeData.Nodes[0].Bounds.Max));
 	AddLakeRecursive(InPoly, LakeBounds, FVector2D(InLakeBounds.Min.Z, InLakeBounds.Max.Z), true, TreeDepth * 2, InWaterBodyIndex);
 }
@@ -743,7 +756,10 @@ void FWaterQuadTree::BuildWaterTileInstanceData(const FTraversalDesc& InTraversa
 	TRACE_CPUPROFILER_EVENT_SCOPE(BuildWaterTileInstanceData);
 	check(bIsReadOnly);
 
-	NodeData.Nodes[0].SelectLOD(NodeData, TreeDepth, InTraversalDesc, Output);
+	if (!bIsGPUQuadTree)
+	{
+		NodeData.Nodes[0].SelectLOD(NodeData, TreeDepth, InTraversalDesc, Output);
+	}
 
 	// Append Far Mesh tiles
 	if (FarMeshData.InstanceData.Num() > 0 && FarMeshData.MaterialIndex != INDEX_NONE)
diff --git a/Engine/Plugins/Experimental/Water/Source/Runtime/Private/WaterQuadTreeGPU.cpp b/Engine/Plugins/Experimental/Water/Source/Runtime/Private/WaterQuadTreeGPU.cpp
new file mode 100644
--- /dev/null
+++ b/Engine/Plugins/Experimental/Water/Source/Runtime/Private/WaterQuadTreeGPU.cpp
@@ -0,0 +1,680 @@
+// Copyright Epic Games, Inc. All Rights Reserved.
+
+#include "WaterQuadTreeGPU.h"
+
+#include "CoreMinimal.h"
+#include "GlobalShader.h"
+#include "RenderUtils.h"
+#include "RenderGraphUtils.h"
+#include "RHI.h"
+#include "RHIStaticStates.h"
+#include "RHIFeatureLevel.h"
+#include "PipelineStateCache.h"
+#include "PixelShaderUtils.h"
+#include "SystemTextures.h"
+#include "ShaderPrintParameters.h"
+#include "SceneRendering.h"
+
+DECLARE_GPU_STAT(FWaterQuadTreeGPU_Init);
+DECLARE_GPU_STAT(FWaterQuadTreeGPU_Traverse);
+
+int32 GWaterQuadTreeParallelPrefixSum = 1;
+static FAutoConsoleVariableRef CVarWaterQuadTreeParallelPrefixSum(
+	TEXT("r.Water.WaterMesh.GPUQuadTree.ParallelPrefixSum"),
+	GWaterQuadTreeParallelPrefixSum,
+	TEXT("Enables using a parallel prefix sum algorithm for the water quadtree indirect draw call pipeline."),
+	ECVF_RenderThreadSafe);
+
+float GWaterQuadTreeZBoundsPadding = 200.0f;
+static FAutoConsoleVariableRef CVarWaterQuadTreeZBoundsPadding(
+	TEXT("r.Water.WaterMesh.GPUQuadTree.ZBoundsPadding"),
+	GWaterQuadTreeZBoundsPadding,
+	TEXT("Amount of padding to apply to the Z bounds of each quadtree node."
+		" Necessary for sloped rivers, whose complete Z bounds can be underestimated due to using non-conservative rasterization to build the quadtree."),
+	ECVF_RenderThreadSafe);
+
+int32 GWaterQuadTreeOcclusionCulling = 1;
+static FAutoConsoleVariableRef CVarWaterQuadTreeOcclusionCulling(
+	TEXT("r.Water.WaterMesh.GPUQuadTree.OcclusionCulling"),
+	GWaterQuadTreeOcclusionCulling,
+	TEXT("Enables HZB occlusion culling for the water quadtree indirect draw call pipeline."),
+	ECVF_RenderThreadSafe);
+
+class FWaterQuadTreeVS : public FGlobalShader
+{
+public:
+	DECLARE_GLOBAL_SHADER(FWaterQuadTreeVS);
+	SHADER_USE_PARAMETER_STRUCT(FWaterQuadTreeVS, FGlobalShader);
+
+	BEGIN_SHADER_PARAMETER_STRUCT(FParameters, )
+		SHADER_PARAMETER(FMatrix44f, Transform)
+	END_SHADER_PARAMETER_STRUCT()
+
+	static void ModifyCompilationEnvironment(const FGlobalShaderPermutationParameters& Parameters, FShaderCompilerEnvironment& OutEnvironment)
+	{
+		FGlobalShader::ModifyCompilationEnvironment(Parameters, OutEnvironment);
+	}
+};
+IMPLEMENT_GLOBAL_SHADER(FWaterQuadTreeVS, "/Plugin/Water/Private/WaterQuadTreeVertexShader.usf", "Main", SF_Vertex);
+
+class FWaterQuadTreePS : public FGlobalShader
+{
+public:
+	DECLARE_GLOBAL_SHADER(FWaterQuadTreePS);
+	SHADER_USE_PARAMETER_STRUCT(FWaterQuadTreePS, FGlobalShader);
+
+	BEGIN_SHADER_PARAMETER_STRUCT(FParameters, )
+		SHADER_PARAMETER(uint32, WaterBodyRenderDataIndex)
+		SHADER_PARAMETER(int32, Priority)
+		SHADER_PARAMETER(float, WaterBodyMinZ)
+		SHADER_PARAMETER(float, WaterBodyMaxZ)
+		SHADER_PARAMETER(float, MaxWaveHeight)
+		SHADER_PARAMETER(uint32, bIsRiver)
+		RENDER_TARGET_BINDING_SLOTS()
+	END_SHADER_PARAMETER_STRUCT()
+
+	static void ModifyCompilationEnvironment(const FGlobalShaderPermutationParameters& Parameters, FShaderCompilerEnvironment& OutEnvironment)
+	{
+		FGlobalShader::ModifyCompilationEnvironment(Parameters, OutEnvironment);
+	}
+};
+IMPLEMENT_GLOBAL_SHADER(FWaterQuadTreePS, "/Plugin/Water/Private/WaterQuadTreePixelShader.usf", "Main", SF_Pixel);
+
+BEGIN_SHADER_PARAMETER_STRUCT(FWaterQuadTreeParameters, )
+	SHADER_PARAMETER_STRUCT_INCLUDE(FWaterQuadTreeVS::FParameters, VS)
+	SHADER_PARAMETER_STRUCT_INCLUDE(FWaterQuadTreePS::FParameters, PS)
+END_SHADER_PARAMETER_STRUCT()
+
+class FWaterQuadTreeMergePS : public FGlobalShader
+{
+public:
+	DECLARE_GLOBAL_SHADER(FWaterQuadTreeMergePS);
+	SHADER_USE_PARAMETER_STRUCT(FWaterQuadTreeMergePS, FGlobalShader);
+
+	class FNumMSAASamples : SHADER_PERMUTATION_SPARSE_INT("NUM_MSAA_SAMPLES", 1, 2, 4, 8);
+
+	using FPermutationDomain = TShaderPermutationDomain<FNumMSAASamples>;
+
+	BEGIN_SHADER_PARAMETER_STRUCT(FParameters, )
+		SHADER_PARAMETER_RDG_TEXTURE(Texture2DMS, WaterBodyRasterTextureMS)
+		SHADER_PARAMETER_RDG_TEXTURE(Texture2DMS, ZBoundsRasterTextureMS)
+		SHADER_PARAMETER_RDG_TEXTURE(Texture2D, WaterBodyRasterTexture)
+		SHADER_PARAMETER_RDG_TEXTURE(Texture2D, ZBoundsRasterTexture)
+		SHADER_PARAMETER_RDG_BUFFER_SRV(StructuredBuffer, WaterBodyRenderData)
+		SHADER_PARAMETER(int32, SuperSamplingFactor)
+		SHADER_PARAMETER(float, RcpCaptureDepthRange)
+		SHADER_PARAMETER(float, ZBoundsPadding)
+		RENDER_TARGET_BINDING_SLOTS()
+	END_SHADER_PARAMETER_STRUCT()
+
+	static void ModifyCompilationEnvironment(const FGlobalShaderPermutationParameters& Parameters, FShaderCompilerEnvironment& OutEnvironment)
+	{
+		FGlobalShader::ModifyCompilationEnvironment(Parameters, OutEnvironment);
+	}
+};
+IMPLEMENT_GLOBAL_SHADER(FWaterQuadTreeMergePS, "/Plugin/Water/Private/WaterQuadTreeMerge.usf", "Main", SF_Pixel);
+
+class FWaterQuadTreeBuildPS : public FGlobalShader
+{
+public:
+	DECLARE_GLOBAL_SHADER(FWaterQuadTreeBuildPS);
+	SHADER_USE_PARAMETER_STRUCT(FWaterQuadTreeBuildPS, FGlobalShader);
+
+	BEGIN_SHADER_PARAMETER_STRUCT(FParameters, )
+		SHADER_PARAMETER_RDG_TEXTURE_SRV(Texture2D, QuadTreeTexture)
+		SHADER_PARAMETER_RDG_TEXTURE_SRV(Texture2D, WaterZBoundsTexture)
+		RENDER_TARGET_BINDING_SLOTS()
+	END_SHADER_PARAMETER_STRUCT()
+
+	static void ModifyCompilationEnvironment(const FGlobalShaderPermutationParameters& Parameters, FShaderCompilerEnvironment& OutEnvironment)
+	{
+		FGlobalShader::ModifyCompilationEnvironment(Parameters, OutEnvironment);
+	}
+};
+IMPLEMENT_GLOBAL_SHADER(FWaterQuadTreeBuildPS, "/Plugin/Water/Private/WaterQuadTreeBuild.usf", "Main", SF_Pixel);
+
+class FWaterQuadTreeInitializeIndirectArgsCS : public FGlobalShader
+{
+public:
+	DECLARE_GLOBAL_SHADER(FWaterQuadTreeInitializeIndirectArgsCS);
+	SHADER_USE_PARAMETER_STRUCT(FWaterQuadTreeInitializeIndirectArgsCS, FGlobalShader);
+
+	BEGIN_SHADER_PARAMETER_STRUCT(FParameters, )
+		SHADER_PARAMETER_RDG_BUFFER_UAV(RWStructuredBuffer, IndirectArgs)
+		SHADER_PARAMETER(uint32, NumDensities)
+		SHADER_PARAMETER(uint32, NumMaterials)
+		SHADER_PARAMETER(uint32, NumViews)
+		SHADER_PARAMETER(uint32, NumQuadsLOD0)
+	END_SHADER_PARAMETER_STRUCT()
+
+	static void ModifyCompilationEnvironment(const FGlobalShaderPermutationParameters& Parameters, FShaderCompilerEnvironment& OutEnvironment)
+	{
+		FGlobalShader::ModifyCompilationEnvironment(Parameters, OutEnvironment);
+		OutEnvironment.SetDefine(TEXT("INITIALIZE_INDIRECT_ARGS"), 1);
+	}
+};
+IMPLEMENT_GLOBAL_SHADER(FWaterQuadTreeInitializeIndirectArgsCS, "/Plugin/Water/Private/WaterQuadTreeDraws.usf", "MainCS", SF_Compute);
+
+class FWaterQuadTreeTraverseCS : public FGlobalShader
+{
+public:
+	DECLARE_GLOBAL_SHADER(FWaterQuadTreeTraverseCS);
+	SHADER_USE_PARAMETER_STRUCT(FWaterQuadTreeTraverseCS, FGlobalShader);
+
+	class FDebugShaderPrint : SHADER_PERMUTATION_BOOL("DEBUG_SHADER_PRINT");
+
+	using FPermutationDomain = TShaderPermutationDomain<FDebugShaderPrint>;
+
+	BEGIN_SHADER_PARAMETER_STRUCT(FParameters, )
+		SHADER_PARAMETER_STRUCT_INCLUDE(ShaderPrint::FShaderParameters, ShaderPrintUniformBuffer)
+		SHADER_PARAMETER_RDG_BUFFER_UAV(RWByteAddressBuffer, PackedNodes)
+		SHADER_PARAMETER_RDG_TEXTURE(Texture2D, QuadTreeTexture)
+		SHADER_PARAMETER_RDG_TEXTURE(Texture2D, WaterZBoundsTexture)
+		SHADER_PARAMETER_RDG_BUFFER_SRV(StructuredBuffer, WaterBodyRenderData)
+		SHADER_PARAMETER_RDG_TEXTURE(Texture2D, HZBTexture)
+		SHADER_PARAMETER_SAMPLER(SamplerState, HZBSampler)
+		SHADER_PARAMETER(FMatrix44f, TranslatedWorldToClip)
+		SHADER_PARAMETER(FMatrix44f, ViewToClip)
+		SHADER_PARAMETER(FVector4f, CullingBoundsAABB)
+		SHADER_PARAMETER(FVector2f, HZBSize)
+		SHADER_PARAMETER(FVector2f, HZBViewSize)
+		SHADER_PARAMETER(FVector3f, QuadTreePosition)
+		SHADER_PARAMETER(FVector3f, ObserverPosition)
+		SHADER_PARAMETER(uint32, QuadTreeResolutionX)
+		SHADER_PARAMETER(uint32, QuadTreeResolutionY)
+		SHADER_PARAMETER(uint32, NumDensities)
+		SHADER_PARAMETER(uint32, NumMaterials)
+		SHADER_PARAMETER(float, LeafSize)
+		SHADER_PARAMETER(float, LODScale)
+		SHADER_PARAMETER(float, CaptureDepthRange)
+		SHADER_PARAMETER(int32, ForceCollapseDensityLevel)
+		SHADER_PARAMETER(uint32, NumLODs)
+		SHADER_PARAMETER(uint32, NumDispatchedThreads)
+		SHADER_PARAMETER(int32, DebugShowTile)
+		SHADER_PARAMETER(uint32, bHZBOcclusionCullingEnabled)
+	END_SHADER_PARAMETER_STRUCT()
+
+	static bool ShouldCompilePermutation(const FGlobalShaderPermutationParameters& Parameters)
+	{
+		FPermutationDomain PermutationVector(Parameters.PermutationId);
+		const bool bDebugShaderPrint = PermutationVector.Get<FDebugShaderPrint>() != 0;
+		return !bDebugShaderPrint || ShaderPrint::IsSupported(Parameters.Platform);
+	}
+
+	static void ModifyCompilationEnvironment(const FGlobalShaderPermutationParameters& Parameters, FShaderCompilerEnvironment& OutEnvironment)
+	{
+		FGlobalShader::ModifyCompilationEnvironment(Parameters, OutEnvironment);
+		OutEnvironment.SetDefine(TEXT("QUAD_TREE_TRAVERSE"), 1);
+	}
+};
+IMPLEMENT_GLOBAL_SHADER(FWaterQuadTreeTraverseCS, "/Plugin/Water/Private/WaterQuadTreeDraws.usf", "MainCS", SF_Compute);
+
+class FWaterQuadTreeBucketCountsCS : public FGlobalShader
+{
+public:
+	DECLARE_GLOBAL_SHADER(FWaterQuadTreeBucketCountsCS);
+	SHADER_USE_PARAMETER_STRUCT(FWaterQuadTreeBucketCountsCS, FGlobalShader);
+
+	BEGIN_SHADER_PARAMETER_STRUCT(FParameters, )
+		SHADER_PARAMETER_RDG_BUFFER_UAV(RWByteAddressBuffer, BucketCounts)
+		SHADER_PARAMETER_RDG_TEXTURE(Texture2D, QuadTreeTexture)
+		SHADER_PARAMETER_RDG_BUFFER_SRV(StructuredBuffer, WaterBodyRenderData)
+		SHADER_PARAMETER_RDG_BUFFER_SRV(ByteAddressBuffer, PackedNodes)
+		SHADER_PARAMETER(uint32, NumDensities)
+		SHADER_PARAMETER(uint32, NumMaterials)
+		SHADER_PARAMETER(uint32, NumDispatchedThreads)
+	END_SHADER_PARAMETER_STRUCT()
+
+	static void ModifyCompilationEnvironment(const FGlobalShaderPermutationParameters& Parameters, FShaderCompilerEnvironment& OutEnvironment)
+	{
+		FGlobalShader::ModifyCompilationEnvironment(Parameters, OutEnvironment);
+		OutEnvironment.SetDefine(TEXT("COMPUTE_BUCKET_COUNTS"), 1);
+	}
+};
+IMPLEMENT_GLOBAL_SHADER(FWaterQuadTreeBucketCountsCS, "/Plugin/Water/Private/WaterQuadTreeDraws.usf", "MainCS", SF_Compute);
+
+class FWaterQuadTreeBucketPrefixSumCS : public FGlobalShader
+{
+public:
+	DECLARE_GLOBAL_SHADER(FWaterQuadTreeBucketPrefixSumCS);
+	SHADER_USE_PARAMETER_STRUCT(FWaterQuadTreeBucketPrefixSumCS, FGlobalShader);
+
+	class FParallelPrefixSum : SHADER_PERMUTATION_BOOL("PARALLEL_PREFIX_SUM");
+	using FPermutationDomain = TShaderPermutationDomain<FParallelPrefixSum>;
+
+	BEGIN_SHADER_PARAMETER_STRUCT(FParameters, )
+		SHADER_PARAMETER_RDG_BUFFER_UAV(RWBuffer, BucketPrefixSums)
+		SHADER_PARAMETER_RDG_BUFFER_SRV(ByteAddressBuffer, BucketCounts)
+		SHADER_PARAMETER(uint32, NumBuckets)
+		SHADER_PARAMETER(uint32, OutputOffset)
+	END_SHADER_PARAMETER_STRUCT()
+
+	static void ModifyCompilationEnvironment(const FGlobalShaderPermutationParameters& Parameters, FShaderCompilerEnvironment& OutEnvironment)
+	{
+		FGlobalShader::ModifyCompilationEnvironment(Parameters, OutEnvironment);
+		OutEnvironment.SetDefine(TEXT("COMPUTE_BUCKET_PREFIX_SUM"), 1);
+	}
+};
+IMPLEMENT_GLOBAL_SHADER(FWaterQuadTreeBucketPrefixSumCS, "/Plugin/Water/Private/WaterQuadTreeDraws.usf", "MainCS", SF_Compute);
+
+class FWaterQuadTreeGenerateInstanceDataCS : public FGlobalShader
+{
+public:
+	DECLARE_GLOBAL_SHADER(FWaterQuadTreeGenerateInstanceDataCS);
+	SHADER_USE_PARAMETER_STRUCT(FWaterQuadTreeGenerateInstanceDataCS, FGlobalShader);
+
+	BEGIN_SHADER_PARAMETER_STRUCT(FParameters, )
+		SHADER_PARAMETER_RDG_BUFFER_UAV(RWStructuredBuffer, IndirectArgs)
+		SHADER_PARAMETER_RDG_BUFFER_UAV(RWBuffer, InstanceData0)
+		SHADER_PARAMETER_RDG_BUFFER_UAV(RWBuffer, InstanceData1)
+		SHADER_PARAMETER_RDG_BUFFER_UAV(RWBuffer, InstanceData2)
+		SHADER_PARAMETER_RDG_TEXTURE(Texture2D, QuadTreeTexture)
+		SHADER_PARAMETER_RDG_TEXTURE(Texture2D, WaterZBoundsTexture)
+		SHADER_PARAMETER_RDG_BUFFER_SRV(StructuredBuffer, WaterBodyRenderData)
+		SHADER_PARAMETER_RDG_BUFFER_SRV(ByteAddressBuffer, PackedNodes)
+		SHADER_PARAMETER_RDG_BUFFER_SRV(Buffer, InstanceDataOffsets)
+		SHADER_PARAMETER(FVector3f, QuadTreePosition)
+		SHADER_PARAMETER(FVector3f, ObserverPosition)
+		SHADER_PARAMETER(uint32, QuadTreeResolutionX)
+		SHADER_PARAMETER(uint32, QuadTreeResolutionY)
+		SHADER_PARAMETER(uint32, NumDensities)
+		SHADER_PARAMETER(uint32, NumMaterials)
+		SHADER_PARAMETER(uint32, NumDispatchedThreads)
+		SHADER_PARAMETER(uint32, BucketIndexOffset)
+		SHADER_PARAMETER(uint32, NumLODs)
+		SHADER_PARAMETER(float, LeafSize)
+		SHADER_PARAMETER(float, LODScale)
+		SHADER_PARAMETER(float, CaptureDepthRange)
+		SHADER_PARAMETER(uint32, bWithWaterSelectionSupport)
+		SHADER_PARAMETER(uint32, bLODMorphingEnabled)
+	END_SHADER_PARAMETER_STRUCT()
+
+	static void ModifyCompilationEnvironment(const FGlobalShaderPermutationParameters& Parameters, FShaderCompilerEnvironment& OutEnvironment)
+	{
+		FGlobalShader::ModifyCompilationEnvironment(Parameters, OutEnvironment);
+		OutEnvironment.SetDefine(TEXT("GENERATE_INSTANCE_DATA"), 1);
+	}
+};
+IMPLEMENT_GLOBAL_SHADER(FWaterQuadTreeGenerateInstanceDataCS, "/Plugin/Water/Private/WaterQuadTreeDraws.usf", "MainCS", SF_Compute);
+
+
+void FWaterQuadTreeGPU::Init(FRDGBuilder& GraphBuilder, const FInitParams& Params, TArray<FDraw>& Draws)
+{
+	TRACE_CPUPROFILER_EVENT_SCOPE(FWaterQuadTreeGPU::Init);
+	RDG_EVENT_SCOPE(GraphBuilder, "FWaterQuadTreeGPU::Init");
+	RDG_GPU_STAT_SCOPE(GraphBuilder, FWaterQuadTreeGPU_Init);
+
+	// Create resources
+	{
+		check(!bInitialized);
+		check(!QuadTreeTexture);
+		check(!WaterZBoundsTexture);
+		check(!WaterBodyRenderDataBuffer);
+
+		// Make sure the quadtree is a power of two
+		const FIntPoint ResolutionRaw = Params.RequestedQuadTreeResolution;
+		const FIntPoint ResolutionPow2 = FIntPoint(FMath::RoundUpToPowerOfTwo(ResolutionRaw.X), FMath::RoundUpToPowerOfTwo(ResolutionRaw.Y));
+		const int32 MaxDim = (int32)FMath::Max(ResolutionPow2.X, ResolutionPow2.Y);
+		const int32 NumMipLevels = (int32)FMath::FloorLog2(MaxDim) + 1;
+
+		FRHITextureCreateDesc QuadTreeTextureCreateDesc = FRHITextureCreateDesc::Create2D(TEXT("WaterQuadTree.QuadTree"), ResolutionPow2, PF_B8G8R8A8)
+			.SetNumMips(NumMipLevels)
+			.SetFlags(TexCreate_RenderTargetable | TexCreate_ShaderResource);
+		QuadTreeTexture = RHICreateTexture(QuadTreeTextureCreateDesc);
+
+		FRHITextureCreateDesc ZBoundsTextureCreateDesc = FRHITextureCreateDesc::Create2D(TEXT("WaterQuadTree.ZBoundsTexture"), ResolutionPow2, PF_A2B10G10R10)
+			.SetNumMips(NumMipLevels)
+			.SetFlags(TexCreate_RenderTargetable | TexCreate_ShaderResource);
+		WaterZBoundsTexture = RHICreateTexture(ZBoundsTextureCreateDesc);
+
+		WaterBodyRenderDataBuffer = GraphBuilder.ConvertToExternalBuffer(CreateStructuredBuffer<FWaterBodyRenderDataGPU>(GraphBuilder, TEXT("WaterQuadTree.WaterBodyRenderData"), Params.WaterBodyRenderData));
+
+		check(Params.CaptureDepthRange > 0.0f);
+		CaptureDepthRange = Params.CaptureDepthRange;
+		bInitialized = true;
+	}
+
+	auto GetNumMSAASamples = [](int32 RequestedNumSamples)
+	{
+		if (RequestedNumSamples <= 1)
+		{
+			return 1;
+		}
+		else if (RequestedNumSamples < 4)
+		{
+			return 2;
+		}
+		else if (RequestedNumSamples < 8)
+		{
+			return 4;
+		}
+		else
+		{
+			return 8;
+		}
+	};
+	
+
+	FGlobalShaderMap* ShaderMap = GetGlobalShaderMap(GMaxRHIFeatureLevel);
+
+	FRDGTexture* QuadTreeTextureRDG = GraphBuilder.RegisterExternalTexture(CreateRenderTarget(QuadTreeTexture, TEXT("WaterQuadTree.QuadTree")));
+	FRDGTexture* WaterZBoundsTextureRDG = GraphBuilder.RegisterExternalTexture(CreateRenderTarget(WaterZBoundsTexture, TEXT("WaterQuadTree.ZBoundsTexture")));
+
+	FRDGBuffer* WaterBodyRenderDataBufferRDG = GraphBuilder.RegisterExternalBuffer(WaterBodyRenderDataBuffer);
+	FRDGBufferSRV* WaterBodyRenderDataBufferSRV = GraphBuilder.CreateSRV(FRDGBufferSRVDesc(WaterBodyRenderDataBufferRDG));
+
+	const FIntPoint QuadTreeResolution = QuadTreeTextureRDG->Desc.Extent;
+	// Not necessarily a power of two or a multiple of the quadtree LOD0 resolution. We rely on Texture.Load() to return 0 for out of bounds accesses.
+	// Padding this texture is not required as we will not need to create a mip chain for it; but rather use it to initialize mip0 of our quadtree texture.
+	const FIntPoint RasterResolution = Params.RequestedQuadTreeResolution * Params.SuperSamplingFactor;
+	const int32 NumMipLevels = QuadTreeTextureRDG->Desc.NumMips;
+	const uint8 NumSamples = GetNumMSAASamples(Params.NumMSAASamples);
+
+	FRDGTextureDesc WaterBodyRasterTextureDesc = FRDGTextureDesc::Create2D(RasterResolution, PF_A2B10G10R10, FClearValueBinding::Black, TexCreate_RenderTargetable | TexCreate_ShaderResource, 1, NumSamples);
+	FRDGTexture* WaterBodyRasterTexture = GraphBuilder.CreateTexture(WaterBodyRasterTextureDesc, TEXT("WaterQuadTree.WaterBodyRasterTexture"));
+	FRDGTextureDesc ZBoundsRasterTextureDesc = FRDGTextureDesc::Create2D(RasterResolution, PF_A2B10G10R10, FClearValueBinding::Black, TexCreate_RenderTargetable | TexCreate_ShaderResource, 1, NumSamples);
+	FRDGTexture* ZBoundsRasterTexture = GraphBuilder.CreateTexture(ZBoundsRasterTextureDesc, TEXT("WaterQuadTree.ZBoundsRasterTexture"));
+
+
+	// Raster water body meshes
+	{
+		TShaderMapRef<FWaterQuadTreeVS> VertexShader(ShaderMap);
+		TShaderMapRef<FWaterQuadTreePS> PixelShader(ShaderMap);
+
+		FWaterQuadTreeParameters* PassParameters = GraphBuilder.AllocParameters<FWaterQuadTreeParameters>();
+		PassParameters->PS.RenderTargets[0] = FRenderTargetBinding(WaterBodyRasterTexture, ERenderTargetLoadAction::EClear);
+		PassParameters->PS.RenderTargets[1] = FRenderTargetBinding(ZBoundsRasterTexture, ERenderTargetLoadAction::EClear);
+
+		GraphBuilder.AddPass(
+			RDG_EVENT_NAME("WaterQuadTreeRaster"),
+			PassParameters,
+			ERDGPassFlags::Raster,
+			[PassParameters, RasterResolution, Draws = MoveTemp(Draws), VertexShader, PixelShader](FRHICommandList& RHICmdList)
+			{
+				RHICmdList.SetViewport(0.0f, 0.0f, 0.0f, RasterResolution.X, RasterResolution.Y, 1.0f);
+
+				FGraphicsPipelineStateInitializer GraphicsPSOInit;
+				RHICmdList.ApplyCachedRenderTargets(GraphicsPSOInit);
+				GraphicsPSOInit.RasterizerState = TStaticRasterizerState<FM_Solid, CM_None>::GetRHI();
+				GraphicsPSOInit.DepthStencilState = TStaticDepthStencilState<false, CF_Always>::GetRHI();
+				GraphicsPSOInit.BlendState = TStaticBlendState<
+					CW_RGBA, BO_Max, BF_One, BF_One, BO_Max, BF_One, BF_One,
+					CW_RGBA, BO_Max, BF_One, BF_One, BO_Max, BF_One, BF_One>::GetRHI(); // MAX blending
+				GraphicsPSOInit.PrimitiveType = PT_TriangleList;
+				GraphicsPSOInit.BoundShaderState.VertexDeclarationRHI = GetVertexDeclarationFVector3();
+				GraphicsPSOInit.BoundShaderState.VertexShaderRHI = VertexShader.GetVertexShader();
+				GraphicsPSOInit.BoundShaderState.PixelShaderRHI = PixelShader.GetPixelShader();
+
+				SetGraphicsPipelineState(RHICmdList, GraphicsPSOInit, 0);
+
+				for (const FDraw& Draw : Draws)
+				{
+					PassParameters->VS.Transform = Draw.Transform;
+					PassParameters->PS.WaterBodyRenderDataIndex = Draw.WaterBodyRenderDataIndex;
+					PassParameters->PS.Priority = Draw.Priority;
+					PassParameters->PS.WaterBodyMinZ = Draw.MinZ;
+					PassParameters->PS.WaterBodyMaxZ = Draw.MaxZ;
+					PassParameters->PS.MaxWaveHeight = Draw.MaxWaveHeight;
+					PassParameters->PS.bIsRiver = Draw.bIsRiver;
+
+					SetShaderParameters(RHICmdList, VertexShader, VertexShader.GetVertexShader(), PassParameters->VS);
+					SetShaderParameters(RHICmdList, PixelShader, PixelShader.GetPixelShader(), PassParameters->PS);
+					RHICmdList.SetStreamSource(0, Draw.VertexBuffer, 0);
+					RHICmdList.DrawIndexedPrimitive(Draw.IndexBuffer, Draw.BaseVertexIndex, 0, Draw.NumVertices, Draw.FirstIndex, Draw.NumPrimitives, 1);
+				}
+			});
+	}
+
+		// Merge river and non-river water bodies and downsample to quad tree LOD0 resolution
+		{
+			FWaterQuadTreeMergePS::FPermutationDomain PermutationDomain;
+			PermutationDomain.Set<FWaterQuadTreeMergePS::FNumMSAASamples>(NumSamples);
+			TShaderMapRef<FWaterQuadTreeMergePS> PixelShader(ShaderMap, PermutationDomain);
+
+			FWaterQuadTreeMergePS::FParameters* PassParameters = GraphBuilder.AllocParameters<FWaterQuadTreeMergePS::FParameters>();
+			if (NumSamples > 1)
+			{
+				PassParameters->WaterBodyRasterTextureMS = WaterBodyRasterTexture;
+				PassParameters->ZBoundsRasterTextureMS = ZBoundsRasterTexture;
+			}
+			else
+			{
+				PassParameters->WaterBodyRasterTexture = WaterBodyRasterTexture;
+				PassParameters->ZBoundsRasterTexture = ZBoundsRasterTexture;
+			}
+			PassParameters->WaterBodyRenderData = WaterBodyRenderDataBufferSRV;
+			PassParameters->SuperSamplingFactor = Params.SuperSamplingFactor;
+			PassParameters->RcpCaptureDepthRange = 1.0f / Params.CaptureDepthRange;
+			PassParameters->ZBoundsPadding = GWaterQuadTreeZBoundsPadding / Params.CaptureDepthRange;
+			PassParameters->RenderTargets[0] = FRenderTargetBinding(QuadTreeTextureRDG, ERenderTargetLoadAction::ENoAction, 0);
+			PassParameters->RenderTargets[1] = FRenderTargetBinding(WaterZBoundsTextureRDG, ERenderTargetLoadAction::ENoAction, 0);
+
+			const FIntRect MergeViewport(0, 0, QuadTreeResolution.X, QuadTreeResolution.Y);
+
+			FPixelShaderUtils::AddFullscreenPass(
+				GraphBuilder,
+				ShaderMap,
+				RDG_EVENT_NAME("WaterQuadTreeMerge"),
+				PixelShader,
+				PassParameters,
+				MergeViewport);
+		}
+
+		// Build the mip chain
+		{
+			TShaderMapRef<FWaterQuadTreeBuildPS> PixelShader(ShaderMap);
+
+			FIntRect BuildViewport(0, 0, QuadTreeResolution.X, QuadTreeResolution.Y);
+
+			for (int32 MipLevel = 1; MipLevel < NumMipLevels; ++MipLevel)
+			{
+				FWaterQuadTreeBuildPS::FParameters* PassParameters = GraphBuilder.AllocParameters<FWaterQuadTreeBuildPS::FParameters>();
+				PassParameters->QuadTreeTexture = GraphBuilder.CreateSRV(FRDGTextureSRVDesc::CreateForMipLevel(QuadTreeTextureRDG, MipLevel - 1));
+				PassParameters->WaterZBoundsTexture = GraphBuilder.CreateSRV(FRDGTextureSRVDesc::CreateForMipLevel(WaterZBoundsTextureRDG, MipLevel - 1));
+				PassParameters->RenderTargets[0] = FRenderTargetBinding(QuadTreeTextureRDG, ERenderTargetLoadAction::ENoAction, MipLevel);
+				PassParameters->RenderTargets[1] = FRenderTargetBinding(WaterZBoundsTextureRDG, ERenderTargetLoadAction::ENoAction, MipLevel);
+
+				BuildViewport = BuildViewport / 2;
+
+				FPixelShaderUtils::AddFullscreenPass(
+					GraphBuilder,
+					ShaderMap,
+					RDG_EVENT_NAME("WaterQuadTreeBuild"),
+					PixelShader,
+					PassParameters,
+					BuildViewport);
+			}
+		}
+	
+}
+
+void FWaterQuadTreeGPU::Traverse(FRDGBuilder& GraphBuilder, const FTraverseParams& Params) const
+{
+	TRACE_CPUPROFILER_EVENT_SCOPE(FWaterQuadTreeGPU::Traverse);
+	RDG_EVENT_SCOPE(GraphBuilder, "FWaterQuadTreeGPU::Traverse");
+	RDG_GPU_STAT_SCOPE(GraphBuilder, FWaterQuadTreeGPU_Traverse);
+
+	FGlobalShaderMap* ShaderMap = GetGlobalShaderMap(GMaxRHIFeatureLevel);
+
+	FRDGTexture* QuadTreeTextureRDG = GraphBuilder.RegisterExternalTexture(CreateRenderTarget(QuadTreeTexture, TEXT("WaterQuadTree.QuadTree")));
+	FRDGTexture* WaterZBoundsTextureRDG = GraphBuilder.RegisterExternalTexture(CreateRenderTarget(WaterZBoundsTexture, TEXT("WaterQuadTree.WaterSurfaceHeight")));
+	FRDGBuffer* WaterBodyRenderDataBufferRDG = GraphBuilder.RegisterExternalBuffer(WaterBodyRenderDataBuffer);
+	FRDGBuffer* IndirectArgsBuffer = GraphBuilder.RegisterExternalBuffer(Params.OutIndirectArgsBuffer);
+	FRDGBuffer* InstanceDataOffsetsBuffer = GraphBuilder.RegisterExternalBuffer(Params.OutInstanceDataOffsetsBuffer);
+	FRDGBuffer* InstanceData0Buffer = GraphBuilder.RegisterExternalBuffer(Params.OutInstanceData0Buffer);
+	FRDGBuffer* InstanceData1Buffer = GraphBuilder.RegisterExternalBuffer(Params.OutInstanceData1Buffer);
+	FRDGBuffer* InstanceData2Buffer = Params.bWithWaterSelectionSupport ? 
+		GraphBuilder.RegisterExternalBuffer(Params.OutInstanceData2Buffer) : 
+		GraphBuilder.CreateBuffer(FRDGBufferDesc::CreateBufferDesc(sizeof(float) * 4, 1), TEXT("WaterQuadTree.InstanceData2Dummy"));
+	
+	FRDGBufferUAV* IndirectArgsBufferUAV = GraphBuilder.CreateUAV(FRDGBufferUAVDesc(IndirectArgsBuffer));
+	FRDGBufferUAV* InstanceDataOffsetsBufferUAV = GraphBuilder.CreateUAV(FRDGBufferUAVDesc(InstanceDataOffsetsBuffer, PF_R32_UINT));
+	FRDGBufferSRV* InstanceDataOffsetsBufferSRV = GraphBuilder.CreateSRV(FRDGBufferSRVDesc(InstanceDataOffsetsBuffer, PF_R32_UINT));
+	FRDGBufferSRV* WaterBodyRenderDataBufferSRV = GraphBuilder.CreateSRV(FRDGBufferSRVDesc(WaterBodyRenderDataBufferRDG));
+
+	const uint32 NumBucketsPerView = Params.NumDensities * Params.NumMaterials;
+	const uint32 NumBucketsTotal = NumBucketsPerView * Params.NumViews;
+	const FIntPoint QuadTreeResolution = QuadTreeTextureRDG->Desc.Extent;
+
+	// Initialize indirect args
+	{
+		TShaderMapRef<FWaterQuadTreeInitializeIndirectArgsCS> ComputeShader(ShaderMap);
+
+		FWaterQuadTreeInitializeIndirectArgsCS::FParameters* PassParameters = GraphBuilder.AllocParameters<FWaterQuadTreeInitializeIndirectArgsCS::FParameters>();
+		PassParameters->IndirectArgs = IndirectArgsBufferUAV;
+		PassParameters->NumDensities = Params.NumDensities;
+		PassParameters->NumMaterials = Params.NumMaterials;
+		PassParameters->NumViews = Params.NumViews;
+		PassParameters->NumQuadsLOD0 = Params.NumQuadsLOD0;
+		
+		FComputeShaderUtils::AddPass(GraphBuilder, RDG_EVENT_NAME("WaterQuadTreeInitIndirectArgs"), ComputeShader, PassParameters, FComputeShaderUtils::GetGroupCount(NumBucketsTotal, 64));
+	}
+
+	// Iterate over all views for which to create indirect water draws
+	for (uint32 ViewIndex = 0; ViewIndex < Params.NumViews; ++ViewIndex)
+	{
+		const FPerViewInfo& PerViewInfo = Params.PerViewInfo[ViewIndex];
+
+		const uint32 MaxNumDraws = QuadTreeResolution.X * QuadTreeResolution.Y;
+		FRDGBuffer* PackedNodes = GraphBuilder.CreateBuffer(FRDGBufferDesc::CreateByteAddressDesc((1 + MaxNumDraws) * sizeof(uint32)), TEXT("WaterQuadTree.PackedNodes"));
+		FRDGBuffer* BucketCounts = GraphBuilder.CreateBuffer(FRDGBufferDesc::CreateByteAddressDesc(FMath::Max(1u, NumBucketsPerView) * sizeof(uint32)), TEXT("WaterQuadTree.DrawBucketCounts"));
+		FRDGBufferUAV* PackedNodesUAV = GraphBuilder.CreateUAV(FRDGBufferUAVDesc(PackedNodes, PF_R32_UINT));
+		FRDGBufferUAV* BucketCountsUAV = GraphBuilder.CreateUAV(FRDGBufferUAVDesc(BucketCounts, PF_R32_UINT));
+		FRDGBufferSRV* PackedNodesSRV = GraphBuilder.CreateSRV(FRDGBufferSRVDesc(PackedNodes, PF_R32_UINT));
+		FRDGBufferSRV* BucketCountsSRV = GraphBuilder.CreateSRV(FRDGBufferSRVDesc(BucketCounts, PF_R32_UINT));
+
+		// Clear counts and packed nodes counter
+		// TODO: only clear the counter at index 0, not the entire buffer
+		FComputeShaderUtils::ClearUAV(GraphBuilder, ShaderMap, BucketCountsUAV, 0);
+		FComputeShaderUtils::ClearUAV(GraphBuilder, ShaderMap, PackedNodesUAV, 0);
+
+		// Traverse quadtree
+		{
+			if (Params.DebugShowTile != 0)
+			{
+				ShaderPrint::SetEnabled(true);
+			}
+
+			const bool bEnableShaderPrint = Params.DebugShowTile != 0
+				&& PerViewInfo.ViewInfo
+				&& ShaderPrint::IsEnabled()
+				&& ShaderPrint::IsEnabled(PerViewInfo.ViewInfo->ShaderPrintData)
+				&& ShaderPrint::IsValid(PerViewInfo.ViewInfo->ShaderPrintData)
+				&& ShaderPrint::IsSupported(GetFeatureLevelShaderPlatform(GMaxRHIFeatureLevel));
+
+			if (bEnableShaderPrint)
+			{
+				// We'll be potentially drawing a lot of bounding boxes, each made up of 12 line segments, so reserve some space
+				ShaderPrint::RequestSpaceForLines(12 * MaxNumDraws);
+			}
+
+			FWaterQuadTreeTraverseCS::FPermutationDomain PermutationDomain;
+			PermutationDomain.Set<FWaterQuadTreeTraverseCS::FDebugShaderPrint>(bEnableShaderPrint);
+			TShaderMapRef<FWaterQuadTreeTraverseCS> ComputeShader(ShaderMap, PermutationDomain);
+
+			const uint32 NumMipLevels = QuadTreeTextureRDG->Desc.NumMips;
+			uint32 NumTexelsTotal = 0;
+			for (uint32 MipLevelIndex = 0; MipLevelIndex < NumMipLevels; ++MipLevelIndex)
+			{
+				const FIntPoint MipExtent = FIntPoint(FMath::Max(1, QuadTreeTextureRDG->Desc.Extent.X >> MipLevelIndex), FMath::Max(1, QuadTreeTextureRDG->Desc.Extent.Y >> MipLevelIndex));
+				NumTexelsTotal += MipExtent.X * MipExtent.Y;
+			}
+
+			FWaterQuadTreeTraverseCS::FParameters* PassParameters = GraphBuilder.AllocParameters<FWaterQuadTreeTraverseCS::FParameters>();
+			if (bEnableShaderPrint)
+			{
+				ShaderPrint::SetParameters(GraphBuilder, PerViewInfo.ViewInfo->ShaderPrintData, PassParameters->ShaderPrintUniformBuffer);
+			}
+			PassParameters->PackedNodes = PackedNodesUAV;
+			PassParameters->QuadTreeTexture = QuadTreeTextureRDG;
+			PassParameters->WaterZBoundsTexture = WaterZBoundsTextureRDG;
+			PassParameters->WaterBodyRenderData = WaterBodyRenderDataBufferSRV;
+			PassParameters->HZBTexture = (PerViewInfo.ViewInfo && PerViewInfo.ViewInfo->HZB) ? PerViewInfo.ViewInfo->HZB : GSystemTextures.GetBlackDummy(GraphBuilder);
+			PassParameters->HZBSampler = TStaticSamplerState<SF_Point, AM_Clamp, AM_Clamp, AM_Clamp>::GetRHI();
+			PassParameters->TranslatedWorldToClip = PerViewInfo.TranslatedWorldToClip;
+			PassParameters->ViewToClip = PerViewInfo.ViewToClip;
+			PassParameters->CullingBoundsAABB = FVector4f(Params.CullingBounds.Min.X, Params.CullingBounds.Min.Y, Params.CullingBounds.Max.X, Params.CullingBounds.Max.Y);
+			PassParameters->HZBSize = PerViewInfo.ViewInfo ? FVector2f(PerViewInfo.ViewInfo->ViewRect.Size()) : FVector2f::ZeroVector;
+			PassParameters->HZBViewSize = PerViewInfo.ViewInfo ? FVector2f(PerViewInfo.ViewInfo->HZBMipmap0Size) : FVector2f::ZeroVector;
+			PassParameters->QuadTreePosition = PerViewInfo.QuadTreePositionTranslatedWorldSpace;
+			PassParameters->ObserverPosition = PerViewInfo.ObserverPositionTranslatedWorldSpace;
+			PassParameters->QuadTreeResolutionX = QuadTreeResolution.X;
+			PassParameters->QuadTreeResolutionY = QuadTreeResolution.Y;
+			PassParameters->NumDensities = Params.NumDensities;
+			PassParameters->NumMaterials = Params.NumMaterials;
+			PassParameters->LeafSize = Params.LeafSize;
+			PassParameters->LODScale = Params.LODScale;
+			PassParameters->CaptureDepthRange = CaptureDepthRange;
+			PassParameters->ForceCollapseDensityLevel = -1; // Params.ForceCollapseDensityLevel; // TODO: Properly implement support for this parameter
+			PassParameters->NumLODs = NumMipLevels;
+			PassParameters->NumDispatchedThreads = NumTexelsTotal;
+			PassParameters->DebugShowTile = Params.DebugShowTile;
+			PassParameters->bHZBOcclusionCullingEnabled = GWaterQuadTreeOcclusionCulling && (PerViewInfo.ViewInfo && PerViewInfo.ViewInfo->HZB);
+
+			FComputeShaderUtils::AddPass(GraphBuilder, RDG_EVENT_NAME("WaterQuadTreeTraverse(View: %i)", ViewIndex), ComputeShader, PassParameters, FComputeShaderUtils::GetGroupCount(NumTexelsTotal, 64));
+		}
+
+		// Compute bucket counts
+		{
+			TShaderMapRef<FWaterQuadTreeBucketCountsCS> ComputeShader(ShaderMap);
+
+			FWaterQuadTreeBucketCountsCS::FParameters* PassParameters = GraphBuilder.AllocParameters<FWaterQuadTreeBucketCountsCS::FParameters>();
+			PassParameters->BucketCounts = BucketCountsUAV;
+			PassParameters->QuadTreeTexture = QuadTreeTextureRDG;
+			PassParameters->WaterBodyRenderData = WaterBodyRenderDataBufferSRV;
+			PassParameters->PackedNodes = PackedNodesSRV;
+			PassParameters->NumDensities = Params.NumDensities;
+			PassParameters->NumMaterials = Params.NumMaterials;
+			PassParameters->NumDispatchedThreads = FMath::Min(FMath::DivideAndRoundUp(MaxNumDraws, 64u), 65535u) * 64u; // The maximum number of dispatched groups is 64k
+
+			FComputeShaderUtils::AddPass(GraphBuilder, RDG_EVENT_NAME("WaterQuadTreeDrawBucketCounts(View: %i)", ViewIndex), ComputeShader, PassParameters, FComputeShaderUtils::GetGroupCount(PassParameters->NumDispatchedThreads, 64));
+		}
+
+		// Compute bucket prefix sums
+		{
+			FWaterQuadTreeBucketPrefixSumCS::FPermutationDomain PermutationVector;
+			PermutationVector.Set<FWaterQuadTreeBucketPrefixSumCS::FParallelPrefixSum>(GWaterQuadTreeParallelPrefixSum != 0);
+			TShaderMapRef<FWaterQuadTreeBucketPrefixSumCS> ComputeShader(ShaderMap, PermutationVector);
+
+			FWaterQuadTreeBucketPrefixSumCS::FParameters* PassParameters = GraphBuilder.AllocParameters<FWaterQuadTreeBucketPrefixSumCS::FParameters>();
+			PassParameters->BucketPrefixSums = InstanceDataOffsetsBufferUAV;
+			PassParameters->BucketCounts = GraphBuilder.CreateSRV(FRDGBufferSRVDesc(BucketCounts));
+			PassParameters->NumBuckets = NumBucketsPerView;
+			PassParameters->OutputOffset = ViewIndex * NumBucketsPerView;
+
+			FComputeShaderUtils::AddPass(GraphBuilder, RDG_EVENT_NAME("WaterQuadTreeDrawBucketPrefixSums(View: %i)", ViewIndex), ComputeShader, PassParameters, FIntVector3(1, 1, 1));
+		}
+
+		// Generate instance data and fill indirect args
+		{
+			TShaderMapRef<FWaterQuadTreeGenerateInstanceDataCS> ComputeShader(ShaderMap);
+
+			FWaterQuadTreeGenerateInstanceDataCS::FParameters* PassParameters = GraphBuilder.AllocParameters<FWaterQuadTreeGenerateInstanceDataCS::FParameters>();
+			PassParameters->IndirectArgs = IndirectArgsBufferUAV;
+			PassParameters->InstanceData0 = GraphBuilder.CreateUAV(FRDGBufferUAVDesc(InstanceData0Buffer, PF_A32B32G32R32F));
+			PassParameters->InstanceData1 = GraphBuilder.CreateUAV(FRDGBufferUAVDesc(InstanceData1Buffer, PF_A32B32G32R32F));
+			PassParameters->InstanceData2 = GraphBuilder.CreateUAV(FRDGBufferUAVDesc(InstanceData2Buffer, PF_A32B32G32R32F));
+			PassParameters->QuadTreeTexture = QuadTreeTextureRDG;
+			PassParameters->WaterZBoundsTexture = WaterZBoundsTextureRDG;
+			PassParameters->WaterBodyRenderData = WaterBodyRenderDataBufferSRV;
+			PassParameters->PackedNodes = PackedNodesSRV;
+			PassParameters->InstanceDataOffsets = InstanceDataOffsetsBufferSRV;
+			PassParameters->QuadTreePosition = PerViewInfo.QuadTreePositionTranslatedWorldSpace;
+			PassParameters->ObserverPosition = PerViewInfo.ObserverPositionTranslatedWorldSpace;
+			PassParameters->QuadTreeResolutionX = QuadTreeResolution.X;
+			PassParameters->QuadTreeResolutionY = QuadTreeResolution.Y;
+			PassParameters->NumDensities = Params.NumDensities;
+			PassParameters->NumMaterials = Params.NumMaterials;
+			PassParameters->NumDispatchedThreads = FMath::Min(FMath::DivideAndRoundUp(MaxNumDraws, 64u), 65535u) * 64u; // The maximum number of dispatched groups is 64k
+			PassParameters->BucketIndexOffset = ViewIndex * NumBucketsPerView;
+			PassParameters->NumLODs = QuadTreeTextureRDG->Desc.NumMips;
+			PassParameters->LeafSize = Params.LeafSize;
+			PassParameters->LODScale = Params.LODScale;
+			PassParameters->CaptureDepthRange = CaptureDepthRange;
+			PassParameters->bWithWaterSelectionSupport = Params.bWithWaterSelectionSupport;
+			PassParameters->bLODMorphingEnabled = Params.bLODMorphingEnabled;
+
+			FComputeShaderUtils::AddPass(GraphBuilder, RDG_EVENT_NAME("WaterQuadTreeGenerateDraws(View: %i)", ViewIndex), ComputeShader, PassParameters, FComputeShaderUtils::GetGroupCount(PassParameters->NumDispatchedThreads, 64));
+		}
+	}
+}
diff --git a/Engine/Plugins/Experimental/Water/Source/Runtime/Private/WaterQuadTreeGPU.h b/Engine/Plugins/Experimental/Water/Source/Runtime/Private/WaterQuadTreeGPU.h
new file mode 100644
--- /dev/null
+++ b/Engine/Plugins/Experimental/Water/Source/Runtime/Private/WaterQuadTreeGPU.h
@@ -0,0 +1,104 @@
+// Copyright Epic Games, Inc. All Rights Reserved.
+#pragma once
+#include "CoreMinimal.h"
+
+class FRHIBuffer;
+class FRHITexture;
+class FRDGBuilder;
+class FRDGTexture;
+class FRDGPooledBuffer;
+class FViewInfo;
+struct FShaderPrintData;
+
+class FWaterQuadTreeGPU
+{
+public:
+	struct FDraw
+	{
+		FMatrix44f Transform = FMatrix44f::Identity;
+		FRHIBuffer* VertexBuffer = nullptr;
+		FRHIBuffer* IndexBuffer = nullptr;
+		uint32 FirstIndex = 0;
+		uint32 NumPrimitives = 0;
+		uint32 BaseVertexIndex = 0;
+		uint32 NumVertices = 0;
+		uint32 WaterBodyRenderDataIndex = INDEX_NONE;
+		int32 Priority = 0;
+		float MinZ = 1.0f; // Min Z of the water body bounding box scaled into 0-1 depth buffer space
+		float MaxZ = 0.0f; // Max Z of the water body bounding box scaled into 0-1 depth buffer space
+		float MaxWaveHeight = 0.0f; // Scaled into 0-1 depth buffer space
+		bool bIsRiver = false;
+	};
+
+	struct FWaterBodyRenderDataGPU
+	{
+		uint32_t WaterBodyIndex = INDEX_NONE;
+		uint32_t MaterialIndex = INDEX_NONE;
+		uint32_t RiverToLakeMaterialIndex = INDEX_NONE;
+		uint32_t RiverToOceanMaterialIndex = INDEX_NONE;
+		uint32_t WaterBodyType = INDEX_NONE;
+		uint32_t HitProxyColorAndIsSelected = 0;
+		float SurfaceBaseHeight = 0.0f;
+		float MinZ = 1.0f;
+		float MaxZ = 0.0f;
+		float MaxWaveHeight = 0.0f;
+	};
+
+	struct FInitParams
+	{
+		TArrayView<FWaterBodyRenderDataGPU> WaterBodyRenderData;
+		FIntPoint RequestedQuadTreeResolution = FIntPoint::ZeroValue;
+		int32 SuperSamplingFactor = 1;
+		int32 NumMSAASamples = 4;
+		float CaptureDepthRange = 1.0f;
+	};
+
+	struct FPerViewInfo
+	{
+		FMatrix44f TranslatedWorldToClip = FMatrix44f::Identity;
+		FMatrix44f ViewToClip = FMatrix44f::Identity;
+		FVector3f QuadTreePositionTranslatedWorldSpace = FVector3f::ZeroVector;
+		FVector3f ObserverPositionTranslatedWorldSpace = FVector3f::ZeroVector;
+		const FViewInfo* ViewInfo = nullptr;
+	};
+
+	struct FTraverseParams
+	{
+		FRDGPooledBuffer* OutIndirectArgsBuffer = nullptr;
+		FRDGPooledBuffer* OutInstanceDataOffsetsBuffer = nullptr;
+		FRDGPooledBuffer* OutInstanceData0Buffer = nullptr;
+		FRDGPooledBuffer* OutInstanceData1Buffer = nullptr;
+		FRDGPooledBuffer* OutInstanceData2Buffer = nullptr;
+		TArray<FPerViewInfo> PerViewInfo;
+		FBox2D CullingBounds = FBox2D(ForceInit);
+		uint32 NumDensities = 0;
+		uint32 NumMaterials = 0;
+		uint32 NumViews = 0;
+		uint32 NumQuadsLOD0 = 0;
+		int32 ForceCollapseDensityLevel = -1;
+		float LeafSize = 1.0f;
+		float LODScale = 1.0f;
+		int32 DebugShowTile = 0;
+		bool bWithWaterSelectionSupport = false;
+		bool bLODMorphingEnabled = false;
+	};
+
+	void Init(FRDGBuilder& GraphBuilder, const FInitParams& Params, TArray<FDraw>& Draws);
+
+	void Traverse(FRDGBuilder& GraphBuilder, const FTraverseParams& Params) const;
+
+	bool IsInitialized() const { return bInitialized; }
+
+private:
+	/** Mipped texture representing the water quadtree on the GPU */
+	TRefCountPtr<FRHITexture> QuadTreeTexture;
+
+	/** Mipped texture representing the Z bounds of each node of the water quadtree on the GPU */
+	TRefCountPtr<FRHITexture> WaterZBoundsTexture;
+
+	/** GPU representation of the water quad trees' array of water body render data */
+	TRefCountPtr<FRDGPooledBuffer> WaterBodyRenderDataBuffer;
+
+	float CaptureDepthRange = 0.0f;
+	bool bInitialized = false;
+};
\ No newline at end of file
diff --git a/Engine/Plugins/Experimental/Water/Source/Runtime/Private/WaterVertexFactory.cpp b/Engine/Plugins/Experimental/Water/Source/Runtime/Private/WaterVertexFactory.cpp
--- a/Engine/Plugins/Experimental/Water/Source/Runtime/Private/WaterVertexFactory.cpp
+++ b/Engine/Plugins/Experimental/Water/Source/Runtime/Private/WaterVertexFactory.cpp
@@ -13,18 +13,31 @@ IMPLEMENT_GLOBAL_SHADER_PARAMETER_STRUCT(FWaterVertexFactoryRaytracingParameters
 /**
  * Shader parameters for water vertex factory.
  */
-template <bool bWithWaterSelectionSupport>
+template <bool bWithWaterSelectionSupport, bool bIndirectDraws>
 class TWaterVertexFactoryShaderParameters : public FVertexFactoryShaderParameters
 {
-	DECLARE_TYPE_LAYOUT(TWaterVertexFactoryShaderParameters<bWithWaterSelectionSupport>, NonVirtual);
+	using WaterVertexFactoryShaderParametersType = TWaterVertexFactoryShaderParameters<bWithWaterSelectionSupport, bIndirectDraws>;
+	DECLARE_TYPE_LAYOUT(WaterVertexFactoryShaderParametersType, NonVirtual);
 
 public:
-	using WaterVertexFactoryType = TWaterVertexFactory<bWithWaterSelectionSupport>;
+	using WaterVertexFactoryType = TWaterVertexFactory<bWithWaterSelectionSupport, bIndirectDraws>;
 	using WaterMeshUserDataType = TWaterMeshUserData<bWithWaterSelectionSupport>;
 	using WaterInstanceDataBuffersType = TWaterInstanceDataBuffers<bWithWaterSelectionSupport>;
 
 	void Bind(const FShaderParameterMap& ParameterMap)
 	{
+		if (bIndirectDraws)
+		{
+			InstanceDataOffsetsBufferParameter.Bind(ParameterMap, TEXT("InstanceDataOffsetsBuffer"));
+			InstanceData0BufferParameter.Bind(ParameterMap, TEXT("InstanceData0Buffer"));
+			InstanceData1BufferParameter.Bind(ParameterMap, TEXT("InstanceData1Buffer"));
+			if (bWithWaterSelectionSupport)
+			{
+				InstanceData2BufferParameter.Bind(ParameterMap, TEXT("InstanceData2Buffer"));
+			}
+			
+			DrawBucketIndexParameter.Bind(ParameterMap, TEXT("DrawBucketIndex"));
+		}
 	}
 
 	void GetElementShaderBindings(
@@ -44,8 +57,6 @@ public:
 
 		const WaterInstanceDataBuffersType* InstanceDataBuffers = WaterMeshUserData->InstanceDataBuffers;
 
-		const int32 InstanceOffsetValue = BatchElement.UserIndex;
-
 		ShaderBindings.Add(Shader->GetUniformBufferParameter<FWaterVertexFactoryParameters>(), VertexFactory->GetWaterVertexFactoryUniformBuffer(WaterMeshUserData->RenderGroupType));
 
 #if RHI_RAYTRACING
@@ -55,7 +66,7 @@ public:
 		}
 #endif
 
-		if (VertexStreams.Num() > 0)
+		if (VertexStreams.Num() > 0 && !bIndirectDraws)
 		{
 			for (int32 i = 0; i < WaterInstanceDataBuffersType::NumBuffers; ++i)
 			{
@@ -67,24 +78,52 @@ public:
 				InstanceInputStream->VertexBuffer = InstanceDataBuffers->GetBuffer(i);
 			}
 
+			const int32 InstanceOffsetValue = BatchElement.UserIndex;
 			if (InstanceOffsetValue > 0)
 			{
 				VertexFactory->OffsetInstanceStreams(InstanceOffsetValue, InputStreamType, VertexStreams);
 			}
 		}
+
+		if (bIndirectDraws)
+		{
+			const uint32 DrawBucketIndex = BatchElement.UserIndex;
+			ShaderBindings.Add(InstanceDataOffsetsBufferParameter, WaterMeshUserData->IndirectInstanceDataOffsets);
+			ShaderBindings.Add(InstanceData0BufferParameter, WaterMeshUserData->IndirectInstanceData0);
+			ShaderBindings.Add(InstanceData1BufferParameter, WaterMeshUserData->IndirectInstanceData1);
+			if (bWithWaterSelectionSupport)
+			{
+				ShaderBindings.Add(InstanceData2BufferParameter, WaterMeshUserData->IndirectInstanceData2);
+			}
+			ShaderBindings.Add(DrawBucketIndexParameter, DrawBucketIndex);
+		}
 	}
+
+private:
+	LAYOUT_FIELD(FShaderResourceParameter, InstanceDataOffsetsBufferParameter);
+	LAYOUT_FIELD(FShaderResourceParameter, InstanceData0BufferParameter);
+	LAYOUT_FIELD(FShaderResourceParameter, InstanceData1BufferParameter);
+	LAYOUT_FIELD(FShaderResourceParameter, InstanceData2BufferParameter);
+	LAYOUT_FIELD(FShaderParameter, DrawBucketIndexParameter);
 };
 
 // ----------------------------------------------------------------------------------
 
+using FWaterVertexFactoryParametersNoSelectionNoIndirect = TWaterVertexFactoryShaderParameters</*bWithWaterSelectionSupport = */ false, /*bIndirectDraws = */ false>;
+using FWaterVertexFactoryParametersNoSelectionWithIndirect = TWaterVertexFactoryShaderParameters</*bWithWaterSelectionSupport = */ false, /*bIndirectDraws = */ true>;
+using FWaterVertexFactoryNoSelectionNoIndirect = TWaterVertexFactory</*bWithWaterSelectionSupport = */ false, /*bIndirectDraws = */ false>;
+using FWaterVertexFactoryNoSelectionWithIndirect = TWaterVertexFactory</*bWithWaterSelectionSupport = */ false, /*bIndirectDraws = */ true>;
+
 // Always implement the basic vertex factory so that it's there for both editor and non-editor builds :
-IMPLEMENT_TEMPLATE_TYPE_LAYOUT(template<>, TWaterVertexFactoryShaderParameters</*bWithWaterSelectionSupport = */ false>);
-IMPLEMENT_VERTEX_FACTORY_PARAMETER_TYPE(TWaterVertexFactory</*bWithWaterSelectionSupport = */ false>, SF_Vertex, TWaterVertexFactoryShaderParameters</*bWithWaterSelectionSupport = */ false>);
+IMPLEMENT_TEMPLATE_TYPE_LAYOUT(template<>, FWaterVertexFactoryParametersNoSelectionNoIndirect);
+IMPLEMENT_TEMPLATE_TYPE_LAYOUT(template<>, FWaterVertexFactoryParametersNoSelectionWithIndirect);
+IMPLEMENT_VERTEX_FACTORY_PARAMETER_TYPE(FWaterVertexFactoryNoSelectionNoIndirect, SF_Vertex, FWaterVertexFactoryParametersNoSelectionNoIndirect);
+IMPLEMENT_VERTEX_FACTORY_PARAMETER_TYPE(FWaterVertexFactoryNoSelectionWithIndirect, SF_Vertex, FWaterVertexFactoryParametersNoSelectionWithIndirect);
 #if RHI_RAYTRACING
-IMPLEMENT_VERTEX_FACTORY_PARAMETER_TYPE(TWaterVertexFactory</*bWithWaterSelectionSupport = */ false>, SF_Compute, TWaterVertexFactoryShaderParameters</*bWithWaterSelectionSupport = */ false>);
-IMPLEMENT_VERTEX_FACTORY_PARAMETER_TYPE(TWaterVertexFactory</*bWithWaterSelectionSupport = */ false>, SF_RayHitGroup, TWaterVertexFactoryShaderParameters</*bWithWaterSelectionSupport = */ false>);
+IMPLEMENT_VERTEX_FACTORY_PARAMETER_TYPE(FWaterVertexFactoryNoSelectionNoIndirect, SF_Compute, FWaterVertexFactoryParametersNoSelectionNoIndirect);
+IMPLEMENT_VERTEX_FACTORY_PARAMETER_TYPE(FWaterVertexFactoryNoSelectionNoIndirect, SF_RayHitGroup, FWaterVertexFactoryParametersNoSelectionNoIndirect);
 #endif // RHI_RAYTRACING
-IMPLEMENT_TEMPLATE_VERTEX_FACTORY_TYPE(template<>, TWaterVertexFactory</*bWithWaterSelectionSupport = */ false>, "/Plugin/Water/Private/WaterMeshVertexFactory.ush",
+IMPLEMENT_TEMPLATE_VERTEX_FACTORY_TYPE(template<>, FWaterVertexFactoryNoSelectionNoIndirect, "/Plugin/Water/Private/WaterMeshVertexFactory.ush",
 	  EVertexFactoryFlags::UsedWithMaterials
 	| EVertexFactoryFlags::SupportsDynamicLighting
 	| EVertexFactoryFlags::SupportsPrecisePrevWorldPos
@@ -93,17 +132,33 @@ IMPLEMENT_TEMPLATE_VERTEX_FACTORY_TYPE(template<>, TWaterVertexFactory</*bWithWa
 	| EVertexFactoryFlags::SupportsRayTracingDynamicGeometry
 	| EVertexFactoryFlags::SupportsPSOPrecaching
 );
+IMPLEMENT_TEMPLATE_VERTEX_FACTORY_TYPE(template<>, FWaterVertexFactoryNoSelectionWithIndirect, "/Plugin/Water/Private/WaterMeshVertexFactory.ush",
+	EVertexFactoryFlags::UsedWithMaterials
+	| EVertexFactoryFlags::SupportsDynamicLighting
+	| EVertexFactoryFlags::SupportsPrecisePrevWorldPos
+	| EVertexFactoryFlags::SupportsPrimitiveIdStream
+	| EVertexFactoryFlags::SupportsRayTracing
+	| EVertexFactoryFlags::SupportsRayTracingDynamicGeometry
+	| EVertexFactoryFlags::SupportsPSOPrecaching
+);
 
 #if WITH_WATER_SELECTION_SUPPORT
 
+using FWaterVertexFactoryParametersWithSelectionNoIndirect = TWaterVertexFactoryShaderParameters</*bWithWaterSelectionSupport = */ true, /*bIndirectDraws = */ false>;
+using FWaterVertexFactoryParametersWithSelectionWithIndirect = TWaterVertexFactoryShaderParameters</*bWithWaterSelectionSupport = */ true, /*bIndirectDraws = */ true>;
+using FWaterVertexFactoryWithSelectionNoIndirect = TWaterVertexFactory</*bWithWaterSelectionSupport = */ true, /*bIndirectDraws = */ false>;
+using FWaterVertexFactoryWithSelectionWithIndirect = TWaterVertexFactory</*bWithWaterSelectionSupport = */ true, /*bIndirectDraws = */ true>;
+
 // In editor builds, also implement the vertex factory that supports water selection:
-IMPLEMENT_TEMPLATE_TYPE_LAYOUT(template<>, TWaterVertexFactoryShaderParameters</*bWithWaterSelectionSupport = */ true>);
-IMPLEMENT_VERTEX_FACTORY_PARAMETER_TYPE(TWaterVertexFactory</*bWithWaterSelectionSupport = */ true>, SF_Vertex, TWaterVertexFactoryShaderParameters</*bWithWaterSelectionSupport = */ true>);
+IMPLEMENT_TEMPLATE_TYPE_LAYOUT(template<>, FWaterVertexFactoryParametersWithSelectionNoIndirect);
+IMPLEMENT_TEMPLATE_TYPE_LAYOUT(template<>, FWaterVertexFactoryParametersWithSelectionWithIndirect);
+IMPLEMENT_VERTEX_FACTORY_PARAMETER_TYPE(FWaterVertexFactoryWithSelectionNoIndirect, SF_Vertex, FWaterVertexFactoryParametersWithSelectionNoIndirect);
+IMPLEMENT_VERTEX_FACTORY_PARAMETER_TYPE(FWaterVertexFactoryWithSelectionWithIndirect, SF_Vertex, FWaterVertexFactoryParametersWithSelectionWithIndirect);
 #if RHI_RAYTRACING
-IMPLEMENT_VERTEX_FACTORY_PARAMETER_TYPE(TWaterVertexFactory</*bWithWaterSelectionSupport = */ true>, SF_Compute, TWaterVertexFactoryShaderParameters</*bWithWaterSelectionSupport = */ true>);
-IMPLEMENT_VERTEX_FACTORY_PARAMETER_TYPE(TWaterVertexFactory</*bWithWaterSelectionSupport = */ true>, SF_RayHitGroup, TWaterVertexFactoryShaderParameters</*bWithWaterSelectionSupport = */ true>);
+IMPLEMENT_VERTEX_FACTORY_PARAMETER_TYPE(FWaterVertexFactoryWithSelectionNoIndirect, SF_Compute, FWaterVertexFactoryParametersWithSelectionNoIndirect);
+IMPLEMENT_VERTEX_FACTORY_PARAMETER_TYPE(FWaterVertexFactoryWithSelectionNoIndirect, SF_RayHitGroup, FWaterVertexFactoryParametersWithSelectionNoIndirect);
 #endif
-IMPLEMENT_TEMPLATE_VERTEX_FACTORY_TYPE(template<>, TWaterVertexFactory</*bWithWaterSelectionSupport = */ true>, "/Plugin/Water/Private/WaterMeshVertexFactory.ush",
+IMPLEMENT_TEMPLATE_VERTEX_FACTORY_TYPE(template<>, FWaterVertexFactoryWithSelectionNoIndirect, "/Plugin/Water/Private/WaterMeshVertexFactory.ush",
 	  EVertexFactoryFlags::UsedWithMaterials
 	| EVertexFactoryFlags::SupportsDynamicLighting
 	| EVertexFactoryFlags::SupportsPrecisePrevWorldPos
@@ -112,20 +167,44 @@ IMPLEMENT_TEMPLATE_VERTEX_FACTORY_TYPE(template<>, TWaterVertexFactory</*bWithWa
 	| EVertexFactoryFlags::SupportsRayTracingDynamicGeometry
 	| EVertexFactoryFlags::SupportsPSOPrecaching
 );
+IMPLEMENT_TEMPLATE_VERTEX_FACTORY_TYPE(template<>, FWaterVertexFactoryWithSelectionWithIndirect, "/Plugin/Water/Private/WaterMeshVertexFactory.ush",
+	EVertexFactoryFlags::UsedWithMaterials
+	| EVertexFactoryFlags::SupportsDynamicLighting
+	| EVertexFactoryFlags::SupportsPrecisePrevWorldPos
+	| EVertexFactoryFlags::SupportsPrimitiveIdStream
+	| EVertexFactoryFlags::SupportsRayTracing
+	| EVertexFactoryFlags::SupportsRayTracingDynamicGeometry
+	| EVertexFactoryFlags::SupportsPSOPrecaching
+);
 
 #endif // WITH_WATER_SELECTION_SUPPORT
 
-const FVertexFactoryType* GetWaterVertexFactoryType(bool bWithWaterSelectionSupport)
+const FVertexFactoryType* GetWaterVertexFactoryType(bool bWithWaterSelectionSupport, bool bIndirectDraws)
 {
 #if WITH_WATER_SELECTION_SUPPORT
 	if (bWithWaterSelectionSupport)
 	{
-		return &TWaterVertexFactory</*bWithWaterSelectionSupport = */ true>::StaticType;
+		if (bIndirectDraws)
+		{
+			return &TWaterVertexFactory</*bWithWaterSelectionSupport = */ true, /*bIndirectDraws = */ true>::StaticType;
+		}
+		else
+		{
+			return &TWaterVertexFactory</*bWithWaterSelectionSupport = */ true, /*bIndirectDraws = */ false>::StaticType;
+		}
 	}
 	else
 #endif // WITH_WATER_SELECTION_SUPPORT
 	{
 		check(!bWithWaterSelectionSupport);
-		return &TWaterVertexFactory</*bWithWaterSelectionSupport = */ false>::StaticType;
+		if (bIndirectDraws)
+		{
+			return &TWaterVertexFactory</*bWithWaterSelectionSupport = */ false, /*bIndirectDraws = */ true>::StaticType;
+		}
+		else
+		{
+			return &TWaterVertexFactory</*bWithWaterSelectionSupport = */ false, /*bIndirectDraws = */ false>::StaticType;
+		}
+		
 	}
 }
diff --git a/Engine/Plugins/Experimental/Water/Source/Runtime/Private/WaterViewExtension.cpp b/Engine/Plugins/Experimental/Water/Source/Runtime/Private/WaterViewExtension.cpp
--- a/Engine/Plugins/Experimental/Water/Source/Runtime/Private/WaterViewExtension.cpp
+++ b/Engine/Plugins/Experimental/Water/Source/Runtime/Private/WaterViewExtension.cpp
@@ -33,6 +33,8 @@ static TAutoConsoleVariable<int32> CVarWaterInfoRenderMethod(
 
 // ----------------------------------------------------------------------------------
 
+FWaterMeshGPUWork GWaterMeshGPUWork;
+
 FWaterViewExtension::FWaterViewExtension(const FAutoRegister& AutoReg, UWorld* InWorld)
 	: FWorldSceneViewExtension(AutoReg, InWorld)
 	, WaterGPUData(MakeShared<FWaterGPUResources, ESPMode::ThreadSafe>())
@@ -410,6 +412,14 @@ void FWaterViewExtension::PreRenderView_RenderThread(FRDGBuilder& GraphBuilder,
 	}
 }
 
+void FWaterViewExtension::PreRenderBasePass_RenderThread(FRDGBuilder& GraphBuilder)
+{
+	for (FWaterMeshGPUWork::FCallback& Callback : GWaterMeshGPUWork.Callbacks)
+	{
+		Callback.Function(GraphBuilder);
+	}
+}
+
 void FWaterViewExtension::MarkWaterInfoTextureForRebuild(const UE::WaterInfo::FRenderingContext& RenderContext)
 {
 	WaterInfoContextsToRender.Emplace(RenderContext.ZoneToRender, RenderContext);
diff --git a/Engine/Plugins/Experimental/Water/Source/Runtime/Public/WaterMeshSceneProxy.h b/Engine/Plugins/Experimental/Water/Source/Runtime/Public/WaterMeshSceneProxy.h
--- a/Engine/Plugins/Experimental/Water/Source/Runtime/Public/WaterMeshSceneProxy.h
+++ b/Engine/Plugins/Experimental/Water/Source/Runtime/Public/WaterMeshSceneProxy.h
@@ -8,6 +8,7 @@
 #include "WaterQuadTree.h"
 #include "WaterVertexFactory.h"
 #include "RayTracingGeometry.h"
+#include "WaterQuadTreeGPU.h"
 
 class FMeshElementCollector;
 struct FRayTracingMaterialGatheringContext;
@@ -31,6 +32,8 @@ public:
 
 	virtual void CreateRenderThreadResources(FRHICommandListBase& RHICmdList) override;
 
+	virtual void DestroyRenderThreadResources() override;
+
 	virtual void GetDynamicMeshElements(const TArray<const FSceneView*>& Views, const FSceneViewFamily& ViewFamily, uint32 VisibilityMap, FMeshElementCollector& Collector) const override;
 
 	virtual const TArray<FBoxSphereBounds>* GetOcclusionQueries(const FSceneView* View) const override;
@@ -53,7 +56,7 @@ public:
 
 	uint32 GetAllocatedSize() const
 	{
-		return(FPrimitiveSceneProxy::GetAllocatedSize() + (WaterVertexFactories.GetAllocatedSize() + WaterVertexFactories.Num() * sizeof(WaterVertexFactoryType)) + WaterQuadTree.GetAllocatedSize());
+		return(FPrimitiveSceneProxy::GetAllocatedSize() + (WaterVertexFactories.GetAllocatedSize() + WaterVertexFactories.Num() * sizeof(FWaterVertexFactoryType)) + WaterQuadTree.GetAllocatedSize());
 	}
 
 #if WITH_WATER_SELECTION_SUPPORT
@@ -61,9 +64,10 @@ public:
 #endif // WITH_WATER_SELECTION_SUPPORT
 
 	// At runtime, we only ever need one version of the vertex factory : with selection support (editor) or without : 
-	using WaterVertexFactoryType = TWaterVertexFactory<WITH_WATER_SELECTION_SUPPORT>;
-	using WaterInstanceDataBuffersType = TWaterInstanceDataBuffers<WITH_WATER_SELECTION_SUPPORT>;
-	using WaterMeshUserDataBuffersType = TWaterMeshUserDataBuffers<WITH_WATER_SELECTION_SUPPORT>;
+	using FWaterVertexFactoryType = TWaterVertexFactory<WITH_WATER_SELECTION_SUPPORT, /*bIndirectDraws = */ false>;
+	using FWaterVertexFactoryIndirectDrawType = TWaterVertexFactory<WITH_WATER_SELECTION_SUPPORT, /*bIndirectDraws = */ true>;
+	using FWaterInstanceDataBuffersType = TWaterInstanceDataBuffers<WITH_WATER_SELECTION_SUPPORT>;
+	using FWaterMeshUserDataBuffersType = TWaterMeshUserDataBuffers<WITH_WATER_SELECTION_SUPPORT>;
 
 #if RHI_RAYTRACING
 	virtual void GetDynamicRayTracingInstances(FRayTracingMaterialGatheringContext& Context, TArray<FRayTracingInstance>& OutRayTracingInstances) override final;
@@ -101,16 +105,23 @@ private:
 	FMaterialRelevance MaterialRelevance;
 
 	// One vertex factory per LOD
-	TArray<WaterVertexFactoryType*> WaterVertexFactories;
+	TArray<FWaterVertexFactoryType*> WaterVertexFactories;
+	TArray<FWaterVertexFactoryIndirectDrawType*> WaterVertexFactoriesIndirectDraw;
 
 	/** Tiles containing water, stored in a quad tree */
 	FWaterQuadTree WaterQuadTree;
 
+	/** GPU quad tree instance. Only initialized and used if WaterQuadTree.IsGPUQuadTree() is true. */
+	FWaterQuadTreeGPU QuadTreeGPU;
+
 	/** Unique Instance data buffer shared accross water batch draw calls */	
-	WaterInstanceDataBuffersType* WaterInstanceDataBuffers;
+	FWaterInstanceDataBuffersType* WaterInstanceDataBuffers;
 
 	/** Per-"water render group" user data (the number of groups might vary depending on whether we're in the editor or not) */
-	WaterMeshUserDataBuffersType* WaterMeshUserDataBuffers;
+	FWaterMeshUserDataBuffersType* WaterMeshUserDataBuffers;
+
+	double WaterQuadTreeMinHeight = DBL_MAX;
+	double WaterQuadTreeMaxHeight = -DBL_MAX;
 
 	/** The world-space bounds of the current water info texture coverage. The Water mesh should only render tiles within this bounding box. */
 	FBox2D WaterInfoBounds = FBox2D(ForceInit);
@@ -118,6 +129,9 @@ private:
 	/** Scale of the concentric LOD squares  */
 	float LODScale = -1.0f;
 
+	/** Number of quads per side of a water quad tree tile at LOD0 */
+	int32 NumQuadsLOD0 = 0;
+
 	/** Number of densities (same as number of grid index/vertex buffers) */
 	int32 DensityCount = 0;
 
@@ -142,6 +156,12 @@ private:
 	UE::FMutex OcclusionResultsMutex;
 	int32 OcclusionResultsFarMeshOffset = INT32_MAX;
 	uint32 SceneProxyCreatedFrameNumberRenderThread = INDEX_NONE;
+
+	mutable FWaterQuadTreeGPU::FTraverseParams WaterQuadTreeGPUTraverseParams;
+	mutable bool bNeedToTraverseGPUQuadTree = false;
+
+	/** Initializes the GPU quad tree */
+	void BuildGPUQuadTree(FRDGBuilder& GraphBuilder);
 };
 
 #if UE_ENABLE_INCLUDE_ORDER_DEPRECATED_IN_5_2
diff --git a/Engine/Plugins/Experimental/Water/Source/Runtime/Public/WaterQuadTree.h b/Engine/Plugins/Experimental/Water/Source/Runtime/Public/WaterQuadTree.h
--- a/Engine/Plugins/Experimental/Water/Source/Runtime/Public/WaterQuadTree.h
+++ b/Engine/Plugins/Experimental/Water/Source/Runtime/Public/WaterQuadTree.h
@@ -28,6 +28,13 @@ struct FWaterBodyRenderData
 	/** World Z position of the waterbody, this is where the tiles for this water body will be rendered*/
 	double SurfaceBaseHeight = 0.0;
 
+	/** Z bounds of the water body */
+	double BoundsMinZ = DBL_MAX;
+	double BoundsMaxZ = -DBL_MAX;
+
+	/** Maximum Z displacement due to waves */
+	float MaxWaveHeight = 0.0f;
+
 	/** Render priority. If two water bodies overlap, this will decide which water body is used for a tile */
 	int16 Priority = TNumericLimits<int16>::Min();
 
@@ -66,6 +73,15 @@ struct FWaterBodyRenderData
 	}
 };
 
+struct FWaterBodyQuadTreeRasterInfo
+{
+	FTransform LocalToWorld = FTransform::Identity;
+	const class FStaticMeshRenderData* RenderData = nullptr;
+	uint32 WaterBodyRenderDataIndex = INDEX_NONE;
+	uint8 Priority = 0;
+	bool bIsRiver = false;
+};
+
 struct FWaterQuadTree 
 {
 	enum { INVALID_PARENT = 0xFFFFFFF };
@@ -125,7 +141,7 @@ struct FWaterQuadTree
 	 *	Initialize the tree. This will unlock the tree for node insertion using AddWaterTilesInsideBounds(...). 
 	 *	Tree must be locked before traversal, see Lock(). 
 	 */
-	void InitTree(const FBox2D& InBounds, float InTileSize, FIntPoint InExtentInTiles);
+	void InitTree(const FBox2D& InBounds, float InTileSize, FIntPoint InExtentInTiles, bool bInIsGPUQuadTree);
 
 	/** Unlock to make it read-only. This will optionally prune the node array to remove redundant nodes, nodes that can be implicitly traversed */
 	void Unlock(bool bPruneRedundantNodes);
@@ -148,6 +164,8 @@ struct FWaterQuadTree
 	/** Add water body render data to this tree. Returns the index in the array. Use this index to add tiles with this water body to the tree, see AddWaterTilesInsideBounds(..) */
 	uint32 AddWaterBodyRenderData(const FWaterBodyRenderData& InWaterBodyRenderData) { return NodeData.WaterBodyRenderData.Add(InWaterBodyRenderData); }
 
+	void AddWaterBodyRasterInfo(const FWaterBodyQuadTreeRasterInfo& InWaterBodyRasterInfo) { WaterBodyRasterInfos.Add(InWaterBodyRasterInfo); }
+
 	/** Get bounds of the root node if there is one, otherwise some default box */
 	FBox GetBounds() const { return NodeData.Nodes.Num() > 0 ? NodeData.Nodes[0].Bounds : FBox(-FVector::OneVector, FVector::OneVector); }
 
@@ -183,11 +201,19 @@ struct FWaterQuadTree
 
 	const TArray<FMaterialRenderProxy*>& GetWaterMaterials() const { return WaterMaterials; }
 
+	const TArray<FWaterBodyQuadTreeRasterInfo>& GetWaterBodyRasterInfos() const { return WaterBodyRasterInfos; }
+
+	const TArray<FWaterBodyRenderData>& GetWaterBodyRenderData() const { return NodeData.WaterBodyRenderData; }
+
+	FIntPoint GetResolution() const { return ExtentInTiles * 2; }
+
+	bool IsGPUQuadTree() const { return bIsGPUQuadTree; }
+
 	/** Calculate the world distance to a LOD */
 	static float GetLODDistance(int32 InLODLevel, float InLODScale) { return FMath::Pow(2.0f, (float)(InLODLevel + 1)) * InLODScale; }
 
 	/** Total memory dynamically allocated by this object */
-	uint32 GetAllocatedSize() const { return NodeData.GetAllocatedSize() + WaterMaterials.GetAllocatedSize() + BreadthFirstOrder.GetAllocatedSize() + FarMeshData.GetAllocatedSize(); }
+	uint32 GetAllocatedSize() const { return NodeData.GetAllocatedSize() + WaterMaterials.GetAllocatedSize() + WaterBodyRasterInfos.GetAllocatedSize() + BreadthFirstOrder.GetAllocatedSize() + FarMeshData.GetAllocatedSize(); }
 
 #if WITH_WATER_SELECTION_SUPPORT
 	/** Obtain all possible hit proxies (proxies of all the water bodies) */
@@ -291,8 +317,9 @@ private:
 	} NodeData;
 
 	TArray<FMaterialRenderProxy*> WaterMaterials;
-
-	/** Node indices ordered by breadth first ordering */
+	TArray<FWaterBodyQuadTreeRasterInfo> WaterBodyRasterInfos;
+	
+/** Node indices ordered by breadth first ordering */
 	TArray<int32> BreadthFirstOrder;
 
 	/** Contains everything needed to render the far mesh. This data lives outside the quadtree structure itself */
@@ -331,4 +358,6 @@ private:
 
 	/** If true, the tree may not change */
 	bool bIsReadOnly = true;
+
+	bool bIsGPUQuadTree = false;
 };
diff --git a/Engine/Plugins/Experimental/Water/Source/Runtime/Public/WaterVertexFactory.h b/Engine/Plugins/Experimental/Water/Source/Runtime/Public/WaterVertexFactory.h
--- a/Engine/Plugins/Experimental/Water/Source/Runtime/Public/WaterVertexFactory.h
+++ b/Engine/Plugins/Experimental/Water/Source/Runtime/Public/WaterVertexFactory.h
@@ -164,16 +164,17 @@ enum class EWaterMeshRenderGroupType : uint8
 #endif // WITH_WATER_SELECTION_SUPPORT
 };
 
-template <bool bWithWaterSelectionSupport>
+template <bool bWithWaterSelectionSupport, bool bIndirectDraws>
 class TWaterVertexFactory : public FVertexFactory
 {
-	DECLARE_VERTEX_FACTORY_TYPE(TWaterVertexFactory<bWithWaterSelectionSupport>);
+	DECLARE_VERTEX_FACTORY_TYPE(FWaterVertexFactoryType);
 
 public:
 	using Super = FVertexFactory;
+	using FWaterVertexFactoryType = TWaterVertexFactory<bWithWaterSelectionSupport, bIndirectDraws>;
 
 	static constexpr int32 NumRenderGroups = bWithWaterSelectionSupport ? 3 : 1; // Must match EWaterMeshRenderGroupType
-	static constexpr int32 NumAdditionalVertexStreams = TWaterInstanceDataBuffers<bWithWaterSelectionSupport>::NumBuffers;
+	static constexpr int32 NumAdditionalVertexStreams = bIndirectDraws ? 0 : TWaterInstanceDataBuffers<bWithWaterSelectionSupport>::NumBuffers;
 
 	TWaterVertexFactory(ERHIFeatureLevel::Type InFeatureLevel, int32 InNumQuadsPerSide,	float InLODScale);
 	~TWaterVertexFactory();
@@ -215,7 +216,7 @@ private:
 	const float LODScale = 0.0f;
 };
 
-extern const FVertexFactoryType* GetWaterVertexFactoryType(bool bWithWaterSelectionSupport);
+extern const FVertexFactoryType* GetWaterVertexFactoryType(bool bWithWaterSelectionSupport, bool bIndirectDraws);
 
 
 /**
@@ -238,6 +239,11 @@ struct TWaterMeshUserData
 #if RHI_RAYTRACING	
 	FUniformBufferRHIRef WaterVertexFactoryRaytracingVFUniformBuffer = nullptr;
 #endif
+
+	FRHIShaderResourceView* IndirectInstanceDataOffsets = nullptr;
+	FRHIShaderResourceView* IndirectInstanceData0 = nullptr;
+	FRHIShaderResourceView* IndirectInstanceData1 = nullptr;
+	FRHIShaderResourceView* IndirectInstanceData2 = nullptr;
 };
 
 /**
@@ -267,7 +273,7 @@ struct TWaterMeshUserDataBuffers
 		return UserData[(int32)InRenderGroupType].Get();
 	}
 
-	TStaticArray<TUniquePtr<WaterMeshUserDataType>, TWaterVertexFactory<bWithWaterSelectionSupport>::NumRenderGroups> UserData;
+	TStaticArray<TUniquePtr<WaterMeshUserDataType>, TWaterVertexFactory<bWithWaterSelectionSupport, /*bIndirectDraws = */ false>::NumRenderGroups> UserData;
 };
 
 #include "WaterVertexFactory.inl"
diff --git a/Engine/Plugins/Experimental/Water/Source/Runtime/Public/WaterVertexFactory.inl b/Engine/Plugins/Experimental/Water/Source/Runtime/Public/WaterVertexFactory.inl
--- a/Engine/Plugins/Experimental/Water/Source/Runtime/Public/WaterVertexFactory.inl
+++ b/Engine/Plugins/Experimental/Water/Source/Runtime/Public/WaterVertexFactory.inl
@@ -13,8 +13,8 @@
 
 // ----------------------------------------------------------------------------------
 
-template <bool bWithWaterSelectionSupport>
-TWaterVertexFactory<bWithWaterSelectionSupport>::TWaterVertexFactory(ERHIFeatureLevel::Type InFeatureLevel, int32 InNumQuadsPerSide, float InLODScale)
+template <bool bWithWaterSelectionSupport, bool bIndirectDraws>
+TWaterVertexFactory<bWithWaterSelectionSupport, bIndirectDraws>::TWaterVertexFactory(ERHIFeatureLevel::Type InFeatureLevel, int32 InNumQuadsPerSide, float InLODScale)
 	: FVertexFactory(InFeatureLevel)
 	, NumQuadsPerSide(InNumQuadsPerSide)
 	, LODScale(InLODScale)
@@ -23,15 +23,15 @@ TWaterVertexFactory<bWithWaterSelectionSupport>::TWaterVertexFactory(ERHIFeature
 	IndexBuffer = new FWaterMeshIndexBuffer(NumQuadsPerSide);
 }
 
-template <bool bWithWaterSelectionSupport>
-TWaterVertexFactory<bWithWaterSelectionSupport>::~TWaterVertexFactory()
+template <bool bWithWaterSelectionSupport, bool bIndirectDraws>
+TWaterVertexFactory<bWithWaterSelectionSupport, bIndirectDraws>::~TWaterVertexFactory()
 {
 	delete VertexBuffer;
 	delete IndexBuffer;
 }
 
-template <bool bWithWaterSelectionSupport>
-void TWaterVertexFactory<bWithWaterSelectionSupport>::InitRHI(FRHICommandListBase& RHICmdList)
+template <bool bWithWaterSelectionSupport, bool bIndirectDraws>
+void TWaterVertexFactory<bWithWaterSelectionSupport, bIndirectDraws>::InitRHI(FRHICommandListBase& RHICmdList)
 {
 	Super::InitRHI(RHICmdList);
 
@@ -70,10 +70,13 @@ void TWaterVertexFactory<bWithWaterSelectionSupport>::InitRHI(FRHICommandListBas
 	Elements.Add(VertexPositionElement);
 
 	// Adds all streams
-	for (int32 StreamIdx = 0; StreamIdx < NumAdditionalVertexStreams; ++StreamIdx)
+	if constexpr (NumAdditionalVertexStreams > 0)
 	{
-		FVertexElement InstanceElement(Streams.Add(InstanceDataVertexStream), 0, VET_Float4, 8 + StreamIdx, InstanceDataVertexStream.Stride, true);
-		Elements.Add(InstanceElement);
+		for (int32 StreamIdx = 0; StreamIdx < NumAdditionalVertexStreams; ++StreamIdx)
+		{
+			FVertexElement InstanceElement(Streams.Add(InstanceDataVertexStream), 0, VET_Float4, 8 + StreamIdx, InstanceDataVertexStream.Stride, true);
+			Elements.Add(InstanceElement);
+		}
 	}
 
 #if 0
@@ -93,8 +96,8 @@ void TWaterVertexFactory<bWithWaterSelectionSupport>::InitRHI(FRHICommandListBas
 	InitDeclaration(Elements);
 }
 
-template <bool bWithWaterSelectionSupport>
-void TWaterVertexFactory<bWithWaterSelectionSupport>::ReleaseRHI()
+template <bool bWithWaterSelectionSupport, bool bIndirectDraws>
+void TWaterVertexFactory<bWithWaterSelectionSupport, bIndirectDraws>::ReleaseRHI()
 {
 	for (auto& UniformBuffer : UniformBuffers)
 	{
@@ -114,8 +117,8 @@ void TWaterVertexFactory<bWithWaterSelectionSupport>::ReleaseRHI()
 	Super::ReleaseRHI();
 }
 
-template <bool bWithWaterSelectionSupport>
-void TWaterVertexFactory<bWithWaterSelectionSupport>::SetupUniformDataForGroup(EWaterMeshRenderGroupType InRenderGroupType)
+template <bool bWithWaterSelectionSupport, bool bIndirectDraws>
+void TWaterVertexFactory<bWithWaterSelectionSupport, bIndirectDraws>::SetupUniformDataForGroup(EWaterMeshRenderGroupType InRenderGroupType)
 {
 	FWaterVertexFactoryParameters UniformParams;
 	UniformParams.NumQuadsPerTileSide = NumQuadsPerSide;
@@ -131,8 +134,8 @@ void TWaterVertexFactory<bWithWaterSelectionSupport>::SetupUniformDataForGroup(E
 	UniformBuffers[(int32)InRenderGroupType] = FWaterVertexFactoryBufferRef::CreateUniformBufferImmediate(UniformParams, UniformBuffer_MultiFrame);
 }
 
-template <bool bWithWaterSelectionSupport>
-bool TWaterVertexFactory<bWithWaterSelectionSupport>::ShouldCompilePermutation(const FVertexFactoryShaderPermutationParameters& Parameters)
+template <bool bWithWaterSelectionSupport, bool bIndirectDraws>
+bool TWaterVertexFactory<bWithWaterSelectionSupport, bIndirectDraws>::ShouldCompilePermutation(const FVertexFactoryShaderPermutationParameters& Parameters)
 {
 	const bool bIsCompatibleWithWater = ((Parameters.MaterialParameters.MaterialDomain == MD_Surface) && Parameters.MaterialParameters.bIsUsedWithWater) || Parameters.MaterialParameters.bIsSpecialEngineMaterial;
 	if (bIsCompatibleWithWater)
@@ -143,8 +146,8 @@ bool TWaterVertexFactory<bWithWaterSelectionSupport>::ShouldCompilePermutation(c
 	return false;
 }
 
-template <bool bWithWaterSelectionSupport>
-void TWaterVertexFactory<bWithWaterSelectionSupport>::ModifyCompilationEnvironment(const FVertexFactoryShaderPermutationParameters& Parameters, FShaderCompilerEnvironment& OutEnvironment)
+template <bool bWithWaterSelectionSupport, bool bIndirectDraws>
+void TWaterVertexFactory<bWithWaterSelectionSupport, bIndirectDraws>::ModifyCompilationEnvironment(const FVertexFactoryShaderPermutationParameters& Parameters, FShaderCompilerEnvironment& OutEnvironment)
 {
 	OutEnvironment.SetDefine(TEXT("WATER_MESH_FACTORY"), 1);
 #if 0
@@ -162,12 +165,16 @@ void TWaterVertexFactory<bWithWaterSelectionSupport>::ModifyCompilationEnvironme
 		OutEnvironment.SetDefine(TEXT("USE_VERTEXFACTORY_HITPROXY_ID"), TEXT("1"));
 		OutEnvironment.SetDefine(TEXT("WITH_WATER_SELECTION_SUPPORT_VF"), TEXT("1"));
 	}
+	if (bIndirectDraws)
+	{
+		OutEnvironment.SetDefine(TEXT("WATER_MESH_DRAW_INDIRECT"), TEXT("1"));
+	}
 
 	OutEnvironment.SetDefine(TEXT("RAY_TRACING_DYNAMIC_MESH_IN_LOCAL_SPACE"), TEXT("1"));
 }
 
-template <bool bWithWaterSelectionSupport>
-void TWaterVertexFactory<bWithWaterSelectionSupport>::ValidateCompiledResult(const FVertexFactoryType* Type, EShaderPlatform Platform, const FShaderParameterMap& ParameterMap, TArray<FString>& OutErrors)
+template <bool bWithWaterSelectionSupport, bool bIndirectDraws>
+void TWaterVertexFactory<bWithWaterSelectionSupport, bIndirectDraws>::ValidateCompiledResult(const FVertexFactoryType* Type, EShaderPlatform Platform, const FShaderParameterMap& ParameterMap, TArray<FString>& OutErrors)
 {
 #if 0
 	if (Type->SupportsPrimitiveIdStream()
@@ -179,15 +186,18 @@ void TWaterVertexFactory<bWithWaterSelectionSupport>::ValidateCompiledResult(con
 #endif
 }
 
-template <bool bWithWaterSelectionSupport>
-void TWaterVertexFactory<bWithWaterSelectionSupport>::GetPSOPrecacheVertexFetchElements(EVertexInputStreamType VertexInputStreamType, FVertexDeclarationElementList& Elements)
+template <bool bWithWaterSelectionSupport, bool bIndirectDraws>
+void TWaterVertexFactory<bWithWaterSelectionSupport, bIndirectDraws>::GetPSOPrecacheVertexFetchElements(EVertexInputStreamType VertexInputStreamType, FVertexDeclarationElementList& Elements)
 {
 	// Add position stream
 	Elements.Add(FVertexElement(0, 0, VET_Float4, 0, sizeof(FVector4f), false));
 
 	// Add all the additional streams
-	for (int32 StreamIdx = 0; StreamIdx < NumAdditionalVertexStreams; ++StreamIdx)
+	if constexpr (NumAdditionalVertexStreams > 0)
 	{
-		Elements.Add(FVertexElement(1 + StreamIdx, 0, VET_Float4, 8 + StreamIdx, sizeof(FVector4f), true));
+		for (int32 StreamIdx = 0; StreamIdx < NumAdditionalVertexStreams; ++StreamIdx)
+		{
+			Elements.Add(FVertexElement(1 + StreamIdx, 0, VET_Float4, 8 + StreamIdx, sizeof(FVector4f), true));
+		}
 	}
 }
diff --git a/Engine/Plugins/Experimental/Water/Source/Runtime/Public/WaterViewExtension.h b/Engine/Plugins/Experimental/Water/Source/Runtime/Public/WaterViewExtension.h
--- a/Engine/Plugins/Experimental/Water/Source/Runtime/Public/WaterViewExtension.h
+++ b/Engine/Plugins/Experimental/Water/Source/Runtime/Public/WaterViewExtension.h
@@ -24,6 +24,7 @@ public:
 	virtual void SetupViewFamily(FSceneViewFamily& InViewFamily) override;
 	virtual void PreRenderViewFamily_RenderThread(FRDGBuilder& GraphBuilder, FSceneViewFamily& InViewFamily) override;
 	virtual void PreRenderView_RenderThread(FRDGBuilder& GraphBuilder, FSceneView& InView) override;
+	virtual void PreRenderBasePass_RenderThread(FRDGBuilder& GraphBuilder) override;
 	// End FSceneViewExtensionBase implementation
 
 	void MarkWaterInfoTextureForRebuild(const UE::WaterInfo::FRenderingContext& RenderContext);
@@ -55,6 +56,18 @@ private:
 	void UpdateGPUBuffers();
 };
 
+struct FWaterMeshGPUWork
+{
+	struct FCallback
+	{
+		class FWaterMeshSceneProxy* Proxy = nullptr;
+		TFunction<void(FRDGBuilder&)> Function;
+	};
+	TArray<FCallback> Callbacks;
+};
+
+extern FWaterMeshGPUWork GWaterMeshGPUWork;
+
 #if UE_ENABLE_INCLUDE_ORDER_DEPRECATED_IN_5_2
 #include "CoreMinimal.h"
 #endif
diff --git a/Engine/Source/Runtime/Engine/Public/SceneViewExtension.h b/Engine/Source/Runtime/Engine/Public/SceneViewExtension.h
--- a/Engine/Source/Runtime/Engine/Public/SceneViewExtension.h
+++ b/Engine/Source/Runtime/Engine/Public/SceneViewExtension.h
@@ -155,6 +155,11 @@ public:
 	 */
 	virtual void PreInitViews_RenderThread(FRDGBuilder& GraphBuilder) {}
 
+	/**
+	 * Called on render thread right before Base Pass rendering.
+	 */
+	virtual void PreRenderBasePass_RenderThread(FRDGBuilder& GraphBuilder) {}
+
 	/**
 	 * Called right after Base Pass rendering finished when using the deferred renderer.
 	 */
diff --git a/Engine/Source/Runtime/Renderer/Private/DeferredShadingRenderer.cpp b/Engine/Source/Runtime/Renderer/Private/DeferredShadingRenderer.cpp
--- a/Engine/Source/Runtime/Renderer/Private/DeferredShadingRenderer.cpp
+++ b/Engine/Source/Runtime/Renderer/Private/DeferredShadingRenderer.cpp
@@ -3462,6 +3462,11 @@ void FDeferredShadingSceneRenderer::Render(FRDGBuilder& GraphBuilder)
 
 		// End early occlusion queries
 
+		for (FSceneViewExtensionRef& ViewExtension : ViewFamily.ViewExtensions)
+		{
+			ViewExtension->PreRenderBasePass_RenderThread(GraphBuilder);
+		}
+
 		BeginAsyncDistanceFieldShadowProjections(GraphBuilder, SceneTextures, InitViewTaskDatas.DynamicShadows);
 
 		// Run local fog volume culling before base pass and after HZB generation tyo benefit from more culling.
diff --git a/Engine/Source/Runtime/Renderer/Private/MobileShadingRenderer.cpp b/Engine/Source/Runtime/Renderer/Private/MobileShadingRenderer.cpp
--- a/Engine/Source/Runtime/Renderer/Private/MobileShadingRenderer.cpp
+++ b/Engine/Source/Runtime/Renderer/Private/MobileShadingRenderer.cpp
@@ -1124,6 +1124,11 @@ void FMobileSceneRenderer::Render(FRDGBuilder& GraphBuilder)
 		}
 	}
 
+	for (FSceneViewExtensionRef& ViewExtension : ViewFamily.ViewExtensions)
+	{
+		ViewExtension->PreRenderBasePass_RenderThread(GraphBuilder);
+	}
+
 	if (bRendererOutputFinalSceneColor)
 	{
 		if (bDeferredShading)
