[Backout] - CL29714981
Original CL Desc
-----------------------------------------------------------------
Implemented a path for generating and traversing the water quadtree on the GPU. Traversal creates indirect draw arguments which are used to draw the tessellated water mesh tiles. The quadtree is created by rasterizing water meshes into a supersampled and multisampled texture, downsampling that texture to quadtree LOD0 resolution and finally building a mip chain. Every frame when water needs to be rendered, all nodes (texels) in that quadtree texture are visited by a compute shader to determine which ones need to be drawn from a given view. The generated list of tiles is sorted into draw buckets (quad mesh tessellation and material) and finally a list of per-instance data is generated that is passed to the modified water vertex factory. The VF can now handle indirect draws with instance data read from a compacted buffer using another indirection to get to the correct offset.

The indirect draw pipeline also supports HZB based occlusion culling.

#rnx

diff --git a/Engine/Plugins/Experimental/Water/Shaders/Private/WaterMeshVertexFactory.ush b/Engine/Plugins/Experimental/Water/Shaders/Private/WaterMeshVertexFactory.ush
--- a/Engine/Plugins/Experimental/Water/Shaders/Private/WaterMeshVertexFactory.ush
+++ b/Engine/Plugins/Experimental/Water/Shaders/Private/WaterMeshVertexFactory.ush
@@ -2,23 +2,6 @@
 
 #include "/Engine/Private/VertexFactoryCommon.ush"
 
-#ifndef WATER_MESH_DRAW_INDIRECT
-#define WATER_MESH_DRAW_INDIRECT 0
-#endif
-
-#if WATER_MESH_DRAW_INDIRECT
-
-Buffer<uint> InstanceDataOffsetsBuffer;
-Buffer<float4> InstanceData0Buffer;
-Buffer<float4> InstanceData1Buffer;
-#if WITH_WATER_SELECTION_SUPPORT_VF
-Buffer<float4> InstanceData2Buffer;
-#endif
-
-uint DrawBucketIndex;
-
-#endif
-
 struct FVertexFactoryInterpolantsVSToPS
 {
 #if NUM_TEX_COORD_INTERPOLATORS
@@ -43,15 +26,11 @@ struct FVertexFactoryInput
 {
 	float4	Position	: ATTRIBUTE0;
 
-#if WATER_MESH_DRAW_INDIRECT && !RAYHITGROUPSHADER && !COMPUTESHADER
-	uint InstanceId : SV_InstanceID;
-#else
 	float4 InstanceData0 : ATTRIBUTE8;
 	float4 InstanceData1 : ATTRIBUTE9; 
 #if HIT_PROXY_SHADER || WITH_WATER_SELECTION_SUPPORT_VF
 	float4 InstanceData2 : ATTRIBUTE10; // Store an HitProxyID per instance
 #endif
-#endif // WATER_MESH_DRAW_INDIRECT && !RAYHITGROUPSHADER && !COMPUTESHADER
 
 	VF_GPUSCENE_DECLARE_INPUT_BLOCK(13)
 	VF_INSTANCED_STEREO_DECLARE_INPUT_BLOCK()
@@ -64,12 +43,8 @@ struct FPositionOnlyVertexFactoryInput
 {
 	float4	Position	: ATTRIBUTE0;
 
-#if WATER_MESH_DRAW_INDIRECT && !RAYHITGROUPSHADER && !COMPUTESHADER
-	uint InstanceId : SV_InstanceID;
-#else
 	float4 InstanceData0 : ATTRIBUTE8;
 	float4 InstanceData1 : ATTRIBUTE9; 
-#endif // WATER_MESH_DRAW_INDIRECT && !RAYHITGROUPSHADER && !COMPUTESHADER
 
 	VF_GPUSCENE_DECLARE_INPUT_BLOCK(1)
 	VF_INSTANCED_STEREO_DECLARE_INPUT_BLOCK()
@@ -83,12 +58,8 @@ struct FPositionAndNormalOnlyVertexFactoryInput
 	float4	Position	: ATTRIBUTE0;
 	float4	Normal		: ATTRIBUTE2;
 
-#if WATER_MESH_DRAW_INDIRECT && !RAYHITGROUPSHADER && !COMPUTESHADER
-	uint InstanceId : SV_InstanceID;
-#else
 	float4 InstanceData0 : ATTRIBUTE8;
 	float4 InstanceData1 : ATTRIBUTE9; 
-#endif // WATER_MESH_DRAW_INDIRECT && !RAYHITGROUPSHADER && !COMPUTESHADER
 
 	VF_GPUSCENE_DECLARE_INPUT_BLOCK(1)
 	VF_INSTANCED_STEREO_DECLARE_INPUT_BLOCK()
@@ -238,14 +209,8 @@ FVertexFactoryIntermediates GetVertexFactoryIntermediates(FVertexFactoryInput In
 	FVertexFactoryIntermediates Intermediates;
 
 	// Get the packed instance data
-#if WATER_MESH_DRAW_INDIRECT && !RAYHITGROUPSHADER && !COMPUTESHADER
-	const uint InstanceDataIndex = InstanceDataOffsetsBuffer.Load(DrawBucketIndex) + Input.InstanceId;
-	float4 Data0 = InstanceData0Buffer.Load(InstanceDataIndex);
-	float4 Data1 = InstanceData1Buffer.Load(InstanceDataIndex);
-#else
 	float4 Data0 = Input.InstanceData0;
 	float4 Data1 = Input.InstanceData1;
-#endif // WATER_MESH_DRAW_INDIRECT && !RAYHITGROUPSHADER && !COMPUTESHADER
 
 	const float3 Translation = Data0.xyz;
 	Intermediates.WaveParamIndex = asuint(Data0.w);
@@ -289,11 +254,7 @@ FVertexFactoryIntermediates GetVertexFactoryIntermediates(FVertexFactoryInput In
 	}
 
 #if WITH_WATER_SELECTION_SUPPORT_VF
-#if WATER_MESH_DRAW_INDIRECT && !RAYHITGROUPSHADER && !COMPUTESHADER
-	float SelectedValue = InstanceData2Buffer.Load(InstanceDataIndex).w;
-#else
 	float SelectedValue = Input.InstanceData2.w;
-#endif
 	float IsVisible = WaterVF.bRenderSelected * SelectedValue + WaterVF.bRenderUnselected * (1-SelectedValue);
 	Intermediates.MorphedTranslatedWorldPos *= IsVisible;
 #endif // WITH_WATER_SELECTION_SUPPORT_VF
@@ -307,12 +268,7 @@ FVertexFactoryIntermediates GetVertexFactoryIntermediates(FVertexFactoryInput In
 float4 VertexFactoryGetInstanceHitProxyId(FVertexFactoryInput Input, FVertexFactoryIntermediates Intermediates)
 {
 #if WITH_WATER_SELECTION_SUPPORT_VF
-#if WATER_MESH_DRAW_INDIRECT && !RAYHITGROUPSHADER && !COMPUTESHADER
-	const uint InstanceDataIndex = InstanceDataOffsetsBuffer.Load(DrawBucketIndex) + Input.InstanceId;
-	return float4(InstanceData2Buffer.Load(InstanceDataIndex).rgb, 0);
-#else
 	return float4(Input.InstanceData2.rgb, 0);
-#endif
 #else // WITH_WATER_SELECTION_SUPPORT_VF
 	return 0;
 #endif // !WITH_WATER_SELECTION_SUPPORT_VF
diff --git a/Engine/Plugins/Experimental/Water/Shaders/Private/WaterQuadTreeBuild.usf b/Engine/Plugins/Experimental/Water/Shaders/Private/WaterQuadTreeBuild.usf
deleted file mode 100644
diff --git a/Engine/Plugins/Experimental/Water/Shaders/Private/WaterQuadTreeCommon.ush b/Engine/Plugins/Experimental/Water/Shaders/Private/WaterQuadTreeCommon.ush
deleted file mode 100644
diff --git a/Engine/Plugins/Experimental/Water/Shaders/Private/WaterQuadTreeDraws.usf b/Engine/Plugins/Experimental/Water/Shaders/Private/WaterQuadTreeDraws.usf
deleted file mode 100644
diff --git a/Engine/Plugins/Experimental/Water/Shaders/Private/WaterQuadTreeMerge.usf b/Engine/Plugins/Experimental/Water/Shaders/Private/WaterQuadTreeMerge.usf
deleted file mode 100644
diff --git a/Engine/Plugins/Experimental/Water/Shaders/Private/WaterQuadTreePixelShader.usf b/Engine/Plugins/Experimental/Water/Shaders/Private/WaterQuadTreePixelShader.usf
deleted file mode 100644
diff --git a/Engine/Plugins/Experimental/Water/Shaders/Private/WaterQuadTreeVertexShader.usf b/Engine/Plugins/Experimental/Water/Shaders/Private/WaterQuadTreeVertexShader.usf
deleted file mode 100644
diff --git a/Engine/Plugins/Experimental/Water/Source/Runtime/Private/WaterMeshComponent.cpp b/Engine/Plugins/Experimental/Water/Source/Runtime/Private/WaterMeshComponent.cpp
--- a/Engine/Plugins/Experimental/Water/Source/Runtime/Private/WaterMeshComponent.cpp
+++ b/Engine/Plugins/Experimental/Water/Source/Runtime/Private/WaterMeshComponent.cpp
@@ -6,14 +6,12 @@
 #include "Materials/MaterialInstanceDynamic.h"
 #include "Materials/Material.h"
 #include "PhysicsEngine/BodySetup.h"
-#include "Engine/StaticMesh.h"
 #include "WaterBodyComponent.h"
 #include "WaterMeshSceneProxy.h"
 #include "WaterModule.h"
 #include "WaterSplineComponent.h"
 #include "WaterSubsystem.h"
 #include "WaterUtils.h"
-#include "WaterBodyInfoMeshComponent.h"
 
 #include UE_INLINE_GENERATED_CPP_BY_NAME(WaterMeshComponent)
 
@@ -33,13 +31,6 @@ static TAutoConsoleVariable<float> CVarWaterMeshLODScaleBias(
 	TEXT("This value is added to the LOD Scale of each Mesh Component. Negative values will lower the overall density/resolution or the vertex grid and make the LODs smaller, higher values will increase the density/resolution and make the LODs larger. Smallest value is -0.5. That will make the inner LOD as tight and optimized as possible"),
 	ECVF_Scalability);
 
-static TAutoConsoleVariable<int32> CVarWaterMeshGPUQuadTree(
-	TEXT("r.Water.WaterMesh.GPUQuadTree"),
-	0,
-	TEXT("Builds the water quadtree on the GPU and does indirect draws of water tiles, driven by the GPU."),
-	ECVF_RenderThreadSafe
-);
-
 /** Debug CVars */ 
 static TAutoConsoleVariable<int32> CVarWaterMeshShowTileGenerationGeometry(
 	TEXT("r.Water.WaterMesh.ShowTileGenerationGeometry"),
@@ -82,7 +73,7 @@ void UWaterMeshComponent::PostLoad()
 
 void UWaterMeshComponent::CollectPSOPrecacheData(const FPSOPrecacheParams& BasePrecachePSOParams, FComponentPSOPrecacheParamsList& OutParams)
 {
-	const FVertexFactoryType* WaterVertexFactoryType = GetWaterVertexFactoryType(/*bWithWaterSelectionSupport = */ false, /*bIndirectDraws = */ false);
+	const FVertexFactoryType* WaterVertexFactoryType = GetWaterVertexFactoryType(/*bWithWaterSelectionSupport = */ false);
 	if (FarDistanceMaterial)
 	{
 		FComponentPSOPrecacheParams& ComponentParams = OutParams[OutParams.AddDefaulted()];
@@ -256,10 +247,8 @@ void UWaterMeshComponent::RebuildWaterMesh(float InTileSize, const FIntPoint& In
 		return;
 	}
 
-	const bool bIsGPUQuadTree = CVarWaterMeshGPUQuadTree.GetValueOnGameThread() != 0;
-
 	// This resets the tree to an initial state, ready for node insertion
-	WaterQuadTree.InitTree(WaterWorldBox, InTileSize, InExtentInTiles, bIsGPUQuadTree);
+	WaterQuadTree.InitTree(WaterWorldBox, InTileSize, InExtentInTiles);
 
 	UsedMaterials.Empty();
 
@@ -277,61 +266,25 @@ void UWaterMeshComponent::RebuildWaterMesh(float InTileSize, const FIntPoint& In
 	// Go through all water body actors to figure out bounds and water tiles
 	AWaterZone* OwningZone = GetOwner<AWaterZone>();
 	check(OwningZone);
-
-	// Min and max user defined priority range. (Input also clamped on OverlapMaterialPriority in AWaterBody)
-	constexpr int32 MinWaterBodyPriority = -8192;
-	constexpr int32 MaxWaterBodyPriority = 8191;
-	constexpr int32 GPUQuadTreeMaxNumPriorities = 8;
-
-	// The GPU quadtree only supports 8 different priority values, so we need to remap priorities into that space.
-	// Fortunately, rivers and non-river water bodies are rendered into their own "priority space", so we don't need to worry about
-	// moving river priorities into their own range.
-	TArray<int16> SortedPriorities;
-	if (bIsGPUQuadTree)
+	OwningZone->ForEachWaterBodyComponent([this, WaterWorldBox, bIsFlooded, GlobalOceanHeight, OceanFlood, &FarMeshHeight, &bHasOcean](UWaterBodyComponent* WaterBodyComponent)
 	{
-		OwningZone->ForEachWaterBodyComponent([this, WaterWorldBox, &SortedPriorities](UWaterBodyComponent* WaterBodyComponent)
-		{
-			check(WaterBodyComponent);
-			AActor* Actor = WaterBodyComponent->GetOwner();
-			check(Actor);
-
-			// Skip invisible water bodies
-			if (!WaterBodyComponent->ShouldRender() || !WaterBodyComponent->ShouldGenerateWaterMeshTile())
-			{
-				return true;
-			}
-
-			// Don't process water bodies which have their spline outside of this water mesh
-			const FBox WaterBodyBounds = WaterBodyComponent->Bounds.GetBox();
-			if (!WaterBodyBounds.IntersectXY(FBox(FVector(WaterWorldBox.Min, 0.0f), FVector(WaterWorldBox.Max, 0.0f))))
-			{
-				return true;
-			}
-
-			const int16 Priority = static_cast<int16>(FMath::Clamp(WaterBodyComponent->GetOverlapMaterialPriority(), MinWaterBodyPriority, MaxWaterBodyPriority));
-			SortedPriorities.AddUnique(Priority);
-
-			return true;
-		});
+		check(WaterBodyComponent);
+		AActor* Actor = WaterBodyComponent->GetOwner();
+		check(Actor);
 
-		SortedPriorities.Sort();
+		const FBox WaterBodyBounds = WaterBodyComponent->Bounds.GetBox();
 
-		if (SortedPriorities.Num() > GPUQuadTreeMaxNumPriorities)
+		// Don't process water bodies that has their spline outside of this water mesh
+		if (!WaterBodyBounds.IntersectXY(FBox(FVector(WaterWorldBox.Min, 0.0f), FVector(WaterWorldBox.Max, 0.0f))))
 		{
-			UE_LOG(LogWater, Warning, TEXT("WaterZone has more unique water body priorities (%i) than can be supported with GPU driven water quadtree rendering (%i)!"), SortedPriorities.Num(), GPUQuadTreeMaxNumPriorities);
+			return true;
 		}
-	}
-	
 
-	bool bAnyWaterMeshesNotReady = false;
+		FWaterBodyRenderData RenderData;
 
-	OwningZone->ForEachWaterBodyComponent([this, bIsGPUQuadTree, WaterWorldBox, bIsFlooded, GlobalOceanHeight, OceanFlood, &FarMeshHeight, &bHasOcean, &SortedPriorities, &bAnyWaterMeshesNotReady](UWaterBodyComponent* WaterBodyComponent)
-	{
-		check(WaterBodyComponent);
-		AActor* Actor = WaterBodyComponent->GetOwner();
-		check(Actor);
+		const EWaterBodyType WaterBodyType = WaterBodyComponent->GetWaterBodyType();
 
-		// Skip invisible water bodies
+		// Skip invisible water bodies : 
 		if (!WaterBodyComponent->ShouldRender())
 		{
 			return true;
@@ -343,17 +296,6 @@ void UWaterMeshComponent::RebuildWaterMesh(float InTileSize, const FIntPoint& In
 			return true;
 		}
 
-		// Don't process water bodies that has their spline outside of this water mesh
-		const FBox WaterBodyBounds = WaterBodyComponent->Bounds.GetBox();
-		if (!WaterBodyBounds.IntersectXY(FBox(FVector(WaterWorldBox.Min, 0.0f), FVector(WaterWorldBox.Max, 0.0f))))
-		{
-			return true;
-		}
-
-		FWaterBodyRenderData RenderData;
-
-		const EWaterBodyType WaterBodyType = WaterBodyComponent->GetWaterBodyType();
-
 		if (WaterBodyType != EWaterBodyType::Ocean)
 		{
 			if (bIsFlooded)
@@ -387,12 +329,12 @@ void UWaterMeshComponent::RebuildWaterMesh(float InTileSize, const FIntPoint& In
 			UsedMaterials.Add(RenderData.Material);
 		}
 
-		RenderData.Priority = static_cast<int16>(FMath::Clamp(WaterBodyComponent->GetOverlapMaterialPriority(), MinWaterBodyPriority, MaxWaterBodyPriority));
+		// Min and max user defined priority range. (Input also clamped on OverlapMaterialPriority in AWaterBody)
+		const int32 MinPriority = -8192;
+		const int32 MaxPriority = 8191;
+		RenderData.Priority = static_cast<int16>(FMath::Clamp(WaterBodyComponent->GetOverlapMaterialPriority(), MinPriority, MaxPriority));
 		RenderData.WaterBodyIndex = static_cast<int16>(WaterBodyComponent->GetWaterBodyIndex());
 		RenderData.SurfaceBaseHeight = WaterBodyComponent->GetComponentLocation().Z;
-		RenderData.MaxWaveHeight = WaterBodyComponent->GetMaxWaveHeight();
-		RenderData.BoundsMinZ = WaterBodyComponent->Bounds.GetBox().Min.Z;
-		RenderData.BoundsMaxZ = WaterBodyComponent->Bounds.GetBox().Max.Z;
 		RenderData.WaterBodyType = static_cast<int8>(WaterBodyType);
 #if WITH_WATER_SELECTION_SUPPORT
 		RenderData.HitProxy = new HActor(/*InActor = */Actor, /*InPrimComponent = */nullptr);
@@ -430,248 +372,207 @@ void UWaterMeshComponent::RebuildWaterMesh(float InTileSize, const FIntPoint& In
 		if (RenderData.RiverToLakeMaterial || RenderData.RiverToOceanMaterial)
 		{
 			// Move rivers up to it's own priority space, so that they always have precedence if they have transitions and that they only compare agains other rivers with transitions
-			RenderData.Priority += (MaxWaterBodyPriority-MinWaterBodyPriority)+1;
+			RenderData.Priority += (MaxPriority-MinPriority)+1;
 		}
 
 		uint32 WaterBodyRenderDataIndex = WaterQuadTree.AddWaterBodyRenderData(RenderData);
 
-		if (bIsGPUQuadTree)
+		switch (WaterBodyType)
 		{
-			// On the GPU path, we only submit FWaterBodyQuadTreeRasterInfo for the water quadtree to be rasterized on the GPU. In particular, we need the static mesh, a transform and a priority/WaterBodyRenderData index.
-
-			// TODO: Add support for flooded ocean water heights. This is a legacy feature and currently not used in any project (?), so it's probably not urgent.
-			// In order to implement this, we would need to render a fullscreen quad/triangle for the ocean instead of using the spline derived mesh.
+		case EWaterBodyType::River:
+		{
+			TRACE_CPUPROFILER_EVENT_SCOPE(River);
 
-			UWaterBodyInfoMeshComponent* WaterBodyInfoMeshComponent = WaterBodyComponent->GetWaterInfoMeshComponent();
-			if (ensure(WaterBodyInfoMeshComponent))
+			TArray<FBox, TInlineAllocator<16>> Boxes;
+			TArray<UPrimitiveComponent*> CollisionComponents = WaterBodyComponent->GetCollisionComponents();
+			for (UPrimitiveComponent* Comp : CollisionComponents)
 			{
-				UStaticMesh* StaticMesh = WaterBodyInfoMeshComponent->GetStaticMesh();
-				if (ensure(StaticMesh))
+				if (UBodySetup* BodySetup = (Comp != nullptr) ? Comp->GetBodySetup() : nullptr)
 				{
-					bAnyWaterMeshesNotReady |= StaticMesh->IsCompiling();
-					FStaticMeshRenderData* StaticMeshRenderData = StaticMesh->GetRenderData();
-					if (StaticMeshRenderData)
+					// Go through all sub shapes on the bodysetup to get a tight fit along water body
+					for (const FKConvexElem& ConvElem : BodySetup->AggGeom.ConvexElems)
 					{
-						const int16 ClampedPriority = static_cast<int16>(FMath::Clamp(WaterBodyComponent->GetOverlapMaterialPriority(), MinWaterBodyPriority, MaxWaterBodyPriority));
+						TRACE_CPUPROFILER_EVENT_SCOPE(Add);
 
-						FWaterBodyQuadTreeRasterInfo RasterInfo;
-						RasterInfo.LocalToWorld = WaterBodyComponent->GetComponentTransform();
-						RasterInfo.RenderData = StaticMeshRenderData;
-						RasterInfo.WaterBodyRenderDataIndex = WaterBodyRenderDataIndex;
-						RasterInfo.Priority = FMath::Clamp(SortedPriorities.IndexOfByKey(ClampedPriority), 0, GPUQuadTreeMaxNumPriorities - 1);
-						RasterInfo.bIsRiver = WaterBodyType == EWaterBodyType::River;
+						FBox SubBox = ConvElem.ElemBox.TransformBy(Comp->GetComponentTransform().ToMatrixWithScale());
+						SubBox.Max.Z += WaterBodyComponent->GetMaxWaveHeight();
 
-						WaterQuadTree.AddWaterBodyRasterInfo(RasterInfo);
+						Boxes.Add(SubBox);
 					}
 				}
+				else
+				{
+					// fallback on global AABB: 
+					FVector Center;
+					FVector Extent;
+					Actor->GetActorBounds(false, Center, Extent);
+					FBox Box(FBox::BuildAABB(Center, Extent));
+					Box.Max.Z += WaterBodyComponent->GetMaxWaveHeight();
+					Boxes.Add(Box);
+				}
 			}
 
-			if (WaterBodyType == EWaterBodyType::Ocean)
+			for (const FBox& Box : Boxes)
 			{
-				// Place far mesh height just below the ocean level
-				FarMeshHeight = RenderData.SurfaceBaseHeight - RenderData.MaxWaveHeight;
-				bHasOcean = true;
+				WaterQuadTree.AddWaterTilesInsideBounds(Box, WaterBodyRenderDataIndex);
+
+	#if !(UE_BUILD_SHIPPING || UE_BUILD_TEST)
+				if (!!CVarWaterMeshShowTileGenerationGeometry.GetValueOnGameThread())
+				{
+					DrawDebugBox(GetWorld(), Box.GetCenter(), Box.GetExtent(), FColor::Red);
+				}
+	#endif
 			}
+			break;
 		}
-		else
+		case EWaterBodyType::Lake:
 		{
-			switch (WaterBodyType)
-			{
-			case EWaterBodyType::River:
-			{
-				TRACE_CPUPROFILER_EVENT_SCOPE(River);
+			TRACE_CPUPROFILER_EVENT_SCOPE(Lake);
 
-				TArray<FBox, TInlineAllocator<16>> Boxes;
-				TArray<UPrimitiveComponent*> CollisionComponents = WaterBodyComponent->GetCollisionComponents();
-				for (UPrimitiveComponent* Comp : CollisionComponents)
-				{
-					if (UBodySetup* BodySetup = (Comp != nullptr) ? Comp->GetBodySetup() : nullptr)
-					{
-						// Go through all sub shapes on the bodysetup to get a tight fit along water body
-						for (const FKConvexElem& ConvElem : BodySetup->AggGeom.ConvexElems)
-						{
-							TRACE_CPUPROFILER_EVENT_SCOPE(Add);
+			const UWaterSplineComponent* SplineComp = WaterBodyComponent->GetWaterSpline();
+			const int32 NumOriginalSplinePoints = SplineComp->GetNumberOfSplinePoints();
+			float ConstantZ = SplineComp->GetLocationAtDistanceAlongSpline(0.0f, ESplineCoordinateSpace::World).Z;
 
-							FBox SubBox = ConvElem.ElemBox.TransformBy(Comp->GetComponentTransform().ToMatrixWithScale());
-							SubBox.Max.Z += WaterBodyComponent->GetMaxWaveHeight();
+			FBox LakeBounds = Actor->GetComponentsBoundingBox(/* bNonColliding = */true);
+			LakeBounds.Max.Z += WaterBodyComponent->GetMaxWaveHeight();
 
-							Boxes.Add(SubBox);
-						}
-					}
-					else
-					{
-						// fallback on global AABB: 
-						FVector Center;
-						FVector Extent;
-						Actor->GetActorBounds(false, Center, Extent);
-						FBox Box(FBox::BuildAABB(Center, Extent));
-						Box.Max.Z += WaterBodyComponent->GetMaxWaveHeight();
-						Boxes.Add(Box);
-					}
-				}
-
-				for (const FBox& Box : Boxes)
-				{
-					WaterQuadTree.AddWaterTilesInsideBounds(Box, WaterBodyRenderDataIndex);
-
-#if !(UE_BUILD_SHIPPING || UE_BUILD_TEST)
-					if (!!CVarWaterMeshShowTileGenerationGeometry.GetValueOnGameThread())
-					{
-						DrawDebugBox(GetWorld(), Box.GetCenter(), Box.GetExtent(), FColor::Red);
-					}
-#endif
-				}
+			// Skip lakes with less than 3 spline points
+			if (NumOriginalSplinePoints < 3)
+			{
 				break;
 			}
-			case EWaterBodyType::Lake:
-			{
-				TRACE_CPUPROFILER_EVENT_SCOPE(Lake);
 
-				const UWaterSplineComponent* SplineComp = WaterBodyComponent->GetWaterSpline();
-				const int32 NumOriginalSplinePoints = SplineComp->GetNumberOfSplinePoints();
-				float ConstantZ = SplineComp->GetLocationAtDistanceAlongSpline(0.0f, ESplineCoordinateSpace::World).Z;
-
-				FBox LakeBounds = Actor->GetComponentsBoundingBox(/* bNonColliding = */true);
-				LakeBounds.Max.Z += WaterBodyComponent->GetMaxWaveHeight();
-
-				// Skip lakes with less than 3 spline points
-				if (NumOriginalSplinePoints < 3)
-				{
-					break;
-				}
+			TArray<TArray<FVector2D>> PolygonBatches;
+			// Reuse the convex hulls generated for the physics shape because the work has already been done, but we can fallback to a simple spline evaluation method in case there's no physics:
+			bool bUseFallbackMethod = true;
 
-				TArray<TArray<FVector2D>> PolygonBatches;
-				// Reuse the convex hulls generated for the physics shape because the work has already been done, but we can fallback to a simple spline evaluation method in case there's no physics:
-				bool bUseFallbackMethod = true;
-
-				TArray<UPrimitiveComponent*> CollisionComponents = WaterBodyComponent->GetCollisionComponents();
-				if (CollisionComponents.Num() > 0)
+			TArray<UPrimitiveComponent*> CollisionComponents = WaterBodyComponent->GetCollisionComponents();
+			if (CollisionComponents.Num() > 0)
+			{
+				UPrimitiveComponent* Comp = CollisionComponents[0];
+				if (UBodySetup* BodySetup = (Comp != nullptr) ? Comp->GetBodySetup() : nullptr)
 				{
-					UPrimitiveComponent* Comp = CollisionComponents[0];
-					if (UBodySetup* BodySetup = (Comp != nullptr) ? Comp->GetBodySetup() : nullptr)
+					FTransform CompTransform = Comp->GetComponentTransform();
+					// Go through all sub shapes on the bodysetup to get a tight fit along water body
+					for (const FKConvexElem& ConvElem : BodySetup->AggGeom.ConvexElems)
 					{
-						FTransform CompTransform = Comp->GetComponentTransform();
-						// Go through all sub shapes on the bodysetup to get a tight fit along water body
-						for (const FKConvexElem& ConvElem : BodySetup->AggGeom.ConvexElems)
-						{
-							TRACE_CPUPROFILER_EVENT_SCOPE(Add);
-
-							// Vertex data contains the bottom vertices first, then the top ones :
-							int32 TotalNumVertices = ConvElem.VertexData.Num();
-							check(TotalNumVertices % 2 == 0);
-							int32 NumVertices = TotalNumVertices / 2;
-							if (NumVertices > 0)
+						TRACE_CPUPROFILER_EVENT_SCOPE(Add);
+
+						// Vertex data contains the bottom vertices first, then the top ones :
+						int32 TotalNumVertices = ConvElem.VertexData.Num();
+						check(TotalNumVertices % 2 == 0);
+						int32 NumVertices = TotalNumVertices / 2;
+						if (NumVertices > 0)
+						{ 
+							bUseFallbackMethod = false;
+
+							// Because the physics shape is made of multiple convex hulls, we cannot simply add their vertices to one big list but have to have 1 batch of polygons per convex hull
+							TArray<FVector2D>& Polygon = PolygonBatches.Emplace_GetRef();
+							Polygon.SetNum(NumVertices);
+							for (int32 i = 0; i < NumVertices; ++i)
 							{
-								bUseFallbackMethod = false;
-
-								// Because the physics shape is made of multiple convex hulls, we cannot simply add their vertices to one big list but have to have 1 batch of polygons per convex hull
-								TArray<FVector2D>& Polygon = PolygonBatches.Emplace_GetRef();
-								Polygon.SetNum(NumVertices);
-								for (int32 i = 0; i < NumVertices; ++i)
-								{
-									// Gather the top vertices :
-									Polygon[i] = FVector2D(CompTransform.TransformPosition(ConvElem.VertexData[NumVertices + i]));
-								}
+								// Gather the top vertices :
+								Polygon[i] = FVector2D(CompTransform.TransformPosition(ConvElem.VertexData[NumVertices + i]));
 							}
 						}
 					}
 				}
+			}
 
-				if (bUseFallbackMethod)
-				{
-					TArray<FVector> PolyLineVertices;
-					SplineComp->ConvertSplineToPolyLine(ESplineCoordinateSpace::World, FMath::Square(CVarWaterSplineResampleMaxDistance.GetValueOnGameThread()), PolyLineVertices);
+			if (bUseFallbackMethod)
+			{
+				TArray<FVector> PolyLineVertices;
+				SplineComp->ConvertSplineToPolyLine(ESplineCoordinateSpace::World, FMath::Square(CVarWaterSplineResampleMaxDistance.GetValueOnGameThread()), PolyLineVertices);
 
-					TArray<FVector2D>& Polygon = PolygonBatches.Emplace_GetRef();
-					Polygon.Reserve(PolyLineVertices.Num());
-					Algo::Transform(PolyLineVertices, Polygon, [](const FVector& Vertex) { return FVector2D(Vertex); });
-				}
+				TArray<FVector2D>& Polygon = PolygonBatches.Emplace_GetRef();
+				Polygon.Reserve(PolyLineVertices.Num());
+				Algo::Transform(PolyLineVertices, Polygon, [](const FVector& Vertex) { return FVector2D(Vertex); });
+			}
 
-				for (const TArray<FVector2D>& Polygon : PolygonBatches)
-				{
-					WaterQuadTree.AddLake(Polygon, LakeBounds, WaterBodyRenderDataIndex);
+			for (const TArray<FVector2D>& Polygon : PolygonBatches)
+			{
+				WaterQuadTree.AddLake(Polygon, LakeBounds, WaterBodyRenderDataIndex);
 
 #if !(UE_BUILD_SHIPPING || UE_BUILD_TEST)
-					if (!!CVarWaterMeshShowTileGenerationGeometry.GetValueOnGameThread())
+				if (!!CVarWaterMeshShowTileGenerationGeometry.GetValueOnGameThread())
+				{
+					float Z = SplineComp->GetLocationAtDistanceAlongSpline(0.0f, ESplineCoordinateSpace::World).Z;
+					int32 NumVertices = Polygon.Num();
+					for (int32 i = 0; i < NumVertices; i++)
 					{
-						float Z = SplineComp->GetLocationAtDistanceAlongSpline(0.0f, ESplineCoordinateSpace::World).Z;
-						int32 NumVertices = Polygon.Num();
-						for (int32 i = 0; i < NumVertices; i++)
-						{
-							const FVector2D& Point0 = Polygon[i];
-							const FVector2D& Point1 = Polygon[(i + 1) % NumVertices];
-							DrawDebugLine(GetWorld(), FVector(Point0.X, Point0.Y, Z), FVector(Point1.X, Point1.Y, Z), FColor::Green);
-						}
+						const FVector2D& Point0 = Polygon[i];
+						const FVector2D& Point1 = Polygon[(i + 1) % NumVertices];
+						DrawDebugLine(GetWorld(), FVector(Point0.X, Point0.Y, Z), FVector(Point1.X, Point1.Y, Z), FColor::Green);
 					}
-#endif // !(UE_BUILD_SHIPPING || UE_BUILD_TEST)
 				}
+#endif // !(UE_BUILD_SHIPPING || UE_BUILD_TEST)
+			}
 
-				break;
+			break;
+		}
+		case EWaterBodyType::Ocean:
+		{
+			TRACE_CPUPROFILER_EVENT_SCOPE(Ocean);
+
+			// Add ocean based on the ocean spline when there is no flood. Otherwise add ocean everywhere
+			if (bIsFlooded)
+			{
+				FBox OceanBounds = Actor->GetComponentsBoundingBox();
+				OceanBounds.Max.Z += WaterBodyComponent->GetMaxWaveHeight() + OceanFlood;
+				WaterQuadTree.AddWaterTilesInsideBounds(OceanBounds, WaterBodyRenderDataIndex);
 			}
-			case EWaterBodyType::Ocean:
+			else
 			{
-				TRACE_CPUPROFILER_EVENT_SCOPE(Ocean);
+				const UWaterSplineComponent* SplineComp = WaterBodyComponent->GetWaterSpline();
+				const int32 NumOriginalSplinePoints = SplineComp->GetNumberOfSplinePoints();
 
-				// Add ocean based on the ocean spline when there is no flood. Otherwise add ocean everywhere
-				if (bIsFlooded)
+				// Skip oceans with less than 3 spline points
+				if (NumOriginalSplinePoints < 3)
 				{
-					FBox OceanBounds = Actor->GetComponentsBoundingBox();
-					OceanBounds.Max.Z += WaterBodyComponent->GetMaxWaveHeight() + OceanFlood;
-					WaterQuadTree.AddWaterTilesInsideBounds(OceanBounds, WaterBodyRenderDataIndex);
+					break;
 				}
-				else
-				{
-					const UWaterSplineComponent* SplineComp = WaterBodyComponent->GetWaterSpline();
-					const int32 NumOriginalSplinePoints = SplineComp->GetNumberOfSplinePoints();
-
-					// Skip oceans with less than 3 spline points
-					if (NumOriginalSplinePoints < 3)
-					{
-						break;
-					}
 
-					TArray<FVector2D> Polygon;
+				TArray<FVector2D> Polygon;
 
-					TArray<FVector> PolyLineVertices;
-					SplineComp->ConvertSplineToPolyLine(ESplineCoordinateSpace::World, FMath::Square(CVarWaterSplineResampleMaxDistance.GetValueOnGameThread()), PolyLineVertices);
+				TArray<FVector> PolyLineVertices;
+				SplineComp->ConvertSplineToPolyLine(ESplineCoordinateSpace::World, FMath::Square(CVarWaterSplineResampleMaxDistance.GetValueOnGameThread()), PolyLineVertices);
 
-					Polygon.Reserve(PolyLineVertices.Num());
-					Algo::Transform(PolyLineVertices, Polygon, [](const FVector& Vertex) { return FVector2D(Vertex); });
+				Polygon.Reserve(PolyLineVertices.Num());
+				Algo::Transform(PolyLineVertices, Polygon, [](const FVector& Vertex) { return FVector2D(Vertex); });
 
-					FBox OceanBounds = WaterBodyComponent->Bounds.GetBox();
-					OceanBounds.Max.Z += WaterBodyComponent->GetMaxWaveHeight();
+				FBox OceanBounds = WaterBodyComponent->Bounds.GetBox();
+				OceanBounds.Max.Z += WaterBodyComponent->GetMaxWaveHeight();
 
 #if !(UE_BUILD_SHIPPING || UE_BUILD_TEST)
-					if (!!CVarWaterMeshShowTileGenerationGeometry.GetValueOnGameThread())
+				if (!!CVarWaterMeshShowTileGenerationGeometry.GetValueOnGameThread())
+				{
+					float Z = SplineComp->GetLocationAtDistanceAlongSpline(0.0f, ESplineCoordinateSpace::World).Z;
+					int32 NumVertices = Polygon.Num();
+					for (int32 i = 0; i < NumVertices; i++)
 					{
-						float Z = SplineComp->GetLocationAtDistanceAlongSpline(0.0f, ESplineCoordinateSpace::World).Z;
-						int32 NumVertices = Polygon.Num();
-						for (int32 i = 0; i < NumVertices; i++)
-						{
-							const FVector2D& Point0 = Polygon[i];
-							const FVector2D& Point1 = Polygon[(i + 1) % NumVertices];
-							DrawDebugLine(GetWorld(), FVector(Point0.X, Point0.Y, Z), FVector(Point1.X, Point1.Y, Z), FColor::Blue);
-						}
-
-						DrawDebugBox(GetWorld(), OceanBounds.GetCenter(), OceanBounds.GetExtent(), FColor::Blue);
+						const FVector2D& Point0 = Polygon[i];
+						const FVector2D& Point1 = Polygon[(i + 1) % NumVertices];
+						DrawDebugLine(GetWorld(), FVector(Point0.X, Point0.Y, Z), FVector(Point1.X, Point1.Y, Z), FColor::Blue);
 					}
-#endif //!(UE_BUILD_SHIPPING || UE_BUILD_TEST)
 
-					WaterQuadTree.AddOcean(Polygon, OceanBounds, WaterBodyRenderDataIndex);
+					DrawDebugBox(GetWorld(), OceanBounds.GetCenter(), OceanBounds.GetExtent(), FColor::Blue);
 				}
+#endif //!(UE_BUILD_SHIPPING || UE_BUILD_TEST)
 
-				// Place far mesh height just below the ocean level
-				FarMeshHeight = RenderData.SurfaceBaseHeight - WaterBodyComponent->GetMaxWaveHeight();
-				bHasOcean = true;
-
-				break;
-			}
-			case EWaterBodyType::Transition:
-				// Transitions dont require rendering
-				break;
-			default:
-				ensureMsgf(false, TEXT("This water body type is not implemented and will not produce any water tiles. "));
+				WaterQuadTree.AddOcean(Polygon, OceanBounds, WaterBodyRenderDataIndex);
 			}
+
+			// Place far mesh height just below the ocean level
+			FarMeshHeight = RenderData.SurfaceBaseHeight - WaterBodyComponent->GetMaxWaveHeight();
+			bHasOcean = true;
+
+			break;
+		}
+		case EWaterBodyType::Transition:
+			// Transitions dont require rendering
+			break;
+		default:
+			ensureMsgf(false, TEXT("This water body type is not implemented and will not produce any water tiles. "));
 		}
 
 		return true;
@@ -691,12 +592,6 @@ void UWaterMeshComponent::RebuildWaterMesh(float InTileSize, const FIntPoint& In
 	WaterQuadTree.Unlock(true);
 
 	MarkRenderStateDirty();
-
-	// Force another rebuild next frame if water body meshes are still compiling
-	if (bIsGPUQuadTree && bAnyWaterMeshesNotReady)
-	{
-		bNeedsRebuild = true;
-	}
 }
 
 void UWaterMeshComponent::Update()
diff --git a/Engine/Plugins/Experimental/Water/Source/Runtime/Private/WaterMeshSceneProxy.cpp b/Engine/Plugins/Experimental/Water/Source/Runtime/Private/WaterMeshSceneProxy.cpp
--- a/Engine/Plugins/Experimental/Water/Source/Runtime/Private/WaterMeshSceneProxy.cpp
+++ b/Engine/Plugins/Experimental/Water/Source/Runtime/Private/WaterMeshSceneProxy.cpp
@@ -4,19 +4,16 @@
 #include "MaterialShared.h"
 #include "WaterMeshComponent.h"
 #include "WaterZoneActor.h"
-#include "WaterViewExtension.h"
 #include "Materials/Material.h"
 #include "WaterUtils.h"
 #include "PrimitiveViewRelevance.h"
 #include "SceneManagement.h"
-#include "SceneRendering.h"
+#include "Engine/Engine.h"
 #include "Materials/MaterialRenderProxy.h"
 #include "Math/ColorList.h"
 #include "RayTracingInstance.h"
-#include "StaticMeshResources.h"
 #include "SceneInterface.h"
 #include "RenderGraphBuilder.h"
-#include "RenderGraphUtils.h"
 
 DECLARE_STATS_GROUP(TEXT("Water Mesh"), STATGROUP_WaterMesh, STATCAT_Advanced);
 
@@ -33,16 +30,6 @@ static TAutoConsoleVariable<int32> CVarWaterMeshLODMorphEnabled(
 	TEXT("If the smooth LOD morph is enabled. Turning this off may cause slight popping between LOD levels but will skip the calculations in the vertex shader, making it cheaper"),
 	ECVF_Scalability | ECVF_RenderThreadSafe);
 
-static TAutoConsoleVariable<int32> CVarWaterMeshGPUQuadTreeSuperSampling(
-	TEXT("r.Water.WaterMesh.GPUQuadTree.SuperSampling"), 2,
-	TEXT("Rasterizes water meshes into the GPU water quadtree at a higher resolution, reducing missing water tile artifacts near the edges of water bodies. Default: 2, Min: 1, Max : 8"),
-	ECVF_RenderThreadSafe);
-
-static TAutoConsoleVariable<int32> CVarWaterMeshGPUQuadTreeMultiSampling(
-	TEXT("r.Water.WaterMesh.GPUQuadTree.MultiSampling"), 4,
-	TEXT("Rasterizes water meshes into the GPU water quadtree with a multisampled rendertarget, reducing missing water tile artifacts near the edges of water bodies. Default: 4, Min: 1, Max : 8"),
-	ECVF_RenderThreadSafe);
-
 static TAutoConsoleVariable<int32> CVarWaterMeshOcclusionCullingMaxQueries(
 	TEXT("r.Water.WaterMesh.OcclusionCulling.MaxQueries"), 256,
 	TEXT("Maximum number of occlusion queries for the CPU water quadtree nodes. Using fewer queries than nodes will result in coarser culling."),
@@ -123,13 +110,6 @@ static TAutoConsoleVariable<float> CVarWaterMeshOcclusionCullExpandBoundsAmountX
 
 // ----------------------------------------------------------------------------------
 
-template <bool bWithWaterSelectionSupport>
-class TWaterVertexFactoryUserDataWrapper : public FOneFrameResource
-{
-public:
-	TWaterMeshUserData<bWithWaterSelectionSupport> UserData;
-};
-
 FWaterMeshSceneProxy::FWaterMeshSceneProxy(UWaterMeshComponent* Component)
 	: FPrimitiveSceneProxy(Component)
 	, MaterialRelevance(Component->GetWaterMaterialRelevance(GetScene().GetFeatureLevel()))
@@ -146,16 +126,12 @@ FWaterMeshSceneProxy::FWaterMeshSceneProxy(UWaterMeshComponent* Component)
 	}
 
 	int32 NumQuads = (int32)FMath::Pow(2.0f, (float)Component->GetTessellationFactor());
-	NumQuadsLOD0 = NumQuads;
 
 	WaterVertexFactories.Reserve(WaterQuadTree.GetTreeDepth());
-	WaterVertexFactoriesIndirectDraw.Reserve(WaterQuadTree.GetTreeDepth());
 	for (uint8 i = 0; i < WaterQuadTree.GetTreeDepth(); i++)
 	{
-		WaterVertexFactories.Add(new FWaterVertexFactoryType(GetScene().GetFeatureLevel(), NumQuads, LODScale));
+		WaterVertexFactories.Add(new WaterVertexFactoryType(GetScene().GetFeatureLevel(), NumQuads, LODScale));
 		BeginInitResource(WaterVertexFactories.Last());
-		WaterVertexFactoriesIndirectDraw.Add(new FWaterVertexFactoryIndirectDrawType(GetScene().GetFeatureLevel(), NumQuads, LODScale));
-		BeginInitResource(WaterVertexFactoriesIndirectDraw.Last());
 
 		NumQuads /= 2;
 
@@ -167,13 +143,12 @@ FWaterMeshSceneProxy::FWaterMeshSceneProxy(UWaterMeshComponent* Component)
 	}
 
 	WaterVertexFactories.Shrink();
-	WaterVertexFactoriesIndirectDraw.Shrink();
 	DensityCount = WaterVertexFactories.Num();
 
 	const int32 TotalLeafNodes = WaterQuadTree.GetMaxLeafCount();
-	WaterInstanceDataBuffers = new FWaterInstanceDataBuffersType(TotalLeafNodes);
+	WaterInstanceDataBuffers = new WaterInstanceDataBuffersType(TotalLeafNodes);
 
-	WaterMeshUserDataBuffers = new FWaterMeshUserDataBuffersType(WaterInstanceDataBuffers);
+	WaterMeshUserDataBuffers = new WaterMeshUserDataBuffersType(WaterInstanceDataBuffers);
 
 	WaterQuadTree.BuildMaterialIndices();
 
@@ -187,48 +162,15 @@ FWaterMeshSceneProxy::FWaterMeshSceneProxy(UWaterMeshComponent* Component)
 	RayTracingWaterData.SetNum(DensityCount);
 #endif
 
-	// Initialize Z bounds needed for GPU driven rendering
-	if (WaterQuadTree.IsGPUQuadTree())
-	{
-		WaterQuadTreeMinHeight = DBL_MAX;
-		WaterQuadTreeMaxHeight = -DBL_MAX;
-
-		const TArrayView<const FWaterBodyRenderData> WaterBodyRenderData = WaterQuadTree.GetWaterBodyRenderData();
-		for (const FWaterBodyRenderData& RenderData : WaterBodyRenderData)
-		{
-			WaterQuadTreeMinHeight = FMath::Min(WaterQuadTreeMinHeight, RenderData.BoundsMinZ - RenderData.MaxWaveHeight);
-			WaterQuadTreeMaxHeight = FMath::Max(WaterQuadTreeMaxHeight, RenderData.BoundsMaxZ + RenderData.MaxWaveHeight);
-		}
-
-		if (WaterQuadTreeMinHeight == DBL_MAX)
-		{
-			WaterQuadTreeMinHeight = 0.0;
-			WaterQuadTreeMaxHeight = 0.0;
-		}
-		WaterQuadTreeMinHeight -= 1.0;
-		WaterQuadTreeMaxHeight += 1.0;
-	}
-
-	// Always do CPU occlusion queries, even if this is a GPU quadtree. The GPU quadtree still potentially uses the far mesh which is CPU driven.
 	const int32 MaxQueries = CVarWaterMeshOcclusionCullingMaxQueries.GetValueOnGameThread();
 	const bool bIncludeFarMeshOcclusionQueries = CVarWaterMeshOcclusionCullingIncludeFarMesh.GetValueOnGameThread() != 0;
 	OcclusionCullingBounds = WaterQuadTree.ComputeNodeBounds(MaxQueries, CVarWaterMeshOcclusionCullExpandBoundsAmountXY.GetValueOnGameThread(), bIncludeFarMeshOcclusionQueries, &OcclusionResultsFarMeshOffset);
 	EmptyOcclusionCullingBounds.Add(WaterQuadTree.GetBoundsIncludingFarMesh());
-	// If this is a GPU quadtree, the CPU root node will have an invalid bounding box, so derive conservative bounds now
-	if (WaterQuadTree.IsGPUQuadTree() && !OcclusionCullingBounds.IsEmpty())
-	{
-		OcclusionCullingBounds[0] = FBox(FVector(WaterQuadTree.GetTileRegion().Min, WaterQuadTreeMinHeight), FVector(WaterQuadTree.GetTileRegion().Max, WaterQuadTreeMaxHeight));
-	}
 }
 
 FWaterMeshSceneProxy::~FWaterMeshSceneProxy()
 {
-	for (FWaterVertexFactoryType* WaterFactory : WaterVertexFactories)
-	{
-		WaterFactory->ReleaseResource();
-		delete WaterFactory;
-	}
-	for (FWaterVertexFactoryIndirectDrawType* WaterFactory : WaterVertexFactoriesIndirectDraw)
+	for (WaterVertexFactoryType* WaterFactory : WaterVertexFactories)
 	{
 		WaterFactory->ReleaseResource();
 		delete WaterFactory;
@@ -253,33 +195,6 @@ FWaterMeshSceneProxy::~FWaterMeshSceneProxy()
 void FWaterMeshSceneProxy::CreateRenderThreadResources(FRHICommandListBase& RHICmdList)
 {
 	SceneProxyCreatedFrameNumberRenderThread = GFrameNumberRenderThread;
-
-	if (WaterQuadTree.IsGPUQuadTree())
-	{
-		FWaterMeshGPUWork::FCallback Callback;
-		Callback.Proxy = this;
-		Callback.Function = [this](FRDGBuilder& GraphBuilder)
-		{
-			if (bNeedToTraverseGPUQuadTree)
-			{
-				if (!QuadTreeGPU.IsInitialized())
-				{
-					BuildGPUQuadTree(GraphBuilder);
-				}
-				QuadTreeGPU.Traverse(GraphBuilder, WaterQuadTreeGPUTraverseParams);
-				bNeedToTraverseGPUQuadTree = false;
-			}
-		};
-		GWaterMeshGPUWork.Callbacks.Add(MoveTemp(Callback));
-	}
-}
-
-void FWaterMeshSceneProxy::DestroyRenderThreadResources()
-{
-	if (WaterQuadTree.IsGPUQuadTree())
-	{
-		GWaterMeshGPUWork.Callbacks.RemoveAllSwap([this](const FWaterMeshGPUWork::FCallback& Callback) { return Callback.Proxy == this; });
-	}
 }
 
 FWaterMeshSceneProxy::FWaterLODParams FWaterMeshSceneProxy::GetWaterLODParams(const FVector& Position) const
@@ -310,7 +225,7 @@ void FWaterMeshSceneProxy::GetDynamicMeshElements(const TArray<const FSceneView*
 	FRHICommandListBase& RHICmdList = Collector.GetRHICommandList();
 
 	// The water render groups we have to render for this batch : 
-	TArray<EWaterMeshRenderGroupType, TInlineAllocator<FWaterVertexFactoryType::NumRenderGroups>> BatchRenderGroups;
+	TArray<EWaterMeshRenderGroupType, TInlineAllocator<WaterVertexFactoryType::NumRenderGroups>> BatchRenderGroups;
 	// By default, render all water tiles : 
 	BatchRenderGroups.Add(EWaterMeshRenderGroupType::RG_RenderWaterTiles);
 
@@ -344,208 +259,7 @@ void FWaterMeshSceneProxy::GetDynamicMeshElements(const TArray<const FSceneView*
 		Collector.RegisterOneFrameMaterialProxy(WireframeMaterialInstance);
 	}
 
-	const int32 NumWaterMaterials = WaterQuadTree.GetWaterMaterials().Num();
-	const int32 NumBuckets = NumWaterMaterials * DensityCount;
-
-	// If this is a GPU quadtree, we are using GPU driven rendering to first create a quadtree on the GPU (once) and then traverse it every frame to build a list of indirect draws
-	const bool bIsGPUQuadTree = WaterQuadTree.IsGPUQuadTree();
-	if (bIsGPUQuadTree && NumBuckets > 0)
-	{
-		bool bEncounteredISRView = false;
-		int32 InstanceFactor = 1;
-		int32 NumVisibleViews = 0;
-
-		TArray<FWaterQuadTreeGPU::FPerViewInfo> PerViewInfoArray;
-
-		// Gather per view data for all renderable views
-		for (int32 ViewIndex = 0; ViewIndex < Views.Num(); ViewIndex++)
-		{
-			const FSceneView* View = Views[ViewIndex];
-			if (!bEncounteredISRView && View->IsInstancedStereoPass())
-			{
-				bEncounteredISRView = true;
-				InstanceFactor = View->GetStereoPassInstanceFactor();
-			}
-
-			// skip gathering visible tiles from instanced right eye views
-			if ((VisibilityMap & (1 << ViewIndex)) && (!bEncounteredISRView || View->IsPrimarySceneView()))
-			{
-				const FVector QuadTreePosition = FVector(WaterQuadTree.GetTileRegion().Min, WaterQuadTreeMinHeight);
-				const FVector PreViewTranslation = View->ViewMatrices.GetPreViewTranslation();
-				const FViewInfo* ViewInfo = View->bIsViewInfo ? reinterpret_cast<const FViewInfo*>(View) : nullptr;
-
-				FWaterQuadTreeGPU::FPerViewInfo& PerViewInfo = PerViewInfoArray.AddDefaulted_GetRef();
-				PerViewInfo.QuadTreePositionTranslatedWorldSpace = FVector3f(QuadTreePosition + PreViewTranslation);
-				PerViewInfo.ObserverPositionTranslatedWorldSpace = FVector3f(View->ViewMatrices.GetViewOrigin() + PreViewTranslation);
-				PerViewInfo.TranslatedWorldToClip = FMatrix44f(View->ViewMatrices.GetTranslatedViewProjectionMatrix());
-				PerViewInfo.ViewToClip = FMatrix44f(View->ViewMatrices.GetProjectionMatrix());
-				PerViewInfo.ViewInfo = ViewInfo;
-				
-				++NumVisibleViews;
-			}
-		}
-
-		const int32 NumIndirectDrawCalls = NumBuckets * NumVisibleViews;
-		const int32 NumInstances = WaterQuadTree.GetMaxLeafCount() * NumVisibleViews;
-
-		// Allocate buffers for the indirect draws
-		struct FIndirectDrawResources
-		{
-			TRefCountPtr<FRDGPooledBuffer> IndirectArgs;
-			TRefCountPtr<FRDGPooledBuffer> InstanceDataOffsets;
-			TRefCountPtr<FRDGPooledBuffer> InstanceData0;
-			TRefCountPtr<FRDGPooledBuffer> InstanceData1;
-			TRefCountPtr<FRDGPooledBuffer> InstanceData2;
-		};
-		FIndirectDrawResources& IndirectDrawResources = Collector.AllocateOneFrameResource<FIndirectDrawResources>();
-		IndirectDrawResources.IndirectArgs = AllocatePooledBuffer(FRDGBufferDesc::CreateIndirectDesc<FRHIDrawIndexedIndirectParameters>(FMath::Max(1, NumIndirectDrawCalls)), TEXT("WaterQuadTree.IndirectArgsBuffer"));
-		IndirectDrawResources.InstanceDataOffsets = AllocatePooledBuffer(FRDGBufferDesc::CreateBufferDesc(sizeof(uint32), FMath::Max(1, NumInstances)), TEXT("WaterQuadTree.InstanceDataOffsetsBuffer"));
-		IndirectDrawResources.InstanceData0 = AllocatePooledBuffer(FRDGBufferDesc::CreateBufferDesc(sizeof(FVector4f), FMath::Max(1, NumInstances)), TEXT("WaterQuadTree.InstanceDataBuffer0"));
-		IndirectDrawResources.InstanceData1 = AllocatePooledBuffer(FRDGBufferDesc::CreateBufferDesc(sizeof(FVector4f), FMath::Max(1, NumInstances)), TEXT("WaterQuadTree.InstanceDataBuffer1"));
-		if (WITH_WATER_SELECTION_SUPPORT != 0)
-		{
-			IndirectDrawResources.InstanceData2 = AllocatePooledBuffer(FRDGBufferDesc::CreateBufferDesc(sizeof(FVector4f), FMath::Max(1, NumInstances)), TEXT("WaterQuadTree.InstanceDataBuffer2"));
-		}
-		
-
-		// Traverse quadtree. We actually only set up the parameters here; the actual traversal is done later in the frame in a callback invoked from FWaterViewExtension::PreRenderBasePass_RenderThread.
-		// That callback accesses these params, so we need to ensure that they are still valid at that point.
-		{
-			WaterQuadTreeGPUTraverseParams = {};
-			WaterQuadTreeGPUTraverseParams.OutIndirectArgsBuffer = IndirectDrawResources.IndirectArgs;
-			WaterQuadTreeGPUTraverseParams.OutInstanceDataOffsetsBuffer = IndirectDrawResources.InstanceDataOffsets;
-			WaterQuadTreeGPUTraverseParams.OutInstanceData0Buffer = IndirectDrawResources.InstanceData0;
-			WaterQuadTreeGPUTraverseParams.OutInstanceData1Buffer = IndirectDrawResources.InstanceData1;
-			WaterQuadTreeGPUTraverseParams.OutInstanceData2Buffer = IndirectDrawResources.InstanceData2;
-			WaterQuadTreeGPUTraverseParams.PerViewInfo = MoveTemp(PerViewInfoArray);
-			WaterQuadTreeGPUTraverseParams.CullingBounds = WaterInfoBounds.ShiftBy(-WaterQuadTree.GetTileRegion().Min);
-			WaterQuadTreeGPUTraverseParams.NumDensities = DensityCount;
-			WaterQuadTreeGPUTraverseParams.NumMaterials = NumWaterMaterials;
-			WaterQuadTreeGPUTraverseParams.NumViews = NumVisibleViews;
-			WaterQuadTreeGPUTraverseParams.NumQuadsLOD0 = NumQuadsLOD0;
-			WaterQuadTreeGPUTraverseParams.ForceCollapseDensityLevel = ForceCollapseDensityLevel;
-			WaterQuadTreeGPUTraverseParams.LeafSize = WaterQuadTree.GetLeafSize();
-			WaterQuadTreeGPUTraverseParams.LODScale = LODScale;
-			WaterQuadTreeGPUTraverseParams.DebugShowTile = CVarWaterMeshShowTileBounds.GetValueOnRenderThread();
-			WaterQuadTreeGPUTraverseParams.bWithWaterSelectionSupport = WITH_WATER_SELECTION_SUPPORT != 0;
-			WaterQuadTreeGPUTraverseParams.bLODMorphingEnabled = !!CVarWaterMeshLODMorphEnabled.GetValueOnRenderThread();
-
-			bNeedToTraverseGPUQuadTree = true;
-		}
-
-
-		using FWaterVertexFactoryUserDataWrapperType = TWaterVertexFactoryUserDataWrapper<WITH_WATER_SELECTION_SUPPORT>;
-		FWaterVertexFactoryUserDataWrapperType* UserDataWrappers[3] = {};
-		for (EWaterMeshRenderGroupType RenderGroup : BatchRenderGroups)
-		{
-			FWaterVertexFactoryUserDataWrapperType& UserDataWrapper = Collector.AllocateOneFrameResource<FWaterVertexFactoryUserDataWrapperType>();
-			UserDataWrapper.UserData.RenderGroupType = RenderGroup;
-			UserDataWrapper.UserData.IndirectInstanceDataOffsets = IndirectDrawResources.InstanceDataOffsets->GetSRV(RHICmdList, FRHIBufferSRVCreateInfo(PF_R32_UINT));
-			UserDataWrapper.UserData.IndirectInstanceData0 = IndirectDrawResources.InstanceData0->GetSRV(RHICmdList, FRHIBufferSRVCreateInfo(PF_A32B32G32R32F));
-			UserDataWrapper.UserData.IndirectInstanceData1 = IndirectDrawResources.InstanceData1->GetSRV(RHICmdList, FRHIBufferSRVCreateInfo(PF_A32B32G32R32F));
-			UserDataWrapper.UserData.IndirectInstanceData2 = (WITH_WATER_SELECTION_SUPPORT != 0) ? IndirectDrawResources.InstanceData2->GetSRV(RHICmdList, FRHIBufferSRVCreateInfo(PF_A32B32G32R32F)) : nullptr;
-			
-			
-			UserDataWrappers[(int32)RenderGroup] = &UserDataWrapper;
-		}
-
-		// Go through all buckets and issue one batched draw call per LOD level per material per view
-		int32 CompactViewIndex = 0;
-		for (int32 ViewIndex = 0; ViewIndex < Views.Num(); ViewIndex++)
-		{
-			// when rendering ISR, don't process the instanced view
-			if ((VisibilityMap & (1 << ViewIndex)) && (!bEncounteredISRView || Views[ViewIndex]->IsPrimarySceneView()))
-			{
-				TRACE_CPUPROFILER_EVENT_SCOPE(BucketsPerView);
-
-				for (int32 MaterialIndex = 0; MaterialIndex < NumWaterMaterials; ++MaterialIndex)
-				{
-					TRACE_CPUPROFILER_EVENT_SCOPE(MaterialBucket);
-					bool bMaterialDrawn = false;
-
-					for (int32 DensityIndex = 0; DensityIndex < DensityCount; ++DensityIndex)
-					{
-						const int32 BucketIndex = MaterialIndex * DensityCount + DensityIndex;
-
-						TRACE_CPUPROFILER_EVENT_SCOPE(DensityBucket);
-
-						const FMaterialRenderProxy* MaterialRenderProxy = (WireframeMaterialInstance != nullptr) ? WireframeMaterialInstance : WaterQuadTree.GetWaterMaterials()[MaterialIndex];
-						check(MaterialRenderProxy != nullptr);
-
-						bool bUseForDepthPass = false;
-
-						// If there's a valid material, use that to figure out the depth pass status
-						if (const FMaterial* BucketMaterial = MaterialRenderProxy->GetMaterialNoFallback(GetScene().GetFeatureLevel()))
-						{
-							// Preemptively turn off depth rendering for this mesh batch if the material doesn't need it
-							bUseForDepthPass = !BucketMaterial->GetShadingModels().HasShadingModel(MSM_SingleLayerWater) && !IsTranslucentOnlyBlendMode(*BucketMaterial);
-						}
-
-						bMaterialDrawn = true;
-						for (EWaterMeshRenderGroupType RenderGroup : BatchRenderGroups)
-						{
-							// Set up mesh batch
-							FMeshBatch& Mesh = Collector.AllocateMesh();
-							Mesh.bWireframe = bWireframe;
-							Mesh.VertexFactory = WaterVertexFactoriesIndirectDraw[DensityIndex];
-							Mesh.MaterialRenderProxy = MaterialRenderProxy;
-							Mesh.ReverseCulling = IsLocalToWorldDeterminantNegative();
-							Mesh.Type = PT_TriangleList;
-							Mesh.DepthPriorityGroup = SDPG_World;
-							Mesh.bCanApplyViewModeOverrides = false;
-							Mesh.bUseForMaterial = true;
-							Mesh.CastShadow = false;
-							// Preemptively turn off depth rendering for this mesh batch if the material doesn't need it
-							Mesh.bUseForDepthPass = bUseForDepthPass;
-							Mesh.bUseAsOccluder = false;
-
-#if WITH_WATER_SELECTION_SUPPORT
-							Mesh.bUseSelectionOutline = (RenderGroup == EWaterMeshRenderGroupType::RG_RenderSelectedWaterTilesOnly);
-							Mesh.bUseWireframeSelectionColoring = (RenderGroup == EWaterMeshRenderGroupType::RG_RenderSelectedWaterTilesOnly);
-#endif // WITH_WATER_SELECTION_SUPPORT
-
-							Mesh.Elements.SetNumZeroed(1);
-							{
-								TRACE_CPUPROFILER_EVENT_SCOPE_STR("Setup batch element");
-
-								// Set up one mesh batch element
-								FMeshBatchElement& BatchElement = Mesh.Elements[0];
-
-								// Set up for indirect draw
-								BatchElement.FirstIndex = 0;
-								BatchElement.NumPrimitives = 0; // Must be 0 to enable usage of IndirectArgsBuffer
-								BatchElement.IndirectArgsBuffer = IndirectDrawResources.IndirectArgs->GetRHI();
-								BatchElement.IndirectArgsOffset = (CompactViewIndex * NumBuckets + BucketIndex) * sizeof(FRHIDrawIndexedIndirectParameters);
-								BatchElement.UserData = (void*)&UserDataWrappers[(int32)RenderGroup]->UserData;
-								BatchElement.UserIndex = CompactViewIndex * NumBuckets + BucketIndex;
-
-								BatchElement.IndexBuffer = WaterVertexFactoriesIndirectDraw[DensityIndex]->IndexBuffer;
-								BatchElement.PrimitiveIdMode = PrimID_ForceZero;
-
-								// We need the uniform buffer of this primitive because it stores the proper value for the bOutputVelocity flag.
-								// The identity primitive uniform buffer simply stores false for this flag which leads to missing motion vectors.
-								BatchElement.PrimitiveUniformBuffer = GetUniformBuffer();
-							}
-
-							{
-								INC_DWORD_STAT(STAT_WaterDrawCalls);
-								TRACE_CPUPROFILER_EVENT_SCOPE(Collector.AddMesh);
-
-								Collector.AddMesh(ViewIndex, Mesh);
-							}
-						}
-					}
-
-					INC_DWORD_STAT_BY(STAT_WaterDrawnMats, (int32)bMaterialDrawn);
-				}
-
-				++CompactViewIndex;
-			}
-		}
-	}
-	
-	// Even if we have a GPU quadtree, we still go through the setup for the non-GPU path because the water quadtree might have an ocean far mesh which we render traditionally.
-	// If this is a GPU quadtree, traversing it on the CPU will only return the far mesh tiles (if present).
+	const int32 NumBuckets = WaterQuadTree.GetWaterMaterials().Num() * DensityCount;
 
 	TArray<FWaterQuadTree::FTraversalOutput, TInlineAllocator<4>> WaterInstanceDataPerView;
 
@@ -566,7 +280,7 @@ void FWaterMeshSceneProxy::GetDynamicMeshElements(const TArray<const FSceneView*
 		if ((VisibilityMap & (1 << ViewIndex)) && (!bEncounteredISRView || View->IsPrimarySceneView()))
 		{
 			const FVector ObserverPosition = View->ViewMatrices.GetViewOrigin();
-
+			
 			FWaterLODParams WaterLODParams = GetWaterLODParams(ObserverPosition);
 
 #if !(UE_BUILD_SHIPPING || UE_BUILD_TEST)
@@ -653,6 +367,7 @@ void FWaterMeshSceneProxy::GetDynamicMeshElements(const TArray<const FSceneView*
 			TRACE_CPUPROFILER_EVENT_SCOPE(BucketsPerView);
 
 			FWaterQuadTree::FTraversalOutput& WaterInstanceData = WaterInstanceDataPerView[TraversalIndex];
+			const int32 NumWaterMaterials = WaterQuadTree.GetWaterMaterials().Num();
 			TraversalIndex++;
 
 			for (int32 MaterialIndex = 0; MaterialIndex < NumWaterMaterials; ++MaterialIndex)
@@ -673,7 +388,7 @@ void FWaterMeshSceneProxy::GetDynamicMeshElements(const TArray<const FSceneView*
 					TRACE_CPUPROFILER_EVENT_SCOPE(DensityBucket);
 
 					const FMaterialRenderProxy* MaterialRenderProxy = (WireframeMaterialInstance != nullptr) ? WireframeMaterialInstance : WaterQuadTree.GetWaterMaterials()[MaterialIndex];
-					check(MaterialRenderProxy != nullptr);
+					check (MaterialRenderProxy != nullptr);
 
 					bool bUseForDepthPass = false;
 
@@ -731,7 +446,7 @@ void FWaterMeshSceneProxy::GetDynamicMeshElements(const TArray<const FSceneView*
 
 							// We need the uniform buffer of this primitive because it stores the proper value for the bOutputVelocity flag.
 							// The identity primitive uniform buffer simply stores false for this flag which leads to missing motion vectors.
-							BatchElement.PrimitiveUniformBuffer = GetUniformBuffer();
+							BatchElement.PrimitiveUniformBuffer = GetUniformBuffer(); 
 						}
 
 						{
@@ -760,7 +475,7 @@ void FWaterMeshSceneProxy::GetDynamicMeshElements(const TArray<const FSceneView*
 				const FWaterQuadTree::FStagingInstanceData& Data = WaterInstanceData.StagingInstanceData[Idx];
 				const int32 WriteIndex = WaterInstanceData.BucketInstanceCounts[Data.BucketIndex]++;
 
-				for (int32 StreamIdx = 0; StreamIdx < FWaterInstanceDataBuffersType::NumBuffers; ++StreamIdx)
+				for (int32 StreamIdx = 0; StreamIdx < WaterInstanceDataBuffersType::NumBuffers; ++StreamIdx)
 				{
 					TArrayView<FVector4f> BufferMemory = WaterInstanceDataBuffers->GetBufferMemory(StreamIdx);
 					for (int32 IdxMultipliedInstance = 0; IdxMultipliedInstance < InstanceFactor; ++IdxMultipliedInstance)
@@ -871,6 +586,13 @@ void FWaterMeshSceneProxy::SetupRayTracingInstances(FRHICommandListBase& RHICmdL
 	}
 }
 
+template <bool bWithWaterSelectionSupport>
+class TWaterVertexFactoryUserDataWrapper : public FOneFrameResource
+{
+public:
+	TWaterMeshUserData<bWithWaterSelectionSupport> UserData;
+};
+
 void FWaterMeshSceneProxy::GetDynamicRayTracingInstances(FRayTracingMaterialGatheringContext& Context, TArray<FRayTracingInstance>& OutRayTracingInstances)
 {
 	if (!HasWaterData() || !FWaterUtils::IsWaterMeshRenderingEnabled(/*bIsRenderThread = */true) || !CVarRayTracingGeometryWater.GetValueOnRenderThread())
@@ -1059,117 +781,3 @@ HHitProxy* FWaterMeshSceneProxy::CreateHitProxies(UPrimitiveComponent* Component
 	return nullptr;
 }
 #endif // WITH_WATER_SELECTION_SUPPORT
-
-void FWaterMeshSceneProxy::BuildGPUQuadTree(FRDGBuilder& GraphBuilder)
-{
-	auto BuildOrthoMatrix = [](float InOrthoWidth, float InOrthoHeight, float DepthRange)
-	{
-		check((int32)ERHIZBuffer::IsInverted);
-
-		const FMatrix::FReal OrthoWidth = InOrthoWidth / 2.0f;
-		const FMatrix::FReal OrthoHeight = InOrthoHeight / 2.0f;
-
-		const FMatrix::FReal NearPlane = 0.f;
-		const FMatrix::FReal FarPlane = DepthRange;
-
-		const FMatrix::FReal ZScale = 1.0f / (FarPlane - NearPlane);
-		const FMatrix::FReal ZOffset = 0;
-
-		return FReversedZOrthoMatrix(
-			OrthoWidth,
-			OrthoHeight,
-			ZScale,
-			ZOffset
-		);
-	};
-
-	// Create GPU array of water body render data
-	const TArrayView<const FWaterBodyRenderData> WaterBodyRenderData = WaterQuadTree.GetWaterBodyRenderData();
-	TArray<FWaterQuadTreeGPU::FWaterBodyRenderDataGPU> WaterBodyRenderDataGPU;
-	{
-		WaterBodyRenderDataGPU.Reserve(WaterBodyRenderData.Num());
-
-		for (const FWaterBodyRenderData& RenderData : WaterBodyRenderData)
-		{
-			FColor HitProxyColor = FColor(0, 0, 0, 0);
-#if WITH_WATER_SELECTION_SUPPORT
-			if (RenderData.HitProxy)
-			{
-				HitProxyColor = RenderData.HitProxy->Id.GetColor();
-				HitProxyColor.A = RenderData.bWaterBodySelected ? 255 : 0;
-			}
-#endif
-
-			FWaterQuadTreeGPU::FWaterBodyRenderDataGPU& RenderDataGPU = WaterBodyRenderDataGPU.AddDefaulted_GetRef();
-			RenderDataGPU.WaterBodyIndex = RenderData.WaterBodyIndex;
-			RenderDataGPU.MaterialIndex = RenderData.MaterialIndex;
-			RenderDataGPU.RiverToLakeMaterialIndex = RenderData.RiverToLakeMaterialIndex;
-			RenderDataGPU.RiverToOceanMaterialIndex = RenderData.RiverToOceanMaterialIndex;
-			RenderDataGPU.WaterBodyType = RenderData.WaterBodyType;
-			RenderDataGPU.HitProxyColorAndIsSelected = HitProxyColor.Bits;
-			RenderDataGPU.SurfaceBaseHeight = float(RenderData.SurfaceBaseHeight - WaterQuadTreeMinHeight);
-			RenderDataGPU.MaxWaveHeight = RenderData.MaxWaveHeight;
-		}
-	}
-
-	const FIntPoint Resolution = WaterQuadTree.GetResolution();
-	const FBox2D TileRegion = WaterQuadTree.GetTileRegion();
-	const float LeafSize = WaterQuadTree.GetLeafSize();
-
-	const FVector2D Center = TileRegion.Min + FVector2D(Resolution) * (double)LeafSize * 0.5;
-	const FVector QuadTreeCorner = FVector(TileRegion.Min, WaterQuadTreeMinHeight); // Make everything relative to this to avoid LWC precision issues
-	const float WaterQuadTreeDepthRange = WaterQuadTreeMaxHeight - WaterQuadTreeMinHeight;
-	const float RcpWaterQuadTreeDepthRange = 1.0f / WaterQuadTreeDepthRange;
-
-	// Create view/projection matrices
-	const FVector ViewLocation = FVector(Center, WaterQuadTreeMaxHeight) - QuadTreeCorner;
-	const FVector LookAt = ViewLocation - FVector(0.0, 0.0, 1.0);
-	const FMatrix44f ViewMatrix = FMatrix44f(FLookAtMatrix(ViewLocation, LookAt, FVector(0.f, -1.f, 0.f)));
-	const FMatrix44f ProjectionMatrix = FMatrix44f(BuildOrthoMatrix(Resolution.X * LeafSize, Resolution.Y * LeafSize, WaterQuadTreeDepthRange));
-	const FMatrix44f ViewProjection = ViewMatrix * ProjectionMatrix;
-
-	// Create array of water body draws
-	TArray<FWaterQuadTreeGPU::FDraw> Draws;
-	{
-		const TArray<FWaterBodyQuadTreeRasterInfo>& WaterBodyRasterInfos = WaterQuadTree.GetWaterBodyRasterInfos();
-		for (const FWaterBodyQuadTreeRasterInfo& Info : WaterBodyRasterInfos)
-		{
-			const FStaticMeshLODResources& LODResource = Info.RenderData->LODResources[0];
-			for (const FStaticMeshSection& MeshSection : LODResource.Sections)
-			{
-				if (LODResource.IndexBuffer.GetRHI())
-				{
-					// Make relative to quadtree to avoid LWC precision issues
-					FTransform LocalToQuadTreeWorld = Info.LocalToWorld;
-					LocalToQuadTreeWorld.SetTranslation(LocalToQuadTreeWorld.GetTranslation() - QuadTreeCorner);
-
-					const FWaterBodyRenderData& WBRenderData = WaterBodyRenderData[Info.WaterBodyRenderDataIndex];
-
-					FWaterQuadTreeGPU::FDraw& Draw = Draws.AddDefaulted_GetRef();
-					Draw.Transform = FMatrix44f(LocalToQuadTreeWorld.ToMatrixWithScale()) * ViewProjection;
-					Draw.VertexBuffer = LODResource.VertexBuffers.PositionVertexBuffer.GetRHI();
-					Draw.IndexBuffer = LODResource.IndexBuffer.GetRHI();
-					Draw.FirstIndex = MeshSection.FirstIndex;
-					Draw.NumPrimitives = MeshSection.NumTriangles;
-					Draw.BaseVertexIndex = 0;
-					Draw.NumVertices = MeshSection.MaxVertexIndex - MeshSection.MinVertexIndex + 1;
-					Draw.WaterBodyRenderDataIndex = Info.WaterBodyRenderDataIndex;
-					Draw.Priority = Info.Priority;
-					Draw.MinZ = float(WBRenderData.BoundsMinZ - WaterQuadTreeMinHeight) * RcpWaterQuadTreeDepthRange;
-					Draw.MaxZ = float(WBRenderData.BoundsMaxZ + WBRenderData.MaxWaveHeight - WaterQuadTreeMinHeight) * RcpWaterQuadTreeDepthRange;
-					Draw.MaxWaveHeight = WBRenderData.MaxWaveHeight * RcpWaterQuadTreeDepthRange;
-					Draw.bIsRiver = Info.bIsRiver;
-				}
-			}
-		}
-	}
-
-	FWaterQuadTreeGPU::FInitParams Params;
-	Params.WaterBodyRenderData = WaterBodyRenderDataGPU;
-	Params.RequestedQuadTreeResolution = Resolution;
-	Params.SuperSamplingFactor = FMath::Clamp(CVarWaterMeshGPUQuadTreeSuperSampling.GetValueOnRenderThread(), 1, 8);
-	Params.NumMSAASamples = FMath::Clamp(CVarWaterMeshGPUQuadTreeMultiSampling.GetValueOnRenderThread(), 1, 8);
-	Params.CaptureDepthRange = WaterQuadTreeDepthRange;
-
-	QuadTreeGPU.Init(GraphBuilder, Params, Draws);
-}
diff --git a/Engine/Plugins/Experimental/Water/Source/Runtime/Private/WaterQuadTree.cpp b/Engine/Plugins/Experimental/Water/Source/Runtime/Private/WaterQuadTree.cpp
--- a/Engine/Plugins/Experimental/Water/Source/Runtime/Private/WaterQuadTree.cpp
+++ b/Engine/Plugins/Experimental/Water/Source/Runtime/Private/WaterQuadTree.cpp
@@ -451,15 +451,13 @@ bool FWaterQuadTree::FNode::QueryBoundsAtLocation(const FNodeData& InNodeData, c
 	return ChildCount == 0;
 }
 
-void FWaterQuadTree::InitTree(const FBox2D& InBounds, float InTileSize, FIntPoint InExtentInTiles, bool bInIsGPUQuadTree)
+void FWaterQuadTree::InitTree(const FBox2D& InBounds, float InTileSize, FIntPoint InExtentInTiles)
 {
 	ensure(InBounds.GetArea() > 0.0f);
 	ensure(InTileSize > 0.0f);
 	ensure(InExtentInTiles.X > 0);
 	ensure(InExtentInTiles.Y > 0);
 
-	bIsGPUQuadTree = bInIsGPUQuadTree;
-
 	FarMeshData.Clear();
 
 	// Maximum number of allocated leaf nodes for this config
@@ -474,19 +472,11 @@ void FWaterQuadTree::InitTree(const FBox2D& InBounds, float InTileSize, FIntPoin
 
 	TileRegion = InBounds;
 
-	WaterBodyRasterInfos.Reset();
 	BreadthFirstOrder.Reset();
 
 	// Allocate theoretical max, shrink later in Lock()
 	// This is so that the node array doesn't move in memory while inserting
-	if (!bIsGPUQuadTree)
-	{
-		NodeData.Nodes.Empty((float)(FMath::Square(RootDim) * 4) / 3.0f);
-	}
-	else
-	{
-		NodeData.Nodes.Empty(1);
-	}
+	NodeData.Nodes.Empty((float)(FMath::Square(RootDim) * 4) / 3.0f);
 
 	// Add defaulted water body render data to slot 0. This is the "null" render data, pointed to by all newly created nodes. Has lowest priority so it will always be overwritten
 	NodeData.WaterBodyRenderData.Empty(1);
@@ -645,14 +635,12 @@ void FWaterQuadTree::Unlock(bool bPruneRedundantNodes)
 void FWaterQuadTree::AddWaterTilesInsideBounds(const FBox& InBounds, uint32 InWaterBodyIndex)
 {
 	check(!bIsReadOnly);
-	check(!bIsGPUQuadTree);
 	NodeData.Nodes[0].AddNodes(NodeData, FBox(FVector(TileRegion.Min, 0.0f), FVector(TileRegion.Max, 0.0f)),  InBounds, InWaterBodyIndex, TreeDepth, 0);
 }
 
 void FWaterQuadTree::AddOcean(const TArray<FVector2D>& InPoly, const FBox& InOceanBounds, uint32 InWaterBodyIndex)
 {
 	check(!bIsReadOnly);
-	check(!bIsGPUQuadTree);
 	const FBox2D OceanBounds(FVector2D(InOceanBounds.Min), FVector2D(InOceanBounds.Max));
 	AddOceanRecursive(InPoly, OceanBounds, FVector2D(InOceanBounds.Min.Z, InOceanBounds.Max.Z), true, TreeDepth * 2, InWaterBodyIndex);
 }
@@ -660,7 +648,6 @@ void FWaterQuadTree::AddOcean(const TArray<FVector2D>& InPoly, const FBox& InOce
 void FWaterQuadTree::AddLake(const TArray<FVector2D>& InPoly, const FBox& InLakeBounds, uint32 InWaterBodyIndex)
 {
 	check(!bIsReadOnly);
-	check(!bIsGPUQuadTree);
 	const FBox2D LakeBounds(FVector2D(NodeData.Nodes[0].Bounds.Min), FVector2D(NodeData.Nodes[0].Bounds.Max));
 	AddLakeRecursive(InPoly, LakeBounds, FVector2D(InLakeBounds.Min.Z, InLakeBounds.Max.Z), true, TreeDepth * 2, InWaterBodyIndex);
 }
@@ -756,10 +743,7 @@ void FWaterQuadTree::BuildWaterTileInstanceData(const FTraversalDesc& InTraversa
 	TRACE_CPUPROFILER_EVENT_SCOPE(BuildWaterTileInstanceData);
 	check(bIsReadOnly);
 
-	if (!bIsGPUQuadTree)
-	{
-		NodeData.Nodes[0].SelectLOD(NodeData, TreeDepth, InTraversalDesc, Output);
-	}
+	NodeData.Nodes[0].SelectLOD(NodeData, TreeDepth, InTraversalDesc, Output);
 
 	// Append Far Mesh tiles
 	if (FarMeshData.InstanceData.Num() > 0 && FarMeshData.MaterialIndex != INDEX_NONE)
diff --git a/Engine/Plugins/Experimental/Water/Source/Runtime/Private/WaterQuadTreeGPU.cpp b/Engine/Plugins/Experimental/Water/Source/Runtime/Private/WaterQuadTreeGPU.cpp
deleted file mode 100644
diff --git a/Engine/Plugins/Experimental/Water/Source/Runtime/Private/WaterQuadTreeGPU.h b/Engine/Plugins/Experimental/Water/Source/Runtime/Private/WaterQuadTreeGPU.h
deleted file mode 100644
diff --git a/Engine/Plugins/Experimental/Water/Source/Runtime/Private/WaterVertexFactory.cpp b/Engine/Plugins/Experimental/Water/Source/Runtime/Private/WaterVertexFactory.cpp
--- a/Engine/Plugins/Experimental/Water/Source/Runtime/Private/WaterVertexFactory.cpp
+++ b/Engine/Plugins/Experimental/Water/Source/Runtime/Private/WaterVertexFactory.cpp
@@ -13,31 +13,18 @@ IMPLEMENT_GLOBAL_SHADER_PARAMETER_STRUCT(FWaterVertexFactoryRaytracingParameters
 /**
  * Shader parameters for water vertex factory.
  */
-template <bool bWithWaterSelectionSupport, bool bIndirectDraws>
+template <bool bWithWaterSelectionSupport>
 class TWaterVertexFactoryShaderParameters : public FVertexFactoryShaderParameters
 {
-	using WaterVertexFactoryShaderParametersType = TWaterVertexFactoryShaderParameters<bWithWaterSelectionSupport, bIndirectDraws>;
-	DECLARE_TYPE_LAYOUT(WaterVertexFactoryShaderParametersType, NonVirtual);
+	DECLARE_TYPE_LAYOUT(TWaterVertexFactoryShaderParameters<bWithWaterSelectionSupport>, NonVirtual);
 
 public:
-	using WaterVertexFactoryType = TWaterVertexFactory<bWithWaterSelectionSupport, bIndirectDraws>;
+	using WaterVertexFactoryType = TWaterVertexFactory<bWithWaterSelectionSupport>;
 	using WaterMeshUserDataType = TWaterMeshUserData<bWithWaterSelectionSupport>;
 	using WaterInstanceDataBuffersType = TWaterInstanceDataBuffers<bWithWaterSelectionSupport>;
 
 	void Bind(const FShaderParameterMap& ParameterMap)
 	{
-		if (bIndirectDraws)
-		{
-			InstanceDataOffsetsBufferParameter.Bind(ParameterMap, TEXT("InstanceDataOffsetsBuffer"));
-			InstanceData0BufferParameter.Bind(ParameterMap, TEXT("InstanceData0Buffer"));
-			InstanceData1BufferParameter.Bind(ParameterMap, TEXT("InstanceData1Buffer"));
-			if (bWithWaterSelectionSupport)
-			{
-				InstanceData2BufferParameter.Bind(ParameterMap, TEXT("InstanceData2Buffer"));
-			}
-			
-			DrawBucketIndexParameter.Bind(ParameterMap, TEXT("DrawBucketIndex"));
-		}
 	}
 
 	void GetElementShaderBindings(
@@ -57,6 +44,8 @@ public:
 
 		const WaterInstanceDataBuffersType* InstanceDataBuffers = WaterMeshUserData->InstanceDataBuffers;
 
+		const int32 InstanceOffsetValue = BatchElement.UserIndex;
+
 		ShaderBindings.Add(Shader->GetUniformBufferParameter<FWaterVertexFactoryParameters>(), VertexFactory->GetWaterVertexFactoryUniformBuffer(WaterMeshUserData->RenderGroupType));
 
 #if RHI_RAYTRACING
@@ -66,7 +55,7 @@ public:
 		}
 #endif
 
-		if (VertexStreams.Num() > 0 && !bIndirectDraws)
+		if (VertexStreams.Num() > 0)
 		{
 			for (int32 i = 0; i < WaterInstanceDataBuffersType::NumBuffers; ++i)
 			{
@@ -78,52 +67,24 @@ public:
 				InstanceInputStream->VertexBuffer = InstanceDataBuffers->GetBuffer(i);
 			}
 
-			const int32 InstanceOffsetValue = BatchElement.UserIndex;
 			if (InstanceOffsetValue > 0)
 			{
 				VertexFactory->OffsetInstanceStreams(InstanceOffsetValue, InputStreamType, VertexStreams);
 			}
 		}
-
-		if (bIndirectDraws)
-		{
-			const uint32 DrawBucketIndex = BatchElement.UserIndex;
-			ShaderBindings.Add(InstanceDataOffsetsBufferParameter, WaterMeshUserData->IndirectInstanceDataOffsets);
-			ShaderBindings.Add(InstanceData0BufferParameter, WaterMeshUserData->IndirectInstanceData0);
-			ShaderBindings.Add(InstanceData1BufferParameter, WaterMeshUserData->IndirectInstanceData1);
-			if (bWithWaterSelectionSupport)
-			{
-				ShaderBindings.Add(InstanceData2BufferParameter, WaterMeshUserData->IndirectInstanceData2);
-			}
-			ShaderBindings.Add(DrawBucketIndexParameter, DrawBucketIndex);
-		}
 	}
-
-private:
-	LAYOUT_FIELD(FShaderResourceParameter, InstanceDataOffsetsBufferParameter);
-	LAYOUT_FIELD(FShaderResourceParameter, InstanceData0BufferParameter);
-	LAYOUT_FIELD(FShaderResourceParameter, InstanceData1BufferParameter);
-	LAYOUT_FIELD(FShaderResourceParameter, InstanceData2BufferParameter);
-	LAYOUT_FIELD(FShaderParameter, DrawBucketIndexParameter);
 };
 
 // ----------------------------------------------------------------------------------
 
-using FWaterVertexFactoryParametersNoSelectionNoIndirect = TWaterVertexFactoryShaderParameters</*bWithWaterSelectionSupport = */ false, /*bIndirectDraws = */ false>;
-using FWaterVertexFactoryParametersNoSelectionWithIndirect = TWaterVertexFactoryShaderParameters</*bWithWaterSelectionSupport = */ false, /*bIndirectDraws = */ true>;
-using FWaterVertexFactoryNoSelectionNoIndirect = TWaterVertexFactory</*bWithWaterSelectionSupport = */ false, /*bIndirectDraws = */ false>;
-using FWaterVertexFactoryNoSelectionWithIndirect = TWaterVertexFactory</*bWithWaterSelectionSupport = */ false, /*bIndirectDraws = */ true>;
-
 // Always implement the basic vertex factory so that it's there for both editor and non-editor builds :
-IMPLEMENT_TEMPLATE_TYPE_LAYOUT(template<>, FWaterVertexFactoryParametersNoSelectionNoIndirect);
-IMPLEMENT_TEMPLATE_TYPE_LAYOUT(template<>, FWaterVertexFactoryParametersNoSelectionWithIndirect);
-IMPLEMENT_VERTEX_FACTORY_PARAMETER_TYPE(FWaterVertexFactoryNoSelectionNoIndirect, SF_Vertex, FWaterVertexFactoryParametersNoSelectionNoIndirect);
-IMPLEMENT_VERTEX_FACTORY_PARAMETER_TYPE(FWaterVertexFactoryNoSelectionWithIndirect, SF_Vertex, FWaterVertexFactoryParametersNoSelectionWithIndirect);
+IMPLEMENT_TEMPLATE_TYPE_LAYOUT(template<>, TWaterVertexFactoryShaderParameters</*bWithWaterSelectionSupport = */ false>);
+IMPLEMENT_VERTEX_FACTORY_PARAMETER_TYPE(TWaterVertexFactory</*bWithWaterSelectionSupport = */ false>, SF_Vertex, TWaterVertexFactoryShaderParameters</*bWithWaterSelectionSupport = */ false>);
 #if RHI_RAYTRACING
-IMPLEMENT_VERTEX_FACTORY_PARAMETER_TYPE(FWaterVertexFactoryNoSelectionNoIndirect, SF_Compute, FWaterVertexFactoryParametersNoSelectionNoIndirect);
-IMPLEMENT_VERTEX_FACTORY_PARAMETER_TYPE(FWaterVertexFactoryNoSelectionNoIndirect, SF_RayHitGroup, FWaterVertexFactoryParametersNoSelectionNoIndirect);
+IMPLEMENT_VERTEX_FACTORY_PARAMETER_TYPE(TWaterVertexFactory</*bWithWaterSelectionSupport = */ false>, SF_Compute, TWaterVertexFactoryShaderParameters</*bWithWaterSelectionSupport = */ false>);
+IMPLEMENT_VERTEX_FACTORY_PARAMETER_TYPE(TWaterVertexFactory</*bWithWaterSelectionSupport = */ false>, SF_RayHitGroup, TWaterVertexFactoryShaderParameters</*bWithWaterSelectionSupport = */ false>);
 #endif // RHI_RAYTRACING
-IMPLEMENT_TEMPLATE_VERTEX_FACTORY_TYPE(template<>, FWaterVertexFactoryNoSelectionNoIndirect, "/Plugin/Water/Private/WaterMeshVertexFactory.ush",
+IMPLEMENT_TEMPLATE_VERTEX_FACTORY_TYPE(template<>, TWaterVertexFactory</*bWithWaterSelectionSupport = */ false>, "/Plugin/Water/Private/WaterMeshVertexFactory.ush",
 	  EVertexFactoryFlags::UsedWithMaterials
 	| EVertexFactoryFlags::SupportsDynamicLighting
 	| EVertexFactoryFlags::SupportsPrecisePrevWorldPos
@@ -132,33 +93,17 @@ IMPLEMENT_TEMPLATE_VERTEX_FACTORY_TYPE(template<>, FWaterVertexFactoryNoSelectio
 	| EVertexFactoryFlags::SupportsRayTracingDynamicGeometry
 	| EVertexFactoryFlags::SupportsPSOPrecaching
 );
-IMPLEMENT_TEMPLATE_VERTEX_FACTORY_TYPE(template<>, FWaterVertexFactoryNoSelectionWithIndirect, "/Plugin/Water/Private/WaterMeshVertexFactory.ush",
-	EVertexFactoryFlags::UsedWithMaterials
-	| EVertexFactoryFlags::SupportsDynamicLighting
-	| EVertexFactoryFlags::SupportsPrecisePrevWorldPos
-	| EVertexFactoryFlags::SupportsPrimitiveIdStream
-	| EVertexFactoryFlags::SupportsRayTracing
-	| EVertexFactoryFlags::SupportsRayTracingDynamicGeometry
-	| EVertexFactoryFlags::SupportsPSOPrecaching
-);
 
 #if WITH_WATER_SELECTION_SUPPORT
 
-using FWaterVertexFactoryParametersWithSelectionNoIndirect = TWaterVertexFactoryShaderParameters</*bWithWaterSelectionSupport = */ true, /*bIndirectDraws = */ false>;
-using FWaterVertexFactoryParametersWithSelectionWithIndirect = TWaterVertexFactoryShaderParameters</*bWithWaterSelectionSupport = */ true, /*bIndirectDraws = */ true>;
-using FWaterVertexFactoryWithSelectionNoIndirect = TWaterVertexFactory</*bWithWaterSelectionSupport = */ true, /*bIndirectDraws = */ false>;
-using FWaterVertexFactoryWithSelectionWithIndirect = TWaterVertexFactory</*bWithWaterSelectionSupport = */ true, /*bIndirectDraws = */ true>;
-
 // In editor builds, also implement the vertex factory that supports water selection:
-IMPLEMENT_TEMPLATE_TYPE_LAYOUT(template<>, FWaterVertexFactoryParametersWithSelectionNoIndirect);
-IMPLEMENT_TEMPLATE_TYPE_LAYOUT(template<>, FWaterVertexFactoryParametersWithSelectionWithIndirect);
-IMPLEMENT_VERTEX_FACTORY_PARAMETER_TYPE(FWaterVertexFactoryWithSelectionNoIndirect, SF_Vertex, FWaterVertexFactoryParametersWithSelectionNoIndirect);
-IMPLEMENT_VERTEX_FACTORY_PARAMETER_TYPE(FWaterVertexFactoryWithSelectionWithIndirect, SF_Vertex, FWaterVertexFactoryParametersWithSelectionWithIndirect);
+IMPLEMENT_TEMPLATE_TYPE_LAYOUT(template<>, TWaterVertexFactoryShaderParameters</*bWithWaterSelectionSupport = */ true>);
+IMPLEMENT_VERTEX_FACTORY_PARAMETER_TYPE(TWaterVertexFactory</*bWithWaterSelectionSupport = */ true>, SF_Vertex, TWaterVertexFactoryShaderParameters</*bWithWaterSelectionSupport = */ true>);
 #if RHI_RAYTRACING
-IMPLEMENT_VERTEX_FACTORY_PARAMETER_TYPE(FWaterVertexFactoryWithSelectionNoIndirect, SF_Compute, FWaterVertexFactoryParametersWithSelectionNoIndirect);
-IMPLEMENT_VERTEX_FACTORY_PARAMETER_TYPE(FWaterVertexFactoryWithSelectionNoIndirect, SF_RayHitGroup, FWaterVertexFactoryParametersWithSelectionNoIndirect);
+IMPLEMENT_VERTEX_FACTORY_PARAMETER_TYPE(TWaterVertexFactory</*bWithWaterSelectionSupport = */ true>, SF_Compute, TWaterVertexFactoryShaderParameters</*bWithWaterSelectionSupport = */ true>);
+IMPLEMENT_VERTEX_FACTORY_PARAMETER_TYPE(TWaterVertexFactory</*bWithWaterSelectionSupport = */ true>, SF_RayHitGroup, TWaterVertexFactoryShaderParameters</*bWithWaterSelectionSupport = */ true>);
 #endif
-IMPLEMENT_TEMPLATE_VERTEX_FACTORY_TYPE(template<>, FWaterVertexFactoryWithSelectionNoIndirect, "/Plugin/Water/Private/WaterMeshVertexFactory.ush",
+IMPLEMENT_TEMPLATE_VERTEX_FACTORY_TYPE(template<>, TWaterVertexFactory</*bWithWaterSelectionSupport = */ true>, "/Plugin/Water/Private/WaterMeshVertexFactory.ush",
 	  EVertexFactoryFlags::UsedWithMaterials
 	| EVertexFactoryFlags::SupportsDynamicLighting
 	| EVertexFactoryFlags::SupportsPrecisePrevWorldPos
@@ -167,44 +112,20 @@ IMPLEMENT_TEMPLATE_VERTEX_FACTORY_TYPE(template<>, FWaterVertexFactoryWithSelect
 	| EVertexFactoryFlags::SupportsRayTracingDynamicGeometry
 	| EVertexFactoryFlags::SupportsPSOPrecaching
 );
-IMPLEMENT_TEMPLATE_VERTEX_FACTORY_TYPE(template<>, FWaterVertexFactoryWithSelectionWithIndirect, "/Plugin/Water/Private/WaterMeshVertexFactory.ush",
-	EVertexFactoryFlags::UsedWithMaterials
-	| EVertexFactoryFlags::SupportsDynamicLighting
-	| EVertexFactoryFlags::SupportsPrecisePrevWorldPos
-	| EVertexFactoryFlags::SupportsPrimitiveIdStream
-	| EVertexFactoryFlags::SupportsRayTracing
-	| EVertexFactoryFlags::SupportsRayTracingDynamicGeometry
-	| EVertexFactoryFlags::SupportsPSOPrecaching
-);
 
 #endif // WITH_WATER_SELECTION_SUPPORT
 
-const FVertexFactoryType* GetWaterVertexFactoryType(bool bWithWaterSelectionSupport, bool bIndirectDraws)
+const FVertexFactoryType* GetWaterVertexFactoryType(bool bWithWaterSelectionSupport)
 {
 #if WITH_WATER_SELECTION_SUPPORT
 	if (bWithWaterSelectionSupport)
 	{
-		if (bIndirectDraws)
-		{
-			return &TWaterVertexFactory</*bWithWaterSelectionSupport = */ true, /*bIndirectDraws = */ true>::StaticType;
-		}
-		else
-		{
-			return &TWaterVertexFactory</*bWithWaterSelectionSupport = */ true, /*bIndirectDraws = */ false>::StaticType;
-		}
+		return &TWaterVertexFactory</*bWithWaterSelectionSupport = */ true>::StaticType;
 	}
 	else
 #endif // WITH_WATER_SELECTION_SUPPORT
 	{
 		check(!bWithWaterSelectionSupport);
-		if (bIndirectDraws)
-		{
-			return &TWaterVertexFactory</*bWithWaterSelectionSupport = */ false, /*bIndirectDraws = */ true>::StaticType;
-		}
-		else
-		{
-			return &TWaterVertexFactory</*bWithWaterSelectionSupport = */ false, /*bIndirectDraws = */ false>::StaticType;
-		}
-		
+		return &TWaterVertexFactory</*bWithWaterSelectionSupport = */ false>::StaticType;
 	}
 }
diff --git a/Engine/Plugins/Experimental/Water/Source/Runtime/Private/WaterViewExtension.cpp b/Engine/Plugins/Experimental/Water/Source/Runtime/Private/WaterViewExtension.cpp
--- a/Engine/Plugins/Experimental/Water/Source/Runtime/Private/WaterViewExtension.cpp
+++ b/Engine/Plugins/Experimental/Water/Source/Runtime/Private/WaterViewExtension.cpp
@@ -33,8 +33,6 @@ static TAutoConsoleVariable<int32> CVarWaterInfoRenderMethod(
 
 // ----------------------------------------------------------------------------------
 
-FWaterMeshGPUWork GWaterMeshGPUWork;
-
 FWaterViewExtension::FWaterViewExtension(const FAutoRegister& AutoReg, UWorld* InWorld)
 	: FWorldSceneViewExtension(AutoReg, InWorld)
 	, WaterGPUData(MakeShared<FWaterGPUResources, ESPMode::ThreadSafe>())
@@ -412,14 +410,6 @@ void FWaterViewExtension::PreRenderView_RenderThread(FRDGBuilder& GraphBuilder,
 	}
 }
 
-void FWaterViewExtension::PreRenderBasePass_RenderThread(FRDGBuilder& GraphBuilder)
-{
-	for (FWaterMeshGPUWork::FCallback& Callback : GWaterMeshGPUWork.Callbacks)
-	{
-		Callback.Function(GraphBuilder);
-	}
-}
-
 void FWaterViewExtension::MarkWaterInfoTextureForRebuild(const UE::WaterInfo::FRenderingContext& RenderContext)
 {
 	WaterInfoContextsToRender.Emplace(RenderContext.ZoneToRender, RenderContext);
diff --git a/Engine/Plugins/Experimental/Water/Source/Runtime/Public/WaterMeshSceneProxy.h b/Engine/Plugins/Experimental/Water/Source/Runtime/Public/WaterMeshSceneProxy.h
--- a/Engine/Plugins/Experimental/Water/Source/Runtime/Public/WaterMeshSceneProxy.h
+++ b/Engine/Plugins/Experimental/Water/Source/Runtime/Public/WaterMeshSceneProxy.h
@@ -8,7 +8,6 @@
 #include "WaterQuadTree.h"
 #include "WaterVertexFactory.h"
 #include "RayTracingGeometry.h"
-#include "WaterQuadTreeGPU.h"
 
 class FMeshElementCollector;
 struct FRayTracingMaterialGatheringContext;
@@ -32,8 +31,6 @@ public:
 
 	virtual void CreateRenderThreadResources(FRHICommandListBase& RHICmdList) override;
 
-	virtual void DestroyRenderThreadResources() override;
-
 	virtual void GetDynamicMeshElements(const TArray<const FSceneView*>& Views, const FSceneViewFamily& ViewFamily, uint32 VisibilityMap, FMeshElementCollector& Collector) const override;
 
 	virtual const TArray<FBoxSphereBounds>* GetOcclusionQueries(const FSceneView* View) const override;
@@ -56,7 +53,7 @@ public:
 
 	uint32 GetAllocatedSize() const
 	{
-		return(FPrimitiveSceneProxy::GetAllocatedSize() + (WaterVertexFactories.GetAllocatedSize() + WaterVertexFactories.Num() * sizeof(FWaterVertexFactoryType)) + WaterQuadTree.GetAllocatedSize());
+		return(FPrimitiveSceneProxy::GetAllocatedSize() + (WaterVertexFactories.GetAllocatedSize() + WaterVertexFactories.Num() * sizeof(WaterVertexFactoryType)) + WaterQuadTree.GetAllocatedSize());
 	}
 
 #if WITH_WATER_SELECTION_SUPPORT
@@ -64,10 +61,9 @@ public:
 #endif // WITH_WATER_SELECTION_SUPPORT
 
 	// At runtime, we only ever need one version of the vertex factory : with selection support (editor) or without : 
-	using FWaterVertexFactoryType = TWaterVertexFactory<WITH_WATER_SELECTION_SUPPORT, /*bIndirectDraws = */ false>;
-	using FWaterVertexFactoryIndirectDrawType = TWaterVertexFactory<WITH_WATER_SELECTION_SUPPORT, /*bIndirectDraws = */ true>;
-	using FWaterInstanceDataBuffersType = TWaterInstanceDataBuffers<WITH_WATER_SELECTION_SUPPORT>;
-	using FWaterMeshUserDataBuffersType = TWaterMeshUserDataBuffers<WITH_WATER_SELECTION_SUPPORT>;
+	using WaterVertexFactoryType = TWaterVertexFactory<WITH_WATER_SELECTION_SUPPORT>;
+	using WaterInstanceDataBuffersType = TWaterInstanceDataBuffers<WITH_WATER_SELECTION_SUPPORT>;
+	using WaterMeshUserDataBuffersType = TWaterMeshUserDataBuffers<WITH_WATER_SELECTION_SUPPORT>;
 
 #if RHI_RAYTRACING
 	virtual void GetDynamicRayTracingInstances(FRayTracingMaterialGatheringContext& Context, TArray<FRayTracingInstance>& OutRayTracingInstances) override final;
@@ -105,23 +101,16 @@ private:
 	FMaterialRelevance MaterialRelevance;
 
 	// One vertex factory per LOD
-	TArray<FWaterVertexFactoryType*> WaterVertexFactories;
-	TArray<FWaterVertexFactoryIndirectDrawType*> WaterVertexFactoriesIndirectDraw;
+	TArray<WaterVertexFactoryType*> WaterVertexFactories;
 
 	/** Tiles containing water, stored in a quad tree */
 	FWaterQuadTree WaterQuadTree;
 
-	/** GPU quad tree instance. Only initialized and used if WaterQuadTree.IsGPUQuadTree() is true. */
-	FWaterQuadTreeGPU QuadTreeGPU;
-
 	/** Unique Instance data buffer shared accross water batch draw calls */	
-	FWaterInstanceDataBuffersType* WaterInstanceDataBuffers;
+	WaterInstanceDataBuffersType* WaterInstanceDataBuffers;
 
 	/** Per-"water render group" user data (the number of groups might vary depending on whether we're in the editor or not) */
-	FWaterMeshUserDataBuffersType* WaterMeshUserDataBuffers;
-
-	double WaterQuadTreeMinHeight = DBL_MAX;
-	double WaterQuadTreeMaxHeight = -DBL_MAX;
+	WaterMeshUserDataBuffersType* WaterMeshUserDataBuffers;
 
 	/** The world-space bounds of the current water info texture coverage. The Water mesh should only render tiles within this bounding box. */
 	FBox2D WaterInfoBounds = FBox2D(ForceInit);
@@ -129,9 +118,6 @@ private:
 	/** Scale of the concentric LOD squares  */
 	float LODScale = -1.0f;
 
-	/** Number of quads per side of a water quad tree tile at LOD0 */
-	int32 NumQuadsLOD0 = 0;
-
 	/** Number of densities (same as number of grid index/vertex buffers) */
 	int32 DensityCount = 0;
 
@@ -156,12 +142,6 @@ private:
 	UE::FMutex OcclusionResultsMutex;
 	int32 OcclusionResultsFarMeshOffset = INT32_MAX;
 	uint32 SceneProxyCreatedFrameNumberRenderThread = INDEX_NONE;
-
-	mutable FWaterQuadTreeGPU::FTraverseParams WaterQuadTreeGPUTraverseParams;
-	mutable bool bNeedToTraverseGPUQuadTree = false;
-
-	/** Initializes the GPU quad tree */
-	void BuildGPUQuadTree(FRDGBuilder& GraphBuilder);
 };
 
 #if UE_ENABLE_INCLUDE_ORDER_DEPRECATED_IN_5_2
diff --git a/Engine/Plugins/Experimental/Water/Source/Runtime/Public/WaterQuadTree.h b/Engine/Plugins/Experimental/Water/Source/Runtime/Public/WaterQuadTree.h
--- a/Engine/Plugins/Experimental/Water/Source/Runtime/Public/WaterQuadTree.h
+++ b/Engine/Plugins/Experimental/Water/Source/Runtime/Public/WaterQuadTree.h
@@ -28,13 +28,6 @@ struct FWaterBodyRenderData
 	/** World Z position of the waterbody, this is where the tiles for this water body will be rendered*/
 	double SurfaceBaseHeight = 0.0;
 
-	/** Z bounds of the water body */
-	double BoundsMinZ = DBL_MAX;
-	double BoundsMaxZ = -DBL_MAX;
-
-	/** Maximum Z displacement due to waves */
-	float MaxWaveHeight = 0.0f;
-
 	/** Render priority. If two water bodies overlap, this will decide which water body is used for a tile */
 	int16 Priority = TNumericLimits<int16>::Min();
 
@@ -73,15 +66,6 @@ struct FWaterBodyRenderData
 	}
 };
 
-struct FWaterBodyQuadTreeRasterInfo
-{
-	FTransform LocalToWorld = FTransform::Identity;
-	const class FStaticMeshRenderData* RenderData = nullptr;
-	uint32 WaterBodyRenderDataIndex = INDEX_NONE;
-	uint8 Priority = 0;
-	bool bIsRiver = false;
-};
-
 struct FWaterQuadTree 
 {
 	enum { INVALID_PARENT = 0xFFFFFFF };
@@ -141,7 +125,7 @@ struct FWaterQuadTree
 	 *	Initialize the tree. This will unlock the tree for node insertion using AddWaterTilesInsideBounds(...). 
 	 *	Tree must be locked before traversal, see Lock(). 
 	 */
-	void InitTree(const FBox2D& InBounds, float InTileSize, FIntPoint InExtentInTiles, bool bInIsGPUQuadTree);
+	void InitTree(const FBox2D& InBounds, float InTileSize, FIntPoint InExtentInTiles);
 
 	/** Unlock to make it read-only. This will optionally prune the node array to remove redundant nodes, nodes that can be implicitly traversed */
 	void Unlock(bool bPruneRedundantNodes);
@@ -164,8 +148,6 @@ struct FWaterQuadTree
 	/** Add water body render data to this tree. Returns the index in the array. Use this index to add tiles with this water body to the tree, see AddWaterTilesInsideBounds(..) */
 	uint32 AddWaterBodyRenderData(const FWaterBodyRenderData& InWaterBodyRenderData) { return NodeData.WaterBodyRenderData.Add(InWaterBodyRenderData); }
 
-	void AddWaterBodyRasterInfo(const FWaterBodyQuadTreeRasterInfo& InWaterBodyRasterInfo) { WaterBodyRasterInfos.Add(InWaterBodyRasterInfo); }
-
 	/** Get bounds of the root node if there is one, otherwise some default box */
 	FBox GetBounds() const { return NodeData.Nodes.Num() > 0 ? NodeData.Nodes[0].Bounds : FBox(-FVector::OneVector, FVector::OneVector); }
 
@@ -201,19 +183,11 @@ struct FWaterQuadTree
 
 	const TArray<FMaterialRenderProxy*>& GetWaterMaterials() const { return WaterMaterials; }
 
-	const TArray<FWaterBodyQuadTreeRasterInfo>& GetWaterBodyRasterInfos() const { return WaterBodyRasterInfos; }
-
-	const TArray<FWaterBodyRenderData>& GetWaterBodyRenderData() const { return NodeData.WaterBodyRenderData; }
-
-	FIntPoint GetResolution() const { return ExtentInTiles * 2; }
-
-	bool IsGPUQuadTree() const { return bIsGPUQuadTree; }
-
 	/** Calculate the world distance to a LOD */
 	static float GetLODDistance(int32 InLODLevel, float InLODScale) { return FMath::Pow(2.0f, (float)(InLODLevel + 1)) * InLODScale; }
 
 	/** Total memory dynamically allocated by this object */
-	uint32 GetAllocatedSize() const { return NodeData.GetAllocatedSize() + WaterMaterials.GetAllocatedSize() + WaterBodyRasterInfos.GetAllocatedSize() + BreadthFirstOrder.GetAllocatedSize() + FarMeshData.GetAllocatedSize(); }
+	uint32 GetAllocatedSize() const { return NodeData.GetAllocatedSize() + WaterMaterials.GetAllocatedSize() + BreadthFirstOrder.GetAllocatedSize() + FarMeshData.GetAllocatedSize(); }
 
 #if WITH_WATER_SELECTION_SUPPORT
 	/** Obtain all possible hit proxies (proxies of all the water bodies) */
@@ -317,9 +291,8 @@ private:
 	} NodeData;
 
 	TArray<FMaterialRenderProxy*> WaterMaterials;
-	TArray<FWaterBodyQuadTreeRasterInfo> WaterBodyRasterInfos;
-	
-/** Node indices ordered by breadth first ordering */
+
+	/** Node indices ordered by breadth first ordering */
 	TArray<int32> BreadthFirstOrder;
 
 	/** Contains everything needed to render the far mesh. This data lives outside the quadtree structure itself */
@@ -358,6 +331,4 @@ private:
 
 	/** If true, the tree may not change */
 	bool bIsReadOnly = true;
-
-	bool bIsGPUQuadTree = false;
 };
diff --git a/Engine/Plugins/Experimental/Water/Source/Runtime/Public/WaterVertexFactory.h b/Engine/Plugins/Experimental/Water/Source/Runtime/Public/WaterVertexFactory.h
--- a/Engine/Plugins/Experimental/Water/Source/Runtime/Public/WaterVertexFactory.h
+++ b/Engine/Plugins/Experimental/Water/Source/Runtime/Public/WaterVertexFactory.h
@@ -164,17 +164,16 @@ enum class EWaterMeshRenderGroupType : uint8
 #endif // WITH_WATER_SELECTION_SUPPORT
 };
 
-template <bool bWithWaterSelectionSupport, bool bIndirectDraws>
+template <bool bWithWaterSelectionSupport>
 class TWaterVertexFactory : public FVertexFactory
 {
-	DECLARE_VERTEX_FACTORY_TYPE(FWaterVertexFactoryType);
+	DECLARE_VERTEX_FACTORY_TYPE(TWaterVertexFactory<bWithWaterSelectionSupport>);
 
 public:
 	using Super = FVertexFactory;
-	using FWaterVertexFactoryType = TWaterVertexFactory<bWithWaterSelectionSupport, bIndirectDraws>;
 
 	static constexpr int32 NumRenderGroups = bWithWaterSelectionSupport ? 3 : 1; // Must match EWaterMeshRenderGroupType
-	static constexpr int32 NumAdditionalVertexStreams = bIndirectDraws ? 0 : TWaterInstanceDataBuffers<bWithWaterSelectionSupport>::NumBuffers;
+	static constexpr int32 NumAdditionalVertexStreams = TWaterInstanceDataBuffers<bWithWaterSelectionSupport>::NumBuffers;
 
 	TWaterVertexFactory(ERHIFeatureLevel::Type InFeatureLevel, int32 InNumQuadsPerSide,	float InLODScale);
 	~TWaterVertexFactory();
@@ -216,7 +215,7 @@ private:
 	const float LODScale = 0.0f;
 };
 
-extern const FVertexFactoryType* GetWaterVertexFactoryType(bool bWithWaterSelectionSupport, bool bIndirectDraws);
+extern const FVertexFactoryType* GetWaterVertexFactoryType(bool bWithWaterSelectionSupport);
 
 
 /**
@@ -239,11 +238,6 @@ struct TWaterMeshUserData
 #if RHI_RAYTRACING	
 	FUniformBufferRHIRef WaterVertexFactoryRaytracingVFUniformBuffer = nullptr;
 #endif
-
-	FRHIShaderResourceView* IndirectInstanceDataOffsets = nullptr;
-	FRHIShaderResourceView* IndirectInstanceData0 = nullptr;
-	FRHIShaderResourceView* IndirectInstanceData1 = nullptr;
-	FRHIShaderResourceView* IndirectInstanceData2 = nullptr;
 };
 
 /**
@@ -273,7 +267,7 @@ struct TWaterMeshUserDataBuffers
 		return UserData[(int32)InRenderGroupType].Get();
 	}
 
-	TStaticArray<TUniquePtr<WaterMeshUserDataType>, TWaterVertexFactory<bWithWaterSelectionSupport, /*bIndirectDraws = */ false>::NumRenderGroups> UserData;
+	TStaticArray<TUniquePtr<WaterMeshUserDataType>, TWaterVertexFactory<bWithWaterSelectionSupport>::NumRenderGroups> UserData;
 };
 
 #include "WaterVertexFactory.inl"
diff --git a/Engine/Plugins/Experimental/Water/Source/Runtime/Public/WaterVertexFactory.inl b/Engine/Plugins/Experimental/Water/Source/Runtime/Public/WaterVertexFactory.inl
--- a/Engine/Plugins/Experimental/Water/Source/Runtime/Public/WaterVertexFactory.inl
+++ b/Engine/Plugins/Experimental/Water/Source/Runtime/Public/WaterVertexFactory.inl
@@ -13,8 +13,8 @@
 
 // ----------------------------------------------------------------------------------
 
-template <bool bWithWaterSelectionSupport, bool bIndirectDraws>
-TWaterVertexFactory<bWithWaterSelectionSupport, bIndirectDraws>::TWaterVertexFactory(ERHIFeatureLevel::Type InFeatureLevel, int32 InNumQuadsPerSide, float InLODScale)
+template <bool bWithWaterSelectionSupport>
+TWaterVertexFactory<bWithWaterSelectionSupport>::TWaterVertexFactory(ERHIFeatureLevel::Type InFeatureLevel, int32 InNumQuadsPerSide, float InLODScale)
 	: FVertexFactory(InFeatureLevel)
 	, NumQuadsPerSide(InNumQuadsPerSide)
 	, LODScale(InLODScale)
@@ -23,15 +23,15 @@ TWaterVertexFactory<bWithWaterSelectionSupport, bIndirectDraws>::TWaterVertexFac
 	IndexBuffer = new FWaterMeshIndexBuffer(NumQuadsPerSide);
 }
 
-template <bool bWithWaterSelectionSupport, bool bIndirectDraws>
-TWaterVertexFactory<bWithWaterSelectionSupport, bIndirectDraws>::~TWaterVertexFactory()
+template <bool bWithWaterSelectionSupport>
+TWaterVertexFactory<bWithWaterSelectionSupport>::~TWaterVertexFactory()
 {
 	delete VertexBuffer;
 	delete IndexBuffer;
 }
 
-template <bool bWithWaterSelectionSupport, bool bIndirectDraws>
-void TWaterVertexFactory<bWithWaterSelectionSupport, bIndirectDraws>::InitRHI(FRHICommandListBase& RHICmdList)
+template <bool bWithWaterSelectionSupport>
+void TWaterVertexFactory<bWithWaterSelectionSupport>::InitRHI(FRHICommandListBase& RHICmdList)
 {
 	Super::InitRHI(RHICmdList);
 
@@ -70,13 +70,10 @@ void TWaterVertexFactory<bWithWaterSelectionSupport, bIndirectDraws>::InitRHI(FR
 	Elements.Add(VertexPositionElement);
 
 	// Adds all streams
-	if constexpr (NumAdditionalVertexStreams > 0)
+	for (int32 StreamIdx = 0; StreamIdx < NumAdditionalVertexStreams; ++StreamIdx)
 	{
-		for (int32 StreamIdx = 0; StreamIdx < NumAdditionalVertexStreams; ++StreamIdx)
-		{
-			FVertexElement InstanceElement(Streams.Add(InstanceDataVertexStream), 0, VET_Float4, 8 + StreamIdx, InstanceDataVertexStream.Stride, true);
-			Elements.Add(InstanceElement);
-		}
+		FVertexElement InstanceElement(Streams.Add(InstanceDataVertexStream), 0, VET_Float4, 8 + StreamIdx, InstanceDataVertexStream.Stride, true);
+		Elements.Add(InstanceElement);
 	}
 
 #if 0
@@ -96,8 +93,8 @@ void TWaterVertexFactory<bWithWaterSelectionSupport, bIndirectDraws>::InitRHI(FR
 	InitDeclaration(Elements);
 }
 
-template <bool bWithWaterSelectionSupport, bool bIndirectDraws>
-void TWaterVertexFactory<bWithWaterSelectionSupport, bIndirectDraws>::ReleaseRHI()
+template <bool bWithWaterSelectionSupport>
+void TWaterVertexFactory<bWithWaterSelectionSupport>::ReleaseRHI()
 {
 	for (auto& UniformBuffer : UniformBuffers)
 	{
@@ -117,8 +114,8 @@ void TWaterVertexFactory<bWithWaterSelectionSupport, bIndirectDraws>::ReleaseRHI
 	Super::ReleaseRHI();
 }
 
-template <bool bWithWaterSelectionSupport, bool bIndirectDraws>
-void TWaterVertexFactory<bWithWaterSelectionSupport, bIndirectDraws>::SetupUniformDataForGroup(EWaterMeshRenderGroupType InRenderGroupType)
+template <bool bWithWaterSelectionSupport>
+void TWaterVertexFactory<bWithWaterSelectionSupport>::SetupUniformDataForGroup(EWaterMeshRenderGroupType InRenderGroupType)
 {
 	FWaterVertexFactoryParameters UniformParams;
 	UniformParams.NumQuadsPerTileSide = NumQuadsPerSide;
@@ -134,8 +131,8 @@ void TWaterVertexFactory<bWithWaterSelectionSupport, bIndirectDraws>::SetupUnifo
 	UniformBuffers[(int32)InRenderGroupType] = FWaterVertexFactoryBufferRef::CreateUniformBufferImmediate(UniformParams, UniformBuffer_MultiFrame);
 }
 
-template <bool bWithWaterSelectionSupport, bool bIndirectDraws>
-bool TWaterVertexFactory<bWithWaterSelectionSupport, bIndirectDraws>::ShouldCompilePermutation(const FVertexFactoryShaderPermutationParameters& Parameters)
+template <bool bWithWaterSelectionSupport>
+bool TWaterVertexFactory<bWithWaterSelectionSupport>::ShouldCompilePermutation(const FVertexFactoryShaderPermutationParameters& Parameters)
 {
 	const bool bIsCompatibleWithWater = ((Parameters.MaterialParameters.MaterialDomain == MD_Surface) && Parameters.MaterialParameters.bIsUsedWithWater) || Parameters.MaterialParameters.bIsSpecialEngineMaterial;
 	if (bIsCompatibleWithWater)
@@ -146,8 +143,8 @@ bool TWaterVertexFactory<bWithWaterSelectionSupport, bIndirectDraws>::ShouldComp
 	return false;
 }
 
-template <bool bWithWaterSelectionSupport, bool bIndirectDraws>
-void TWaterVertexFactory<bWithWaterSelectionSupport, bIndirectDraws>::ModifyCompilationEnvironment(const FVertexFactoryShaderPermutationParameters& Parameters, FShaderCompilerEnvironment& OutEnvironment)
+template <bool bWithWaterSelectionSupport>
+void TWaterVertexFactory<bWithWaterSelectionSupport>::ModifyCompilationEnvironment(const FVertexFactoryShaderPermutationParameters& Parameters, FShaderCompilerEnvironment& OutEnvironment)
 {
 	OutEnvironment.SetDefine(TEXT("WATER_MESH_FACTORY"), 1);
 #if 0
@@ -165,16 +162,12 @@ void TWaterVertexFactory<bWithWaterSelectionSupport, bIndirectDraws>::ModifyComp
 		OutEnvironment.SetDefine(TEXT("USE_VERTEXFACTORY_HITPROXY_ID"), TEXT("1"));
 		OutEnvironment.SetDefine(TEXT("WITH_WATER_SELECTION_SUPPORT_VF"), TEXT("1"));
 	}
-	if (bIndirectDraws)
-	{
-		OutEnvironment.SetDefine(TEXT("WATER_MESH_DRAW_INDIRECT"), TEXT("1"));
-	}
 
 	OutEnvironment.SetDefine(TEXT("RAY_TRACING_DYNAMIC_MESH_IN_LOCAL_SPACE"), TEXT("1"));
 }
 
-template <bool bWithWaterSelectionSupport, bool bIndirectDraws>
-void TWaterVertexFactory<bWithWaterSelectionSupport, bIndirectDraws>::ValidateCompiledResult(const FVertexFactoryType* Type, EShaderPlatform Platform, const FShaderParameterMap& ParameterMap, TArray<FString>& OutErrors)
+template <bool bWithWaterSelectionSupport>
+void TWaterVertexFactory<bWithWaterSelectionSupport>::ValidateCompiledResult(const FVertexFactoryType* Type, EShaderPlatform Platform, const FShaderParameterMap& ParameterMap, TArray<FString>& OutErrors)
 {
 #if 0
 	if (Type->SupportsPrimitiveIdStream()
@@ -186,18 +179,15 @@ void TWaterVertexFactory<bWithWaterSelectionSupport, bIndirectDraws>::ValidateCo
 #endif
 }
 
-template <bool bWithWaterSelectionSupport, bool bIndirectDraws>
-void TWaterVertexFactory<bWithWaterSelectionSupport, bIndirectDraws>::GetPSOPrecacheVertexFetchElements(EVertexInputStreamType VertexInputStreamType, FVertexDeclarationElementList& Elements)
+template <bool bWithWaterSelectionSupport>
+void TWaterVertexFactory<bWithWaterSelectionSupport>::GetPSOPrecacheVertexFetchElements(EVertexInputStreamType VertexInputStreamType, FVertexDeclarationElementList& Elements)
 {
 	// Add position stream
 	Elements.Add(FVertexElement(0, 0, VET_Float4, 0, sizeof(FVector4f), false));
 
 	// Add all the additional streams
-	if constexpr (NumAdditionalVertexStreams > 0)
+	for (int32 StreamIdx = 0; StreamIdx < NumAdditionalVertexStreams; ++StreamIdx)
 	{
-		for (int32 StreamIdx = 0; StreamIdx < NumAdditionalVertexStreams; ++StreamIdx)
-		{
-			Elements.Add(FVertexElement(1 + StreamIdx, 0, VET_Float4, 8 + StreamIdx, sizeof(FVector4f), true));
-		}
+		Elements.Add(FVertexElement(1 + StreamIdx, 0, VET_Float4, 8 + StreamIdx, sizeof(FVector4f), true));
 	}
 }
diff --git a/Engine/Plugins/Experimental/Water/Source/Runtime/Public/WaterViewExtension.h b/Engine/Plugins/Experimental/Water/Source/Runtime/Public/WaterViewExtension.h
--- a/Engine/Plugins/Experimental/Water/Source/Runtime/Public/WaterViewExtension.h
+++ b/Engine/Plugins/Experimental/Water/Source/Runtime/Public/WaterViewExtension.h
@@ -24,7 +24,6 @@ public:
 	virtual void SetupViewFamily(FSceneViewFamily& InViewFamily) override;
 	virtual void PreRenderViewFamily_RenderThread(FRDGBuilder& GraphBuilder, FSceneViewFamily& InViewFamily) override;
 	virtual void PreRenderView_RenderThread(FRDGBuilder& GraphBuilder, FSceneView& InView) override;
-	virtual void PreRenderBasePass_RenderThread(FRDGBuilder& GraphBuilder) override;
 	// End FSceneViewExtensionBase implementation
 
 	void MarkWaterInfoTextureForRebuild(const UE::WaterInfo::FRenderingContext& RenderContext);
@@ -56,18 +55,6 @@ private:
 	void UpdateGPUBuffers();
 };
 
-struct FWaterMeshGPUWork
-{
-	struct FCallback
-	{
-		class FWaterMeshSceneProxy* Proxy = nullptr;
-		TFunction<void(FRDGBuilder&)> Function;
-	};
-	TArray<FCallback> Callbacks;
-};
-
-extern FWaterMeshGPUWork GWaterMeshGPUWork;
-
 #if UE_ENABLE_INCLUDE_ORDER_DEPRECATED_IN_5_2
 #include "CoreMinimal.h"
 #endif
diff --git a/Engine/Source/Runtime/Engine/Public/SceneViewExtension.h b/Engine/Source/Runtime/Engine/Public/SceneViewExtension.h
--- a/Engine/Source/Runtime/Engine/Public/SceneViewExtension.h
+++ b/Engine/Source/Runtime/Engine/Public/SceneViewExtension.h
@@ -155,11 +155,6 @@ public:
 	 */
 	virtual void PreInitViews_RenderThread(FRDGBuilder& GraphBuilder) {}
 
-	/**
-	 * Called on render thread right before Base Pass rendering.
-	 */
-	virtual void PreRenderBasePass_RenderThread(FRDGBuilder& GraphBuilder) {}
-
 	/**
 	 * Called right after Base Pass rendering finished when using the deferred renderer.
 	 */
diff --git a/Engine/Source/Runtime/Renderer/Private/DeferredShadingRenderer.cpp b/Engine/Source/Runtime/Renderer/Private/DeferredShadingRenderer.cpp
--- a/Engine/Source/Runtime/Renderer/Private/DeferredShadingRenderer.cpp
+++ b/Engine/Source/Runtime/Renderer/Private/DeferredShadingRenderer.cpp
@@ -3466,11 +3466,6 @@ void FDeferredShadingSceneRenderer::Render(FRDGBuilder& GraphBuilder)
 
 		// End early occlusion queries
 
-		for (FSceneViewExtensionRef& ViewExtension : ViewFamily.ViewExtensions)
-		{
-			ViewExtension->PreRenderBasePass_RenderThread(GraphBuilder);
-		}
-
 		BeginAsyncDistanceFieldShadowProjections(GraphBuilder, SceneTextures, InitViewTaskDatas.DynamicShadows);
 
 		// Run local fog volume culling before base pass and after HZB generation tyo benefit from more culling.
diff --git a/Engine/Source/Runtime/Renderer/Private/MobileShadingRenderer.cpp b/Engine/Source/Runtime/Renderer/Private/MobileShadingRenderer.cpp
--- a/Engine/Source/Runtime/Renderer/Private/MobileShadingRenderer.cpp
+++ b/Engine/Source/Runtime/Renderer/Private/MobileShadingRenderer.cpp
@@ -1124,11 +1124,6 @@ void FMobileSceneRenderer::Render(FRDGBuilder& GraphBuilder)
 		}
 	}
 
-	for (FSceneViewExtensionRef& ViewExtension : ViewFamily.ViewExtensions)
-	{
-		ViewExtension->PreRenderBasePass_RenderThread(GraphBuilder);
-	}
-
 	if (bRendererOutputFinalSceneColor)
 	{
 		if (bDeferredShading)
