Modified the underwater post process sorting algorithm to prefer selecting water bodies that have waves before any others. If multiple water bodies overlap that both have waves, then we utilize the MaterialOverlapPriority to ensure that the post effect that is used matches the waves of the surface water at the location.

diff --git a/Engine/Plugins/Experimental/Water/Source/Runtime/Private/WaterSubsystem.cpp b/Engine/Plugins/Experimental/Water/Source/Runtime/Private/WaterSubsystem.cpp
--- a/Engine/Plugins/Experimental/Water/Source/Runtime/Private/WaterSubsystem.cpp
+++ b/Engine/Plugins/Experimental/Water/Source/Runtime/Private/WaterSubsystem.cpp
@@ -19,6 +19,7 @@
 #include "WaterUtils.h"
 #include "WaterViewExtension.h"
 #include "Algo/MaxElement.h"
+#include "Algo/RemoveIf.h"
 
 #if WITH_EDITOR
 #include "WaterZoneActorDesc.h"
@@ -701,16 +702,28 @@ void UWaterSubsystem::ComputeUnderwaterPostProcess(FVector ViewLocation, FSceneV
 	{
 		if (Hits.Num() > 1)
 		{
+			// Prepass to remove non-waterbody elements
+			Algo::RemoveIf(Hits, [](const FHitResult& A)
+			{
+				return A.HitObjectHandle.FetchActor<AWaterBody>() == nullptr;
+			});
+			
 			// Sort hits based on their water priority for rendering since we should prioritize evaluating waves in the order those waves will be considered for rendering. 
 			Hits.Sort([](const FHitResult& A, const FHitResult& B)
 			{
 				const AWaterBody* ABody = A.HitObjectHandle.FetchActor<AWaterBody>();
 				const AWaterBody* BBody = B.HitObjectHandle.FetchActor<AWaterBody>();
 
-				const int32 APriority = ABody ? ABody->GetWaterBodyComponent()->GetOverlapMaterialPriority() : -1;
-				const int32 BPriority = BBody ? BBody->GetWaterBodyComponent()->GetOverlapMaterialPriority() : -1;
+				// If both water bodies either have waves or both don't have waves, use the overlap priority to determine which to use, since in this case we need to respect the surface waves
+				if (ABody->GetWaterBodyComponent()->HasWaves() == BBody->GetWaterBodyComponent()->HasWaves())
+				{
+					const int32 APriority = ABody->GetWaterBodyComponent()->GetOverlapMaterialPriority();
+					const int32 BPriority = BBody->GetWaterBodyComponent()->GetOverlapMaterialPriority();
+					return APriority > BPriority;
+				}
 
-				return APriority > BPriority;
+				// Otherwise, prefer the water body with waves to ensure the PP calculates the waves correctly.
+				return ABody->GetWaterBodyComponent()->HasWaves() && !BBody->GetWaterBodyComponent()->HasWaves();
 			});
 		}
 
