Missing SetNum after Algo::RemoveIf in  UWaterSubsystem::ComputeUnderwaterPostProcess to actually remove the data from the array.

diff --git a/Engine/Plugins/Experimental/Water/Source/Runtime/Private/WaterSubsystem.cpp b/Engine/Plugins/Experimental/Water/Source/Runtime/Private/WaterSubsystem.cpp
--- a/Engine/Plugins/Experimental/Water/Source/Runtime/Private/WaterSubsystem.cpp
+++ b/Engine/Plugins/Experimental/Water/Source/Runtime/Private/WaterSubsystem.cpp
@@ -703,10 +703,10 @@ void UWaterSubsystem::ComputeUnderwaterPostProcess(FVector ViewLocation, FSceneV
 		if (Hits.Num() > 1)
 		{
 			// Prepass to remove non-waterbody elements
-			Algo::RemoveIf(Hits, [](const FHitResult& A)
+			Hits.SetNum(Algo::RemoveIf(Hits, [](const FHitResult& A)
 			{
 				return A.HitObjectHandle.FetchActor<AWaterBody>() == nullptr;
-			});
+			}));
 			
 			// Sort hits based on their water priority for rendering since we should prioritize evaluating waves in the order those waves will be considered for rendering. 
 			Hits.Sort([](const FHitResult& A, const FHitResult& B)
