Implemented a new GerstnerWave function that supports splitting the wave spectrum into three buckets for low/mid/high frequency waves. This can be useful to, for example, filter out some of the higher frequency waves on lower end platforms in the materials or utilize some interpolation instead of recomputing everything per-pixel for some frequency buckets.

diff --git a/Engine/Plugins/Experimental/Water/Shaders/Private/GerstnerWaveFunctions.ush b/Engine/Plugins/Experimental/Water/Shaders/Private/GerstnerWaveFunctions.ush
--- a/Engine/Plugins/Experimental/Water/Shaders/Private/GerstnerWaveFunctions.ush
+++ b/Engine/Plugins/Experimental/Water/Shaders/Private/GerstnerWaveFunctions.ush
@@ -18,6 +18,14 @@ struct WaveOutput
 	float3 WPO;
 };
 
+struct WaveOutputs
+{
+	WaveOutput LowFrequency;
+	WaveOutput MiddleFrequency;
+	WaveOutput HighFrequency;
+	WaveOutput AllWaves;
+};
+
 struct GerstnerWaveRenderer
 {
 	float2 WaveParamRTSize;
@@ -389,6 +397,21 @@ WaveOutput GetRangeGerstnerWavesNew(int InMinWaveIndex, int InMaxWaveIndex, int
 	return OutWaves;
 }
 
+WaveOutputs GetPerFrequencyGerstnerWavesNew(int LowFrequencyMaxIndex, int MiddleFrequencyMaxIndex, int InWaterBodyIndex, float InTime, FLWCVector2 InWorldPos)
+{
+	const int NumWaves = GetWaterBodyData(InWaterBodyIndex).NumWaves;
+	// Correct bad user input. Further input handling is done internally inside GetRangeGerstnerWavesNew to ensure the index ranges are valid for the number of waves in a given water body.
+	MiddleFrequencyMaxIndex = max(LowFrequencyMaxIndex + 1, MiddleFrequencyMaxIndex);
+	
+	WaveOutputs Result = (WaveOutputs)0;
+	Result.LowFrequency = GetRangeGerstnerWavesNew(0, LowFrequencyMaxIndex, InWaterBodyIndex, InTime, InWorldPos);
+	Result.MiddleFrequency = GetRangeGerstnerWavesNew(LowFrequencyMaxIndex + 1, MiddleFrequencyMaxIndex, InWaterBodyIndex, InTime, InWorldPos);
+	Result.HighFrequency = GetRangeGerstnerWavesNew(MiddleFrequencyMaxIndex + 1, NumWaves, InWaterBodyIndex, InTime, InWorldPos);
+	Result.AllWaves = AddWavesNew(Result.LowFrequency, AddWavesNew(Result.HighFrequency, Result.MiddleFrequency));
+
+	return Result;
+}
+
 /** Returns the wave attenuation factor according to a given water body and a given water depth. Should match the CPU version (GetWaveAttenuationFactor) */
 float ComputeWaveDepthAttenuationFactorNew(int InWaterBodyIndex, float InWaterDepth)
 {
