Water: update the water material instances when the owning water zone changes. This used to be automatically handled due to AWaterZone::AddWaterBodyComponent triggering a water info texture render which would update material instances with new bounds data, but since the data buffer refactor has been a race condition in certain scenarios.

diff --git a/Engine/Plugins/Experimental/Water/Source/Runtime/Private/WaterBodyComponent.cpp b/Engine/Plugins/Experimental/Water/Source/Runtime/Private/WaterBodyComponent.cpp
--- a/Engine/Plugins/Experimental/Water/Source/Runtime/Private/WaterBodyComponent.cpp
+++ b/Engine/Plugins/Experimental/Water/Source/Runtime/Private/WaterBodyComponent.cpp
@@ -418,6 +418,8 @@ void UWaterBodyComponent::UpdateWaterZones(bool bAllowChangesDuringCook /* = fal
 			}
 
 			OwningWaterZone = FoundZone;
+			
+			UpdateMaterialInstances();
 		}
 	}
 }
