Turn custom render pass on as the default rendering method of water info rendering.

diff --git a/Engine/Plugins/Experimental/Water/Source/Runtime/Private/WaterViewExtension.cpp b/Engine/Plugins/Experimental/Water/Source/Runtime/Private/WaterViewExtension.cpp
--- a/Engine/Plugins/Experimental/Water/Source/Runtime/Private/WaterViewExtension.cpp
+++ b/Engine/Plugins/Experimental/Water/Source/Runtime/Private/WaterViewExtension.cpp
@@ -28,7 +28,7 @@ static TAutoConsoleVariable<float> CVarLocalTessellationUpdateMargin(
 extern void OnCVarWaterInfoSceneProxiesValueChanged(IConsoleVariable*);
 static TAutoConsoleVariable<int32> CVarWaterInfoRenderMethod(
 	TEXT("r.Water.WaterInfo.RenderMethod"),
-	0,
+	2,
 	TEXT("0: SceneCaptures, 1: Custom, 2: CustomRenderPasses"),
 	FConsoleVariableDelegate::CreateStatic(OnCVarWaterInfoSceneProxiesValueChanged),
 	ECVF_Default | ECVF_RenderThreadSafe);
