Implemented pixel precise occlusion culling to GPU water quadtree indirect draw pipeline. Also moved the ShowTileBounds ShaderPrint feature into its own debug shader.

#rnx

diff --git a/Engine/Plugins/Experimental/Water/Shaders/Private/WaterQuadTreeDraws.usf b/Engine/Plugins/Experimental/Water/Shaders/Private/WaterQuadTreeDraws.usf
--- a/Engine/Plugins/Experimental/Water/Shaders/Private/WaterQuadTreeDraws.usf
+++ b/Engine/Plugins/Experimental/Water/Shaders/Private/WaterQuadTreeDraws.usf
@@ -54,9 +54,22 @@ FWaterLODParams GetWaterLODParams(float InObserverHeight, float InWaterHeightFor
 	return WaterLODParams;
 }
 
+float4 GetNodeAABB2D(uint3 InNodeCoord, float3 InQuadTreePosition, float InLeafSize)
+{
+	const float Scale = (1u << InNodeCoord.z) * InLeafSize;
+	return float4(InNodeCoord.xy, InNodeCoord.xy + 1.0f) * Scale + InQuadTreePosition.xyxy;
+}
+
 #ifdef INITIALIZE_INDIRECT_ARGS
 
+#ifndef PRECISE_OCCLUSION_QUERIES
+#define PRECISE_OCCLUSION_QUERIES 0
+#endif
+
 RWStructuredBuffer<FWaterTileDrawIndexedIndirectArgs> IndirectArgs;
+#if PRECISE_OCCLUSION_QUERIES
+RWStructuredBuffer<FWaterTileDrawIndexedIndirectArgs> OcclusionQueryIndirectArgs;
+#endif
 uint NumDensities;
 uint NumMaterials;
 uint NumViews;
@@ -88,22 +101,43 @@ void MainCS(uint3 DispatchThreadID : SV_DispatchThreadID)
 	{
 		IndirectArgs[NumBuckets * ViewIndex + BucketIndex] = BucketArgs;
 	}
+
+#if PRECISE_OCCLUSION_QUERIES
+	if (all(DispatchThreadID == 0))
+	{
+		FWaterTileDrawIndexedIndirectArgs OcclusionQueryArgs;
+		OcclusionQueryArgs.IndexCountPerInstance = 36; // 12 triangles per cube
+		OcclusionQueryArgs.InstanceCount = 0;
+		OcclusionQueryArgs.StartIndexLocation = 0;
+		OcclusionQueryArgs.BaseVertexLocation = 0;
+		OcclusionQueryArgs.StartInstanceLocation = 0;
+
+		for (uint ViewIndex = 0; ViewIndex < NumViews; ++ViewIndex)
+		{
+			OcclusionQueryIndirectArgs[ViewIndex] = OcclusionQueryArgs;
+		}
+	}
+#endif // PRECISE_OCCLUSION_QUERIES
 }
 
 #endif // INITIALIZE_INDIRECT_ARGS
 
 #ifdef QUAD_TREE_TRAVERSE
 
-#ifndef DEBUG_SHADER_PRINT
-#define DEBUG_SHADER_PRINT 0
-#endif
-
 #include "/Engine/Private/Nanite/NaniteHZBCull.ush" // BoxCullFrustum()
-#if DEBUG_SHADER_PRINT
-#include "/Engine/Private/ShaderPrint.ush"
+
+#ifndef PRECISE_OCCLUSION_QUERIES
+#define PRECISE_OCCLUSION_QUERIES 0
 #endif
 
 RWByteAddressBuffer PackedNodes;
+
+#if PRECISE_OCCLUSION_QUERIES
+RWBuffer<float4> OcclusionQueryBoxes;
+RWBuffer<uint> OcclusionVisibility;
+RWStructuredBuffer<FWaterTileDrawIndexedIndirectArgs> OcclusionQueryIndirectArgs;
+#endif // PRECISE_OCCLUSION_QUERIES
+
 Texture2D<float4> QuadTreeTexture;
 Texture2D<float4> WaterZBoundsTexture;
 StructuredBuffer<FWaterBodyRenderData> WaterBodyRenderData;
@@ -116,24 +150,15 @@ uint QuadTreeResolutionX;
 uint QuadTreeResolutionY;
 uint NumDensities;
 uint NumMaterials;
+uint ViewIndex;
 float LeafSize;
 float LODScale;
 float CaptureDepthRange;
 int ForceCollapseDensityLevel;
 uint NumLODs;
 uint NumDispatchedThreads;
-int DebugShowTile;
 uint bHZBOcclusionCullingEnabled;
 
-
-void AddNodeForRender(uint InNodeCoordPacked)
-{
-	uint WritePos = 0;
-	PackedNodes.InterlockedAdd(0, 1, WritePos);
-	WritePos += 1; // Num is stored at index 0
-	PackedNodes.Store(WritePos << 2u, InNodeCoordPacked);
-}
-
 float ComputeSquaredDistanceToPoint(float2 InBox2DMin, float2 InBox2DMax, float2 InPoint)
 {
 	// Accumulates the distance as we iterate axis
@@ -178,12 +203,6 @@ bool CanRender(FWaterQuadTreeNode Node, int InDensityLevel, int InForceCollapseD
 	return MaterialIndex >= 0 && Node.bIsSubtreeSameWaterBody && ((InDensityLevel > InForceCollapseDensityLevel) || Node.bHasCompleteSubtree);
 }
 
-float4 GetNodeAABB2D(uint3 InNodeCoord)
-{
-	const float Scale = (1u << InNodeCoord.z) * LeafSize;
-	return float4(InNodeCoord.xy, InNodeCoord.xy + 1.0f) * Scale + QuadTreePosition.xyxy;
-}
-
 [numthreads(64, 1, 1)]
 void MainCS(uint3 DispatchThreadID : SV_DispatchThreadID)
 {
@@ -210,7 +229,7 @@ void MainCS(uint3 DispatchThreadID : SV_DispatchThreadID)
 	}
 
 	const uint LODLevel = NodeCoord.z;
-	const float4 NodeAABB2D = GetNodeAABB2D(NodeCoord);
+	const float4 NodeAABB2D = GetNodeAABB2D(NodeCoord, QuadTreePosition, LeafSize);
 	
 	// Bounds culling
 	{
@@ -221,9 +240,11 @@ void MainCS(uint3 DispatchThreadID : SV_DispatchThreadID)
 		}
 	}
 
-	// Frustum culling
+	// Frustum and HZB occlusion culling
+	float2 WaterZBounds = 0.0f;
+	bool bCrossesNearPlane = false;
 	{
-		const float2 WaterZBounds = WaterZBoundsTexture.Load(NodeCoord).xy * CaptureDepthRange + QuadTreePosition.z;
+		WaterZBounds = WaterZBoundsTexture.Load(NodeCoord).xy * CaptureDepthRange + QuadTreePosition.z;
 		const float3 NodeAABBMin = float3(NodeAABB2D.xy, WaterZBounds.x);
 		const float3 NodeAABBMax = float3(NodeAABB2D.zw, WaterZBounds.y);
 		const float3 NodeCenter = (NodeAABBMin + NodeAABBMax) * 0.5f;
@@ -242,6 +263,8 @@ void MainCS(uint3 DispatchThreadID : SV_DispatchThreadID)
 		{
 			return;
 		}
+
+		bCrossesNearPlane = Cull.bCrossesNearPlane;
 	}
 
 	// Compute the 2D distance between the observer and the tile. Will be zero if the observer is within the tile
@@ -264,7 +287,7 @@ void MainCS(uint3 DispatchThreadID : SV_DispatchThreadID)
 	{
 		// Does the parent tile intersect the LOD range of tiles at this LOD?
 		uint3 ParentNodeCoord = uint3(NodeCoord.xy / 2, NodeCoord.z + 1);
-		float4 ParentNodeAABB2D = GetNodeAABB2D(ParentNodeCoord);
+		float4 ParentNodeAABB2D = GetNodeAABB2D(ParentNodeCoord, QuadTreePosition, LeafSize);
 		const float ClosestDistanceToParentTile = sqrt(ComputeSquaredDistanceToPoint(ParentNodeAABB2D.xy, ParentNodeAABB2D.zw, ObserverPosition.xy));
 		const bool bParentTileIntersectsThisLOD = ClosestDistanceToParentTile <= GetLODDistance(LODLevel, LODScale) && LODLevel >= WaterLODParams.LowestLOD;
 
@@ -282,7 +305,7 @@ void MainCS(uint3 DispatchThreadID : SV_DispatchThreadID)
 				DensityLevel = ParentLOD - LODLevel;
 
 				ParentNodeCoord = uint3(NodeCoord.xy >> (ParentLOD - LODLevel), ParentLOD);
-				ParentNodeAABB2D = GetNodeAABB2D(ParentNodeCoord);
+				ParentNodeAABB2D = GetNodeAABB2D(ParentNodeCoord, QuadTreePosition, LeafSize);
 
 				const float Dist = sqrt(ComputeSquaredDistanceToPoint(ParentNodeAABB2D.xy, ParentNodeAABB2D.zw, ObserverPosition.xy));
 				if (Dist <= GetLODDistance(ParentLOD, LODScale) && ParentLOD >= WaterLODParams.LowestLOD)
@@ -305,59 +328,101 @@ void MainCS(uint3 DispatchThreadID : SV_DispatchThreadID)
 		const FWaterQuadTreeNode Node = WaterQuadTreeUnpackNodeRGBA8(QuadTreeTexture.Load(NodeCoord));
 		if (CanRender(Node, DensityLevel, ForceCollapseDensityLevel))
 		{
-			AddNodeForRender(PackNodeCoord(uint4(NodeCoord, DensityLevel)));
+			uint WritePos = 0;
+			PackedNodes.InterlockedAdd(0, 1, WritePos);
+			PackedNodes.Store((WritePos + 1) << 2u, PackNodeCoord(uint4(NodeCoord, DensityLevel))); // Num is stored at index 0, so we offset all writes by 1
+
+#if PRECISE_OCCLUSION_QUERIES
+			// Increase the InstanceCount by 1.
+			uint Dummy = 0;
+			InterlockedAdd(OcclusionQueryIndirectArgs[ViewIndex].InstanceCount, 1, Dummy);
+
+			// Write out bounding box and store a flag indicating if the box crosses the near plane.
+			const float3 NodeAABBMin = float3(NodeAABB2D.xy, WaterZBounds.x);
+			const float3 NodeAABBMax = float3(NodeAABB2D.zw, WaterZBounds.y);
+			const float3 NodeCenter = (NodeAABBMin + NodeAABBMax) * 0.5f;
+			const float3 NodeExtent = float3((NodeAABBMax.xy - NodeAABBMin.xy) * 0.5f, (NodeAABBMax.z - NodeAABBMin.z) * 0.5f + 1e-5f);
+			OcclusionQueryBoxes[WritePos * 2 + 0] = float4(NodeCenter, bCrossesNearPlane ? 1.0f : 0.0f);
+			OcclusionQueryBoxes[WritePos * 2 + 1] = float4(NodeExtent, 0.0f);
+
+			// Set bounding boxes crossing the near plane as always visible and initialize to 0 otherwise.
+			OcclusionVisibility[WritePos] = bCrossesNearPlane ? 1 : 0;
+#endif // PRECISE_OCCLUSION_QUERIES
+		}
+	}
+}
 
-#if DEBUG_SHADER_PRINT
-			if (DebugShowTile != 0)
-			{
-				FWaterBodyRenderData WBRenderData = WaterBodyRenderData[Node.WaterBodyRenderDataIndex];
+#endif // QUAD_TREE_TRAVERSE
 
-				// TODO: Support the two missing modes for visualizing the LOD level and the density
-				float4 DebugColor;
-				switch (WBRenderData.WaterBodyType)
-				{
-				case WATER_BODY_TYPE_RIVER: DebugColor = ColorRed; break;
-				case WATER_BODY_TYPE_LAKE: DebugColor = ColorGreen; break;
-				case WATER_BODY_TYPE_OCEAN: DebugColor = ColorBlue; break;
-				default: DebugColor = ColorWhite;
-				}
+#ifdef OCCLUSION_QUERY_RASTER_VS
 
-				if (WBRenderData.WaterBodyType == WATER_BODY_TYPE_RIVER && Node.TransitionWaterBodyRenderDataIndex > 0)
-				{
-					const FWaterBodyRenderData TransitionWBRenderData = WaterBodyRenderData[Node.TransitionWaterBodyRenderDataIndex];
-					if (TransitionWBRenderData.WaterBodyType == WATER_BODY_TYPE_LAKE)
-					{
-						DebugColor = ColorYellow;
-					}
-					else if (TransitionWBRenderData.WaterBodyType == WATER_BODY_TYPE_OCEAN)
-					{
-						DebugColor = ColorPurple;
-					}
-				}
+Buffer<float4> OcclusionQueryBoxes;
+float4x4 ViewProjection;
 
-				const float2 WaterZBounds = WaterZBoundsTexture.Load(NodeCoord).xy * CaptureDepthRange + QuadTreePosition.z;
-				const float3 AABBMin = float3(NodeAABB2D.xy + 20.0f, WaterZBounds.x);
-				const float3 AABBMax = float3(NodeAABB2D.zw - 20.0f, WaterZBounds.y);
-				
-				AddOBBTWS(InitShaderPrintContext(), AABBMin, AABBMax, DebugColor, float4x4(
-					1, 0, 0, 0, 
-					0, 1, 0, 0, 
-					0, 0, 1, 0, 
-					0, 0, 0, 1));
-			}
-#endif // DEBUG_SHADER_PRINT
-		}
+void MainVS(
+	in float3 InPosition : ATTRIBUTE0,
+	in uint InInstanceId : SV_InstanceID,
+	out nointerpolation uint OutQueryIndex : QUERY_INDEX,
+	out float4 OutPosition : SV_Position
+)
+{
+	OutQueryIndex = InInstanceId;
+
+	const float4 CenterAndCrossesNearPlane = OcclusionQueryBoxes[InInstanceId * 2 + 0];
+	const bool bCrossesNearPlane = CenterAndCrossesNearPlane.w != 0.0f;
+
+	if (!bCrossesNearPlane)
+	{
+		const float3 Center = CenterAndCrossesNearPlane.xyz;
+		const float3 Extent = OcclusionQueryBoxes[InInstanceId * 2 + 1].xyz;
+		const float3 TranslatedWorldPosition = (InPosition * Extent) + Center;
+		OutPosition = mul(float4(TranslatedWorldPosition, 1.0f), ViewProjection);
+	}
+	else
+	{
+		// Kill instance by setting all vertices to NaN. If the box crosses the near plane, we already set it as visible in the result buffer.
+		OutPosition = asfloat(0xFFFFFFFF).xxxx;
 	}
 }
 
-#endif // QUAD_TREE_TRAVERSE
+#endif // OCCLUSION_QUERY_RASTER_VS
+
+#ifdef OCCLUSION_QUERY_RASTER_PS
+
+RWBuffer<uint> Visibility;
+
+EARLYDEPTHSTENCIL
+void MainPS(
+	in uint InQueryIndex : QUERY_INDEX, 
+	out float4 OutColor : SV_Target0
+)
+{
+#if COMPILER_SUPPORTS_WAVE_VOTE
+	if (WaveIsFirstLane()) // Try to reduce the number of threads writing the same value to the same location.
+	{
+		Visibility[InQueryIndex] = 1;
+	}
+#else
+	Visibility[InQueryIndex] = 1;
+#endif
+	OutColor = float4(1.0f, 0.0f, 0.0f, 1.0f);
+}
+
+#endif // OCCLUSION_QUERY_RASTER_PS
 
 #ifdef COMPUTE_BUCKET_COUNTS
 
+#ifndef PRECISE_OCCLUSION_QUERIES
+#define PRECISE_OCCLUSION_QUERIES 0
+#endif
+
 RWByteAddressBuffer BucketCounts;
 Texture2D<float4> QuadTreeTexture;
 StructuredBuffer<FWaterBodyRenderData> WaterBodyRenderData;
 ByteAddressBuffer PackedNodes;
+#if PRECISE_OCCLUSION_QUERIES
+Buffer<uint> OcclusionResults;
+#endif // PRECISE_OCCLUSION_QUERIES
 uint NumDensities;
 uint NumMaterials;
 uint NumDispatchedThreads;
@@ -370,8 +435,16 @@ void MainCS(uint3 DispatchThreadID : SV_DispatchThreadID)
 	// We are limited to 65535 threads, but might end up with more draws than that
 	for (uint DrawIndex = DispatchThreadID.x; DrawIndex < NumDraws; DrawIndex += NumDispatchedThreads)
 	{
+#if PRECISE_OCCLUSION_QUERIES
+		const bool bIsVisible = OcclusionResults[DrawIndex] != 0;
+		if (!bIsVisible)
+		{
+			continue;
+		}
+#endif
+
 		// Load packed node coord and density index
-		uint4 NodeCoordAndDensity = UnpackNodeCoord(PackedNodes.Load((DrawIndex + 1) << 2));
+		const uint4 NodeCoordAndDensity = UnpackNodeCoord(PackedNodes.Load((DrawIndex + 1) << 2));
 		const uint DensityIndex = min(NodeCoordAndDensity.w, NumDensities - 1);
 
 		// Sample the quadtree to access the water body render data
@@ -547,6 +620,10 @@ void MainCS(uint3 DispatchThreadID : SV_DispatchThreadID)
 
 #ifdef GENERATE_INSTANCE_DATA
 
+#ifndef PRECISE_OCCLUSION_QUERIES
+#define PRECISE_OCCLUSION_QUERIES 0
+#endif
+
 RWStructuredBuffer<FWaterTileDrawIndexedIndirectArgs> IndirectArgs;
 RWBuffer<float4> InstanceData0;
 RWBuffer<float4> InstanceData1;
@@ -556,6 +633,9 @@ Texture2D<float4> WaterZBoundsTexture;
 StructuredBuffer<FWaterBodyRenderData> WaterBodyRenderData;
 ByteAddressBuffer PackedNodes;
 Buffer<uint> InstanceDataOffsets;
+#if PRECISE_OCCLUSION_QUERIES
+Buffer<uint> OcclusionResults;
+#endif // PRECISE_OCCLUSION_QUERIES
 
 float3 QuadTreePosition;
 float3 ObserverPosition;
@@ -579,8 +659,16 @@ void MainCS(uint3 DispatchThreadID : SV_DispatchThreadID)
 	// We are limited to 65535 threads, but might end up with more draws than that
 	for (uint DrawIndex = DispatchThreadID.x; DrawIndex < NumDraws; DrawIndex += NumDispatchedThreads)
 	{
+#if PRECISE_OCCLUSION_QUERIES
+		const bool bIsVisible = OcclusionResults[DrawIndex] != 0;
+		if (!bIsVisible)
+		{
+			continue;
+		}
+#endif
+
 		// Load packed node coord and density index
-		uint4 NodeCoordAndDensity = UnpackNodeCoord(PackedNodes.Load((DrawIndex + 1) << 2));
+		const uint4 NodeCoordAndDensity = UnpackNodeCoord(PackedNodes.Load((DrawIndex + 1) << 2));
 		const uint2 NodeCoord = NodeCoordAndDensity.xy;
 		const uint LODLevel = NodeCoordAndDensity.z;
 		const uint DensityIndex = NodeCoordAndDensity.w;
@@ -672,4 +760,87 @@ void MainCS(uint3 DispatchThreadID : SV_DispatchThreadID)
 	}
 }
 
-#endif // GENERATE_INSTANCE_DATA
\ No newline at end of file
+#endif // GENERATE_INSTANCE_DATA
+
+#if DEBUG_SHOW_TILES
+
+#include "/Engine/Private/ShaderPrint.ush"
+
+#ifndef PRECISE_OCCLUSION_QUERIES
+#define PRECISE_OCCLUSION_QUERIES 0
+#endif
+
+Texture2D<float4> QuadTreeTexture;
+Texture2D<float4> WaterZBoundsTexture;
+StructuredBuffer<FWaterBodyRenderData> WaterBodyRenderData;
+ByteAddressBuffer PackedNodes;
+#if PRECISE_OCCLUSION_QUERIES
+Buffer<uint> OcclusionResults;
+#endif // PRECISE_OCCLUSION_QUERIES
+
+float3 QuadTreePosition;
+uint NumDispatchedThreads;
+float LeafSize;
+float CaptureDepthRange;
+
+[numthreads(64, 1, 1)]
+void MainCS(uint3 DispatchThreadID : SV_DispatchThreadID)
+{
+	const uint NumDraws = PackedNodes.Load(0);
+	// We are limited to 65535 threads, but might end up with more draws than that
+	for (uint DrawIndex = DispatchThreadID.x; DrawIndex < NumDraws; DrawIndex += NumDispatchedThreads)
+	{
+#if PRECISE_OCCLUSION_QUERIES
+		const bool bIsVisible = OcclusionResults[DrawIndex] != 0;
+		if (!bIsVisible)
+		{
+			continue;
+		}
+#endif
+
+		// Load packed node coord and density index
+		const uint4 NodeCoordAndDensity = UnpackNodeCoord(PackedNodes.Load((DrawIndex + 1) << 2));
+		const uint3 NodeCoord = NodeCoordAndDensity.xyz;
+
+		// Sample the quadtree to access the water body render data
+		const float4 QuadTreeSample = QuadTreeTexture.Load(int3(NodeCoord));
+		const FWaterQuadTreeNode Node = WaterQuadTreeUnpackNodeRGBA8(QuadTreeSample);
+		const FWaterBodyRenderData WBRenderData = WaterBodyRenderData[Node.WaterBodyRenderDataIndex];
+
+		// TODO: Support the two missing modes for visualizing the LOD level and the density
+		float4 DebugColor;
+		switch (WBRenderData.WaterBodyType)
+		{
+		case WATER_BODY_TYPE_RIVER: DebugColor = ColorRed; break;
+		case WATER_BODY_TYPE_LAKE: DebugColor = ColorGreen; break;
+		case WATER_BODY_TYPE_OCEAN: DebugColor = ColorBlue; break;
+		default: DebugColor = ColorWhite;
+		}
+
+		if (WBRenderData.WaterBodyType == WATER_BODY_TYPE_RIVER && Node.TransitionWaterBodyRenderDataIndex > 0)
+		{
+			const FWaterBodyRenderData TransitionWBRenderData = WaterBodyRenderData[Node.TransitionWaterBodyRenderDataIndex];
+			if (TransitionWBRenderData.WaterBodyType == WATER_BODY_TYPE_LAKE)
+			{
+				DebugColor = ColorYellow;
+			}
+			else if (TransitionWBRenderData.WaterBodyType == WATER_BODY_TYPE_OCEAN)
+			{
+				DebugColor = ColorPurple;
+			}
+		}
+
+		const float4 NodeAABB2D = GetNodeAABB2D(NodeCoord, QuadTreePosition, LeafSize);
+		const float2 WaterZBounds = WaterZBoundsTexture.Load(NodeCoord).xy * CaptureDepthRange + QuadTreePosition.z;
+		const float3 AABBMin = float3(NodeAABB2D.xy + 20.0f, WaterZBounds.x);
+		const float3 AABBMax = float3(NodeAABB2D.zw - 20.0f, WaterZBounds.y);
+
+		AddOBBTWS(InitShaderPrintContext(), AABBMin, AABBMax, DebugColor, float4x4(
+			1, 0, 0, 0,
+			0, 1, 0, 0,
+			0, 0, 1, 0,
+			0, 0, 0, 1));
+	}
+}
+
+#endif // DEBUG_SHOW_TILES
\ No newline at end of file
diff --git a/Engine/Plugins/Experimental/Water/Source/Runtime/Private/WaterQuadTreeGPU.cpp b/Engine/Plugins/Experimental/Water/Source/Runtime/Private/WaterQuadTreeGPU.cpp
--- a/Engine/Plugins/Experimental/Water/Source/Runtime/Private/WaterQuadTreeGPU.cpp
+++ b/Engine/Plugins/Experimental/Water/Source/Runtime/Private/WaterQuadTreeGPU.cpp
@@ -33,11 +33,11 @@ static FAutoConsoleVariableRef CVarWaterQuadTreeZBoundsPadding(
 		" Necessary for sloped rivers, whose complete Z bounds can be underestimated due to using non-conservative rasterization to build the quadtree."),
 	ECVF_RenderThreadSafe);
 
-int32 GWaterQuadTreeOcclusionCulling = 1;
+int32 GWaterQuadTreeOcclusionCulling = 3;
 static FAutoConsoleVariableRef CVarWaterQuadTreeOcclusionCulling(
 	TEXT("r.Water.WaterMesh.GPUQuadTree.OcclusionCulling"),
 	GWaterQuadTreeOcclusionCulling,
-	TEXT("Enables HZB occlusion culling for the water quadtree indirect draw call pipeline."),
+	TEXT("0: Disabled, 1: HZB Occlusion Queries, 2: Pixel Precise Raster Queries, 3: HZB + Pixel Precise Raster Queries"),
 	ECVF_RenderThreadSafe);
 
 class FWaterQuadTreeVS : public FGlobalShader
@@ -139,8 +139,12 @@ public:
 	DECLARE_GLOBAL_SHADER(FWaterQuadTreeInitializeIndirectArgsCS);
 	SHADER_USE_PARAMETER_STRUCT(FWaterQuadTreeInitializeIndirectArgsCS, FGlobalShader);
 
+	class FPreciseOcclusionQueries : SHADER_PERMUTATION_BOOL("PRECISE_OCCLUSION_QUERIES");
+	using FPermutationDomain = TShaderPermutationDomain<FPreciseOcclusionQueries>;
+
 	BEGIN_SHADER_PARAMETER_STRUCT(FParameters, )
 		SHADER_PARAMETER_RDG_BUFFER_UAV(RWStructuredBuffer, IndirectArgs)
+		SHADER_PARAMETER_RDG_BUFFER_UAV(RWStructuredBuffer, OcclusionQueryIndirectArgs)
 		SHADER_PARAMETER(uint32, NumDensities)
 		SHADER_PARAMETER(uint32, NumMaterials)
 		SHADER_PARAMETER(uint32, NumViews)
@@ -161,13 +165,14 @@ public:
 	DECLARE_GLOBAL_SHADER(FWaterQuadTreeTraverseCS);
 	SHADER_USE_PARAMETER_STRUCT(FWaterQuadTreeTraverseCS, FGlobalShader);
 
-	class FDebugShaderPrint : SHADER_PERMUTATION_BOOL("DEBUG_SHADER_PRINT");
-
-	using FPermutationDomain = TShaderPermutationDomain<FDebugShaderPrint>;
+	class FPreciseOcclusionQueries : SHADER_PERMUTATION_BOOL("PRECISE_OCCLUSION_QUERIES");
+	using FPermutationDomain = TShaderPermutationDomain<FPreciseOcclusionQueries>;
 
 	BEGIN_SHADER_PARAMETER_STRUCT(FParameters, )
-		SHADER_PARAMETER_STRUCT_INCLUDE(ShaderPrint::FShaderParameters, ShaderPrintUniformBuffer)
 		SHADER_PARAMETER_RDG_BUFFER_UAV(RWByteAddressBuffer, PackedNodes)
+		SHADER_PARAMETER_RDG_BUFFER_UAV(RWBuffer, OcclusionQueryBoxes)
+		SHADER_PARAMETER_RDG_BUFFER_UAV(RWBuffer, OcclusionVisibility)
+		SHADER_PARAMETER_RDG_BUFFER_UAV(RWStructuredBuffer, OcclusionQueryIndirectArgs)
 		SHADER_PARAMETER_RDG_TEXTURE(Texture2D, QuadTreeTexture)
 		SHADER_PARAMETER_RDG_TEXTURE(Texture2D, WaterZBoundsTexture)
 		SHADER_PARAMETER_RDG_BUFFER_SRV(StructuredBuffer, WaterBodyRenderData)
@@ -184,36 +189,67 @@ public:
 		SHADER_PARAMETER(uint32, QuadTreeResolutionY)
 		SHADER_PARAMETER(uint32, NumDensities)
 		SHADER_PARAMETER(uint32, NumMaterials)
+		SHADER_PARAMETER(uint32, ViewIndex)
 		SHADER_PARAMETER(float, LeafSize)
 		SHADER_PARAMETER(float, LODScale)
 		SHADER_PARAMETER(float, CaptureDepthRange)
 		SHADER_PARAMETER(int32, ForceCollapseDensityLevel)
 		SHADER_PARAMETER(uint32, NumLODs)
 		SHADER_PARAMETER(uint32, NumDispatchedThreads)
-		SHADER_PARAMETER(int32, DebugShowTile)
 		SHADER_PARAMETER(uint32, bHZBOcclusionCullingEnabled)
 	END_SHADER_PARAMETER_STRUCT()
 
-	static bool ShouldCompilePermutation(const FGlobalShaderPermutationParameters& Parameters)
+	static void ModifyCompilationEnvironment(const FGlobalShaderPermutationParameters& Parameters, FShaderCompilerEnvironment& OutEnvironment)
 	{
-		FPermutationDomain PermutationVector(Parameters.PermutationId);
-		const bool bDebugShaderPrint = PermutationVector.Get<FDebugShaderPrint>() != 0;
-		return !bDebugShaderPrint || ShaderPrint::IsSupported(Parameters.Platform);
+		FGlobalShader::ModifyCompilationEnvironment(Parameters, OutEnvironment);
+		OutEnvironment.SetDefine(TEXT("QUAD_TREE_TRAVERSE"), 1);
 	}
+};
+IMPLEMENT_GLOBAL_SHADER(FWaterQuadTreeTraverseCS, "/Plugin/Water/Private/WaterQuadTreeDraws.usf", "MainCS", SF_Compute);
+
+class FWaterQuadTreeOcclusionQueryVS : public FGlobalShader
+{
+public:
+	DECLARE_GLOBAL_SHADER(FWaterQuadTreeOcclusionQueryVS);
+	SHADER_USE_PARAMETER_STRUCT(FWaterQuadTreeOcclusionQueryVS, FGlobalShader);
+
+	BEGIN_SHADER_PARAMETER_STRUCT(FParameters, )
+		SHADER_PARAMETER_RDG_BUFFER_SRV(Buffer, OcclusionQueryBoxes)
+		SHADER_PARAMETER(FMatrix44f, ViewProjection)
+	END_SHADER_PARAMETER_STRUCT()
 
 	static void ModifyCompilationEnvironment(const FGlobalShaderPermutationParameters& Parameters, FShaderCompilerEnvironment& OutEnvironment)
 	{
 		FGlobalShader::ModifyCompilationEnvironment(Parameters, OutEnvironment);
-		FPermutationDomain PermutationVector(Parameters.PermutationId);
-		const bool bDebugShaderPrint = PermutationVector.Get<FDebugShaderPrint>() != 0;
-		if (bDebugShaderPrint)
-		{
-			ShaderPrint::ModifyCompilationEnvironment(Parameters, OutEnvironment);
-		}
-		OutEnvironment.SetDefine(TEXT("QUAD_TREE_TRAVERSE"), 1);
+		OutEnvironment.SetDefine(TEXT("OCCLUSION_QUERY_RASTER_VS"), 1);
 	}
 };
-IMPLEMENT_GLOBAL_SHADER(FWaterQuadTreeTraverseCS, "/Plugin/Water/Private/WaterQuadTreeDraws.usf", "MainCS", SF_Compute);
+IMPLEMENT_GLOBAL_SHADER(FWaterQuadTreeOcclusionQueryVS, "/Plugin/Water/Private/WaterQuadTreeDraws.usf", "MainVS", SF_Vertex);
+
+class FWaterQuadTreeOcclusionQueryPS : public FGlobalShader
+{
+public:
+	DECLARE_GLOBAL_SHADER(FWaterQuadTreeOcclusionQueryPS);
+	SHADER_USE_PARAMETER_STRUCT(FWaterQuadTreeOcclusionQueryPS, FGlobalShader);
+
+	BEGIN_SHADER_PARAMETER_STRUCT(FParameters, )
+		SHADER_PARAMETER_RDG_BUFFER_UAV(RWBuffer, Visibility)
+		RENDER_TARGET_BINDING_SLOTS()
+	END_SHADER_PARAMETER_STRUCT()
+
+	static void ModifyCompilationEnvironment(const FGlobalShaderPermutationParameters& Parameters, FShaderCompilerEnvironment& OutEnvironment)
+	{
+		FGlobalShader::ModifyCompilationEnvironment(Parameters, OutEnvironment);
+		OutEnvironment.SetDefine(TEXT("OCCLUSION_QUERY_RASTER_PS"), 1);
+	}
+};
+IMPLEMENT_GLOBAL_SHADER(FWaterQuadTreeOcclusionQueryPS, "/Plugin/Water/Private/WaterQuadTreeDraws.usf", "MainPS", SF_Pixel);
+
+BEGIN_SHADER_PARAMETER_STRUCT(FWaterQuadTreeOcclusionQueryParameters, )
+	RDG_BUFFER_ACCESS(IndirectDrawArgsBuffer, ERHIAccess::IndirectArgs)
+	SHADER_PARAMETER_STRUCT_INCLUDE(FWaterQuadTreeOcclusionQueryVS::FParameters, VS)
+	SHADER_PARAMETER_STRUCT_INCLUDE(FWaterQuadTreeOcclusionQueryPS::FParameters, PS)
+END_SHADER_PARAMETER_STRUCT()
 
 class FWaterQuadTreeBucketCountsCS : public FGlobalShader
 {
@@ -221,11 +257,15 @@ public:
 	DECLARE_GLOBAL_SHADER(FWaterQuadTreeBucketCountsCS);
 	SHADER_USE_PARAMETER_STRUCT(FWaterQuadTreeBucketCountsCS, FGlobalShader);
 
+	class FPreciseOcclusionQueries : SHADER_PERMUTATION_BOOL("PRECISE_OCCLUSION_QUERIES");
+	using FPermutationDomain = TShaderPermutationDomain<FPreciseOcclusionQueries>;
+
 	BEGIN_SHADER_PARAMETER_STRUCT(FParameters, )
 		SHADER_PARAMETER_RDG_BUFFER_UAV(RWByteAddressBuffer, BucketCounts)
 		SHADER_PARAMETER_RDG_TEXTURE(Texture2D, QuadTreeTexture)
 		SHADER_PARAMETER_RDG_BUFFER_SRV(StructuredBuffer, WaterBodyRenderData)
 		SHADER_PARAMETER_RDG_BUFFER_SRV(ByteAddressBuffer, PackedNodes)
+		SHADER_PARAMETER_RDG_BUFFER_SRV(Buffer, OcclusionResults)
 		SHADER_PARAMETER(uint32, NumDensities)
 		SHADER_PARAMETER(uint32, NumMaterials)
 		SHADER_PARAMETER(uint32, NumDispatchedThreads)
@@ -269,6 +309,9 @@ public:
 	DECLARE_GLOBAL_SHADER(FWaterQuadTreeGenerateInstanceDataCS);
 	SHADER_USE_PARAMETER_STRUCT(FWaterQuadTreeGenerateInstanceDataCS, FGlobalShader);
 
+	class FPreciseOcclusionQueries : SHADER_PERMUTATION_BOOL("PRECISE_OCCLUSION_QUERIES");
+	using FPermutationDomain = TShaderPermutationDomain<FPreciseOcclusionQueries>;
+
 	BEGIN_SHADER_PARAMETER_STRUCT(FParameters, )
 		SHADER_PARAMETER_RDG_BUFFER_UAV(RWStructuredBuffer, IndirectArgs)
 		SHADER_PARAMETER_RDG_BUFFER_UAV(RWBuffer, InstanceData0)
@@ -279,6 +322,7 @@ public:
 		SHADER_PARAMETER_RDG_BUFFER_SRV(StructuredBuffer, WaterBodyRenderData)
 		SHADER_PARAMETER_RDG_BUFFER_SRV(ByteAddressBuffer, PackedNodes)
 		SHADER_PARAMETER_RDG_BUFFER_SRV(Buffer, InstanceDataOffsets)
+		SHADER_PARAMETER_RDG_BUFFER_SRV(Buffer, OcclusionResults)
 		SHADER_PARAMETER(FVector3f, QuadTreePosition)
 		SHADER_PARAMETER(FVector3f, ObserverPosition)
 		SHADER_PARAMETER(uint32, QuadTreeResolutionX)
@@ -303,6 +347,42 @@ public:
 };
 IMPLEMENT_GLOBAL_SHADER(FWaterQuadTreeGenerateInstanceDataCS, "/Plugin/Water/Private/WaterQuadTreeDraws.usf", "MainCS", SF_Compute);
 
+class FWaterQuadTreeDebugCS : public FGlobalShader
+{
+public:
+	DECLARE_GLOBAL_SHADER(FWaterQuadTreeDebugCS);
+	SHADER_USE_PARAMETER_STRUCT(FWaterQuadTreeDebugCS, FGlobalShader);
+
+	class FPreciseOcclusionQueries : SHADER_PERMUTATION_BOOL("PRECISE_OCCLUSION_QUERIES");
+	using FPermutationDomain = TShaderPermutationDomain<FPreciseOcclusionQueries>;
+
+	BEGIN_SHADER_PARAMETER_STRUCT(FParameters, )
+		SHADER_PARAMETER_STRUCT_INCLUDE(ShaderPrint::FShaderParameters, ShaderPrintUniformBuffer)
+		SHADER_PARAMETER_RDG_TEXTURE(Texture2D, QuadTreeTexture)
+		SHADER_PARAMETER_RDG_TEXTURE(Texture2D, WaterZBoundsTexture)
+		SHADER_PARAMETER_RDG_BUFFER_SRV(StructuredBuffer, WaterBodyRenderData)
+		SHADER_PARAMETER_RDG_BUFFER_SRV(ByteAddressBuffer, PackedNodes)
+		SHADER_PARAMETER_RDG_BUFFER_SRV(Buffer, OcclusionResults)
+		SHADER_PARAMETER(FVector3f, QuadTreePosition)
+		SHADER_PARAMETER(uint32, NumDispatchedThreads)
+		SHADER_PARAMETER(float, LeafSize)
+		SHADER_PARAMETER(float, CaptureDepthRange)
+	END_SHADER_PARAMETER_STRUCT()
+
+	static bool ShouldCompilePermutation(const FGlobalShaderPermutationParameters& Parameters)
+	{
+		return ShaderPrint::IsSupported(Parameters.Platform);
+	}
+
+	static void ModifyCompilationEnvironment(const FGlobalShaderPermutationParameters& Parameters, FShaderCompilerEnvironment& OutEnvironment)
+	{
+		FGlobalShader::ModifyCompilationEnvironment(Parameters, OutEnvironment);
+		ShaderPrint::ModifyCompilationEnvironment(Parameters, OutEnvironment);
+		OutEnvironment.SetDefine(TEXT("DEBUG_SHOW_TILES"), 1);
+	}
+};
+IMPLEMENT_GLOBAL_SHADER(FWaterQuadTreeDebugCS, "/Plugin/Water/Private/WaterQuadTreeDraws.usf", "MainCS", SF_Compute);
+
 
 void FWaterQuadTreeGPU::Init(FRDGBuilder& GraphBuilder, const FInitParams& Params, TArray<FDraw>& Draws)
 {
@@ -431,67 +511,66 @@ void FWaterQuadTreeGPU::Init(FRDGBuilder& GraphBuilder, const FInitParams& Param
 			});
 	}
 
-		// Merge river and non-river water bodies and downsample to quad tree LOD0 resolution
+	// Merge river and non-river water bodies and downsample to quad tree LOD0 resolution
+	{
+		FWaterQuadTreeMergePS::FPermutationDomain PermutationDomain;
+		PermutationDomain.Set<FWaterQuadTreeMergePS::FNumMSAASamples>(NumSamples);
+		TShaderMapRef<FWaterQuadTreeMergePS> PixelShader(ShaderMap, PermutationDomain);
+
+		FWaterQuadTreeMergePS::FParameters* PassParameters = GraphBuilder.AllocParameters<FWaterQuadTreeMergePS::FParameters>();
+		if (NumSamples > 1)
 		{
-			FWaterQuadTreeMergePS::FPermutationDomain PermutationDomain;
-			PermutationDomain.Set<FWaterQuadTreeMergePS::FNumMSAASamples>(NumSamples);
-			TShaderMapRef<FWaterQuadTreeMergePS> PixelShader(ShaderMap, PermutationDomain);
+			PassParameters->WaterBodyRasterTextureMS = WaterBodyRasterTexture;
+			PassParameters->ZBoundsRasterTextureMS = ZBoundsRasterTexture;
+		}
+		else
+		{
+			PassParameters->WaterBodyRasterTexture = WaterBodyRasterTexture;
+			PassParameters->ZBoundsRasterTexture = ZBoundsRasterTexture;
+		}
+		PassParameters->WaterBodyRenderData = WaterBodyRenderDataBufferSRV;
+		PassParameters->SuperSamplingFactor = Params.SuperSamplingFactor;
+		PassParameters->RcpCaptureDepthRange = 1.0f / Params.CaptureDepthRange;
+		PassParameters->ZBoundsPadding = GWaterQuadTreeZBoundsPadding / Params.CaptureDepthRange;
+		PassParameters->RenderTargets[0] = FRenderTargetBinding(QuadTreeTextureRDG, ERenderTargetLoadAction::ENoAction, 0);
+		PassParameters->RenderTargets[1] = FRenderTargetBinding(WaterZBoundsTextureRDG, ERenderTargetLoadAction::ENoAction, 0);
+
+		const FIntRect MergeViewport(0, 0, QuadTreeResolution.X, QuadTreeResolution.Y);
+
+		FPixelShaderUtils::AddFullscreenPass(
+			GraphBuilder,
+			ShaderMap,
+			RDG_EVENT_NAME("WaterQuadTreeMerge"),
+			PixelShader,
+			PassParameters,
+			MergeViewport);
+	}
 
-			FWaterQuadTreeMergePS::FParameters* PassParameters = GraphBuilder.AllocParameters<FWaterQuadTreeMergePS::FParameters>();
-			if (NumSamples > 1)
-			{
-				PassParameters->WaterBodyRasterTextureMS = WaterBodyRasterTexture;
-				PassParameters->ZBoundsRasterTextureMS = ZBoundsRasterTexture;
-			}
-			else
-			{
-				PassParameters->WaterBodyRasterTexture = WaterBodyRasterTexture;
-				PassParameters->ZBoundsRasterTexture = ZBoundsRasterTexture;
-			}
-			PassParameters->WaterBodyRenderData = WaterBodyRenderDataBufferSRV;
-			PassParameters->SuperSamplingFactor = Params.SuperSamplingFactor;
-			PassParameters->RcpCaptureDepthRange = 1.0f / Params.CaptureDepthRange;
-			PassParameters->ZBoundsPadding = GWaterQuadTreeZBoundsPadding / Params.CaptureDepthRange;
-			PassParameters->RenderTargets[0] = FRenderTargetBinding(QuadTreeTextureRDG, ERenderTargetLoadAction::ENoAction, 0);
-			PassParameters->RenderTargets[1] = FRenderTargetBinding(WaterZBoundsTextureRDG, ERenderTargetLoadAction::ENoAction, 0);
+	// Build the mip chain
+	{
+		TShaderMapRef<FWaterQuadTreeBuildPS> PixelShader(ShaderMap);
+
+		FIntRect BuildViewport(0, 0, QuadTreeResolution.X, QuadTreeResolution.Y);
+
+		for (int32 MipLevel = 1; MipLevel < NumMipLevels; ++MipLevel)
+		{
+			FWaterQuadTreeBuildPS::FParameters* PassParameters = GraphBuilder.AllocParameters<FWaterQuadTreeBuildPS::FParameters>();
+			PassParameters->QuadTreeTexture = GraphBuilder.CreateSRV(FRDGTextureSRVDesc::CreateForMipLevel(QuadTreeTextureRDG, MipLevel - 1));
+			PassParameters->WaterZBoundsTexture = GraphBuilder.CreateSRV(FRDGTextureSRVDesc::CreateForMipLevel(WaterZBoundsTextureRDG, MipLevel - 1));
+			PassParameters->RenderTargets[0] = FRenderTargetBinding(QuadTreeTextureRDG, ERenderTargetLoadAction::ENoAction, MipLevel);
+			PassParameters->RenderTargets[1] = FRenderTargetBinding(WaterZBoundsTextureRDG, ERenderTargetLoadAction::ENoAction, MipLevel);
 
-			const FIntRect MergeViewport(0, 0, QuadTreeResolution.X, QuadTreeResolution.Y);
+			BuildViewport = BuildViewport / 2;
 
 			FPixelShaderUtils::AddFullscreenPass(
 				GraphBuilder,
 				ShaderMap,
-				RDG_EVENT_NAME("WaterQuadTreeMerge"),
+				RDG_EVENT_NAME("WaterQuadTreeBuild"),
 				PixelShader,
 				PassParameters,
-				MergeViewport);
+				BuildViewport);
 		}
-
-		// Build the mip chain
-		{
-			TShaderMapRef<FWaterQuadTreeBuildPS> PixelShader(ShaderMap);
-
-			FIntRect BuildViewport(0, 0, QuadTreeResolution.X, QuadTreeResolution.Y);
-
-			for (int32 MipLevel = 1; MipLevel < NumMipLevels; ++MipLevel)
-			{
-				FWaterQuadTreeBuildPS::FParameters* PassParameters = GraphBuilder.AllocParameters<FWaterQuadTreeBuildPS::FParameters>();
-				PassParameters->QuadTreeTexture = GraphBuilder.CreateSRV(FRDGTextureSRVDesc::CreateForMipLevel(QuadTreeTextureRDG, MipLevel - 1));
-				PassParameters->WaterZBoundsTexture = GraphBuilder.CreateSRV(FRDGTextureSRVDesc::CreateForMipLevel(WaterZBoundsTextureRDG, MipLevel - 1));
-				PassParameters->RenderTargets[0] = FRenderTargetBinding(QuadTreeTextureRDG, ERenderTargetLoadAction::ENoAction, MipLevel);
-				PassParameters->RenderTargets[1] = FRenderTargetBinding(WaterZBoundsTextureRDG, ERenderTargetLoadAction::ENoAction, MipLevel);
-
-				BuildViewport = BuildViewport / 2;
-
-				FPixelShaderUtils::AddFullscreenPass(
-					GraphBuilder,
-					ShaderMap,
-					RDG_EVENT_NAME("WaterQuadTreeBuild"),
-					PixelShader,
-					PassParameters,
-					BuildViewport);
-			}
-		}
-	
+	}
 }
 
 void FWaterQuadTreeGPU::Traverse(FRDGBuilder& GraphBuilder, const FTraverseParams& Params) const
@@ -518,16 +597,31 @@ void FWaterQuadTreeGPU::Traverse(FRDGBuilder& GraphBuilder, const FTraverseParam
 	FRDGBufferSRV* InstanceDataOffsetsBufferSRV = GraphBuilder.CreateSRV(FRDGBufferSRVDesc(InstanceDataOffsetsBuffer, PF_R32_UINT));
 	FRDGBufferSRV* WaterBodyRenderDataBufferSRV = GraphBuilder.CreateSRV(FRDGBufferSRVDesc(WaterBodyRenderDataBufferRDG));
 
+	const EOcclusionQueryMode OcclusionQueryMode = static_cast<EOcclusionQueryMode>(GWaterQuadTreeOcclusionCulling);
+	const bool bHZBOcclusionQueries = OcclusionQueryMode == EOcclusionQueryMode::HZB || OcclusionQueryMode == EOcclusionQueryMode::HZBAndPixelPrecise;
+	const bool bPixelPreciseOcclusionQueries = OcclusionQueryMode == EOcclusionQueryMode::PixelPrecise || OcclusionQueryMode == EOcclusionQueryMode::HZBAndPixelPrecise;
+	
+	FRDGBuffer* OcclusionQueryIndirectArgsBuffer = nullptr;
+	FRDGBufferUAV* OcclusionQueryIndirectArgsBufferUAV = nullptr;
+	if (bPixelPreciseOcclusionQueries)
+	{
+		OcclusionQueryIndirectArgsBuffer = GraphBuilder.CreateBuffer(FRDGBufferDesc::CreateIndirectDesc<FRHIDrawIndexedIndirectParameters>(Params.NumViews), TEXT("WaterQuadTree.OcclusionQueryIndirectArgs"));
+		OcclusionQueryIndirectArgsBufferUAV = GraphBuilder.CreateUAV(FRDGBufferUAVDesc(OcclusionQueryIndirectArgsBuffer));
+	}
+
 	const uint32 NumBucketsPerView = Params.NumDensities * Params.NumMaterials;
 	const uint32 NumBucketsTotal = NumBucketsPerView * Params.NumViews;
 	const FIntPoint QuadTreeResolution = QuadTreeTextureRDG->Desc.Extent;
 
 	// Initialize indirect args
 	{
-		TShaderMapRef<FWaterQuadTreeInitializeIndirectArgsCS> ComputeShader(ShaderMap);
+		FWaterQuadTreeInitializeIndirectArgsCS::FPermutationDomain PermutationDomain;
+		PermutationDomain.Set<FWaterQuadTreeInitializeIndirectArgsCS::FPreciseOcclusionQueries>(bPixelPreciseOcclusionQueries);
+		TShaderMapRef<FWaterQuadTreeInitializeIndirectArgsCS> ComputeShader(ShaderMap, PermutationDomain);
 
 		FWaterQuadTreeInitializeIndirectArgsCS::FParameters* PassParameters = GraphBuilder.AllocParameters<FWaterQuadTreeInitializeIndirectArgsCS::FParameters>();
 		PassParameters->IndirectArgs = IndirectArgsBufferUAV;
+		PassParameters->OcclusionQueryIndirectArgs = OcclusionQueryIndirectArgsBufferUAV;
 		PassParameters->NumDensities = Params.NumDensities;
 		PassParameters->NumMaterials = Params.NumMaterials;
 		PassParameters->NumViews = Params.NumViews;
@@ -549,6 +643,22 @@ void FWaterQuadTreeGPU::Traverse(FRDGBuilder& GraphBuilder, const FTraverseParam
 		FRDGBufferSRV* PackedNodesSRV = GraphBuilder.CreateSRV(FRDGBufferSRVDesc(PackedNodes, PF_R32_UINT));
 		FRDGBufferSRV* BucketCountsSRV = GraphBuilder.CreateSRV(FRDGBufferSRVDesc(BucketCounts, PF_R32_UINT));
 
+		FRDGBuffer* OcclusionQueryBoxes = nullptr;
+		FRDGBuffer* OcclusionQueryResults = nullptr;
+		FRDGBufferUAV* OcclusionQueryBoxesUAV = nullptr;
+		FRDGBufferUAV* OcclusionQueryResultsUAV = nullptr;
+		FRDGBufferSRV* OcclusionQueryBoxesSRV = nullptr;
+		FRDGBufferSRV* OcclusionQueryResultsSRV = nullptr;
+		if (bPixelPreciseOcclusionQueries)
+		{
+			OcclusionQueryBoxes = GraphBuilder.CreateBuffer(FRDGBufferDesc::CreateBufferDesc(sizeof(float) * 4, MaxNumDraws * 2), TEXT("WaterQuadTree.OcclusionQueryBoxes"));
+			OcclusionQueryResults = GraphBuilder.CreateBuffer(FRDGBufferDesc::CreateBufferDesc(sizeof(uint32), MaxNumDraws), TEXT("WaterQuadTree.OcclusionQueryResults"));
+			OcclusionQueryBoxesUAV = GraphBuilder.CreateUAV(FRDGBufferUAVDesc(OcclusionQueryBoxes, PF_A32B32G32R32F));
+			OcclusionQueryResultsUAV = GraphBuilder.CreateUAV(FRDGBufferUAVDesc(OcclusionQueryResults, PF_R32_UINT));
+			OcclusionQueryBoxesSRV = GraphBuilder.CreateSRV(FRDGBufferSRVDesc(OcclusionQueryBoxes, PF_A32B32G32R32F));
+			OcclusionQueryResultsSRV = GraphBuilder.CreateSRV(FRDGBufferSRVDesc(OcclusionQueryResults, PF_R32_UINT));
+		}
+
 		// Clear counts and packed nodes counter
 		// TODO: only clear the counter at index 0, not the entire buffer
 		FComputeShaderUtils::ClearUAV(GraphBuilder, ShaderMap, BucketCountsUAV, 0);
@@ -556,26 +666,8 @@ void FWaterQuadTreeGPU::Traverse(FRDGBuilder& GraphBuilder, const FTraverseParam
 
 		// Traverse quadtree
 		{
-			if (Params.DebugShowTile != 0)
-			{
-				ShaderPrint::SetEnabled(true);
-			}
-
-			const bool bEnableShaderPrint = Params.DebugShowTile != 0
-				&& PerViewInfo.ViewInfo
-				&& ShaderPrint::IsEnabled()
-				&& ShaderPrint::IsEnabled(PerViewInfo.ViewInfo->ShaderPrintData)
-				&& ShaderPrint::IsValid(PerViewInfo.ViewInfo->ShaderPrintData)
-				&& ShaderPrint::IsSupported(GetFeatureLevelShaderPlatform(GMaxRHIFeatureLevel));
-
-			if (bEnableShaderPrint)
-			{
-				// We'll be potentially drawing a lot of bounding boxes, each made up of 12 line segments, so reserve some space
-				ShaderPrint::RequestSpaceForLines(12 * MaxNumDraws);
-			}
-
 			FWaterQuadTreeTraverseCS::FPermutationDomain PermutationDomain;
-			PermutationDomain.Set<FWaterQuadTreeTraverseCS::FDebugShaderPrint>(bEnableShaderPrint);
+			PermutationDomain.Set<FWaterQuadTreeTraverseCS::FPreciseOcclusionQueries>(bPixelPreciseOcclusionQueries);
 			TShaderMapRef<FWaterQuadTreeTraverseCS> ComputeShader(ShaderMap, PermutationDomain);
 
 			const uint32 NumMipLevels = QuadTreeTextureRDG->Desc.NumMips;
@@ -587,11 +679,10 @@ void FWaterQuadTreeGPU::Traverse(FRDGBuilder& GraphBuilder, const FTraverseParam
 			}
 
 			FWaterQuadTreeTraverseCS::FParameters* PassParameters = GraphBuilder.AllocParameters<FWaterQuadTreeTraverseCS::FParameters>();
-			if (bEnableShaderPrint)
-			{
-				ShaderPrint::SetParameters(GraphBuilder, PerViewInfo.ViewInfo->ShaderPrintData, PassParameters->ShaderPrintUniformBuffer);
-			}
 			PassParameters->PackedNodes = PackedNodesUAV;
+			PassParameters->OcclusionQueryBoxes = OcclusionQueryBoxesUAV;
+			PassParameters->OcclusionVisibility = OcclusionQueryResultsUAV;
+			PassParameters->OcclusionQueryIndirectArgs = OcclusionQueryIndirectArgsBufferUAV;
 			PassParameters->QuadTreeTexture = QuadTreeTextureRDG;
 			PassParameters->WaterZBoundsTexture = WaterZBoundsTextureRDG;
 			PassParameters->WaterBodyRenderData = WaterBodyRenderDataBufferSRV;
@@ -608,27 +699,77 @@ void FWaterQuadTreeGPU::Traverse(FRDGBuilder& GraphBuilder, const FTraverseParam
 			PassParameters->QuadTreeResolutionY = QuadTreeResolution.Y;
 			PassParameters->NumDensities = Params.NumDensities;
 			PassParameters->NumMaterials = Params.NumMaterials;
+			PassParameters->ViewIndex = ViewIndex;
 			PassParameters->LeafSize = Params.LeafSize;
 			PassParameters->LODScale = Params.LODScale;
 			PassParameters->CaptureDepthRange = CaptureDepthRange;
 			PassParameters->ForceCollapseDensityLevel = -1; // Params.ForceCollapseDensityLevel; // TODO: Properly implement support for this parameter
 			PassParameters->NumLODs = NumMipLevels;
 			PassParameters->NumDispatchedThreads = NumTexelsTotal;
-			PassParameters->DebugShowTile = Params.DebugShowTile;
-			PassParameters->bHZBOcclusionCullingEnabled = GWaterQuadTreeOcclusionCulling && (PerViewInfo.ViewInfo && PerViewInfo.ViewInfo->HZB);
+			PassParameters->bHZBOcclusionCullingEnabled = bHZBOcclusionQueries && (PerViewInfo.ViewInfo && PerViewInfo.ViewInfo->HZB);
 
 			FComputeShaderUtils::AddPass(GraphBuilder, RDG_EVENT_NAME("WaterQuadTreeTraverse(View: %i)", ViewIndex), ComputeShader, PassParameters, FComputeShaderUtils::GetGroupCount(NumTexelsTotal, 64));
 		}
 
+		// Raster occlusion queries
+		if (bPixelPreciseOcclusionQueries)
+		{
+			TShaderMapRef<FWaterQuadTreeOcclusionQueryVS> VertexShader(ShaderMap);
+			TShaderMapRef<FWaterQuadTreeOcclusionQueryPS> PixelShader(ShaderMap);
+
+			FRDGTextureRef DepthTexture = PerViewInfo.ViewInfo->GetSceneTextures().Depth.Target;
+
+			FWaterQuadTreeOcclusionQueryParameters* PassParameters = GraphBuilder.AllocParameters<FWaterQuadTreeOcclusionQueryParameters>();
+			PassParameters->IndirectDrawArgsBuffer = OcclusionQueryIndirectArgsBuffer;
+			PassParameters->VS.OcclusionQueryBoxes = OcclusionQueryBoxesSRV;
+			PassParameters->VS.ViewProjection = PerViewInfo.TranslatedWorldToClip;
+			PassParameters->PS.Visibility = OcclusionQueryResultsUAV;
+			PassParameters->PS.RenderTargets.DepthStencil = FDepthStencilBinding(DepthTexture, ERenderTargetLoadAction::ELoad, ERenderTargetLoadAction::ENoAction, FExclusiveDepthStencil::DepthRead_StencilNop);
+
+			const FIntVector4 ViewRect = FIntVector4(PerViewInfo.ViewInfo->ViewRect.Min.X, PerViewInfo.ViewInfo->ViewRect.Min.Y, PerViewInfo.ViewInfo->ViewRect.Max.X, PerViewInfo.ViewInfo->ViewRect.Max.Y);
+
+			GraphBuilder.AddPass(
+				RDG_EVENT_NAME("WaterQuadTreeOcclusionCulling(View: %i)", ViewIndex),
+				PassParameters,
+				ERDGPassFlags::Raster,
+				[PassParameters, ViewRect, ViewIndex, VertexShader, PixelShader](FRHICommandList& RHICmdList)
+				{
+					RHICmdList.SetViewport(ViewRect.X, ViewRect.Y, 0.0f, ViewRect.Z, ViewRect.W, 1.0f);
+
+					FGraphicsPipelineStateInitializer GraphicsPSOInit;
+					RHICmdList.ApplyCachedRenderTargets(GraphicsPSOInit);
+					GraphicsPSOInit.BoundShaderState.VertexDeclarationRHI = GetVertexDeclarationFVector4();
+					GraphicsPSOInit.BoundShaderState.VertexShaderRHI = VertexShader.GetVertexShader();
+					GraphicsPSOInit.BoundShaderState.PixelShaderRHI = PixelShader.GetPixelShader();
+					GraphicsPSOInit.BlendState = TStaticBlendState<>::GetRHI();
+					GraphicsPSOInit.RasterizerState = TStaticRasterizerState<FM_Solid, CM_CW>::GetRHI();
+					GraphicsPSOInit.DepthStencilState = TStaticDepthStencilState<false, CF_DepthNearOrEqual>::GetRHI();
+					GraphicsPSOInit.PrimitiveType = PT_TriangleList;
+
+					SetGraphicsPipelineState(RHICmdList, GraphicsPSOInit, 0);
+
+					SetShaderParameters(RHICmdList, VertexShader, VertexShader.GetVertexShader(), PassParameters->VS);
+					SetShaderParameters(RHICmdList, PixelShader, PixelShader.GetPixelShader(), PassParameters->PS);
+
+					PassParameters->IndirectDrawArgsBuffer->MarkResourceAsUsed();
+
+					RHICmdList.SetStreamSource(0, GetUnitCubeVertexBuffer(), 0);
+					RHICmdList.DrawIndexedPrimitiveIndirect(GetUnitCubeIndexBuffer(), PassParameters->IndirectDrawArgsBuffer->GetRHI(), ViewIndex * sizeof(FRHIDrawIndexedIndirectParameters));
+				});
+		}
+
 		// Compute bucket counts
 		{
-			TShaderMapRef<FWaterQuadTreeBucketCountsCS> ComputeShader(ShaderMap);
+			FWaterQuadTreeBucketCountsCS::FPermutationDomain PermutationDomain;
+			PermutationDomain.Set<FWaterQuadTreeBucketCountsCS::FPreciseOcclusionQueries>(bPixelPreciseOcclusionQueries);
+			TShaderMapRef<FWaterQuadTreeBucketCountsCS> ComputeShader(ShaderMap, PermutationDomain);
 
 			FWaterQuadTreeBucketCountsCS::FParameters* PassParameters = GraphBuilder.AllocParameters<FWaterQuadTreeBucketCountsCS::FParameters>();
 			PassParameters->BucketCounts = BucketCountsUAV;
 			PassParameters->QuadTreeTexture = QuadTreeTextureRDG;
 			PassParameters->WaterBodyRenderData = WaterBodyRenderDataBufferSRV;
 			PassParameters->PackedNodes = PackedNodesSRV;
+			PassParameters->OcclusionResults = OcclusionQueryResultsSRV;
 			PassParameters->NumDensities = Params.NumDensities;
 			PassParameters->NumMaterials = Params.NumMaterials;
 			PassParameters->NumDispatchedThreads = FMath::Min(FMath::DivideAndRoundUp(MaxNumDraws, 64u), 65535u) * 64u; // The maximum number of dispatched groups is 64k
@@ -653,7 +794,9 @@ void FWaterQuadTreeGPU::Traverse(FRDGBuilder& GraphBuilder, const FTraverseParam
 
 		// Generate instance data and fill indirect args
 		{
-			TShaderMapRef<FWaterQuadTreeGenerateInstanceDataCS> ComputeShader(ShaderMap);
+			FWaterQuadTreeGenerateInstanceDataCS::FPermutationDomain PermutationDomain;
+			PermutationDomain.Set<FWaterQuadTreeGenerateInstanceDataCS::FPreciseOcclusionQueries>(bPixelPreciseOcclusionQueries);
+			TShaderMapRef<FWaterQuadTreeGenerateInstanceDataCS> ComputeShader(ShaderMap, PermutationDomain);
 
 			FWaterQuadTreeGenerateInstanceDataCS::FParameters* PassParameters = GraphBuilder.AllocParameters<FWaterQuadTreeGenerateInstanceDataCS::FParameters>();
 			PassParameters->IndirectArgs = IndirectArgsBufferUAV;
@@ -665,6 +808,7 @@ void FWaterQuadTreeGPU::Traverse(FRDGBuilder& GraphBuilder, const FTraverseParam
 			PassParameters->WaterBodyRenderData = WaterBodyRenderDataBufferSRV;
 			PassParameters->PackedNodes = PackedNodesSRV;
 			PassParameters->InstanceDataOffsets = InstanceDataOffsetsBufferSRV;
+			PassParameters->OcclusionResults = OcclusionQueryResultsSRV;
 			PassParameters->QuadTreePosition = PerViewInfo.QuadTreePositionTranslatedWorldSpace;
 			PassParameters->ObserverPosition = PerViewInfo.ObserverPositionTranslatedWorldSpace;
 			PassParameters->QuadTreeResolutionX = QuadTreeResolution.X;
@@ -682,5 +826,41 @@ void FWaterQuadTreeGPU::Traverse(FRDGBuilder& GraphBuilder, const FTraverseParam
 
 			FComputeShaderUtils::AddPass(GraphBuilder, RDG_EVENT_NAME("WaterQuadTreeGenerateDraws(View: %i)", ViewIndex), ComputeShader, PassParameters, FComputeShaderUtils::GetGroupCount(PassParameters->NumDispatchedThreads, 64));
 		}
+
+		// Debug show tiles
+		if (Params.DebugShowTile != 0)
+		{
+			ShaderPrint::SetEnabled(true);
+
+			const bool bEnableShaderPrint = PerViewInfo.ViewInfo
+				&& ShaderPrint::IsEnabled()
+				&& ShaderPrint::IsEnabled(PerViewInfo.ViewInfo->ShaderPrintData)
+				&& ShaderPrint::IsValid(PerViewInfo.ViewInfo->ShaderPrintData)
+				&& ShaderPrint::IsSupported(GetFeatureLevelShaderPlatform(GMaxRHIFeatureLevel));
+
+			if (bEnableShaderPrint)
+			{
+				// We'll be potentially drawing a lot of bounding boxes, each made up of 12 line segments, so reserve some space
+				ShaderPrint::RequestSpaceForLines(12 * MaxNumDraws);
+
+				FWaterQuadTreeDebugCS::FPermutationDomain PermutationDomain;
+				PermutationDomain.Set<FWaterQuadTreeDebugCS::FPreciseOcclusionQueries>(bPixelPreciseOcclusionQueries);
+				TShaderMapRef<FWaterQuadTreeDebugCS> ComputeShader(ShaderMap, PermutationDomain);
+
+				FWaterQuadTreeDebugCS::FParameters* PassParameters = GraphBuilder.AllocParameters<FWaterQuadTreeDebugCS::FParameters>();
+				ShaderPrint::SetParameters(GraphBuilder, PerViewInfo.ViewInfo->ShaderPrintData, PassParameters->ShaderPrintUniformBuffer);
+				PassParameters->QuadTreeTexture = QuadTreeTextureRDG;
+				PassParameters->WaterZBoundsTexture = WaterZBoundsTextureRDG;
+				PassParameters->WaterBodyRenderData = WaterBodyRenderDataBufferSRV;
+				PassParameters->PackedNodes = PackedNodesSRV;
+				PassParameters->OcclusionResults = OcclusionQueryResultsSRV;
+				PassParameters->QuadTreePosition = PerViewInfo.QuadTreePositionTranslatedWorldSpace;
+				PassParameters->NumDispatchedThreads = FMath::Min(FMath::DivideAndRoundUp(MaxNumDraws, 64u), 65535u) * 64u; // The maximum number of dispatched groups is 64k
+				PassParameters->LeafSize = Params.LeafSize;
+				PassParameters->CaptureDepthRange = CaptureDepthRange;
+
+				FComputeShaderUtils::AddPass(GraphBuilder, RDG_EVENT_NAME("WaterQuadTreeDebug(View: %i)", ViewIndex), ComputeShader, PassParameters, FComputeShaderUtils::GetGroupCount(PassParameters->NumDispatchedThreads, 64));
+			}
+		}
 	}
 }
diff --git a/Engine/Plugins/Experimental/Water/Source/Runtime/Private/WaterQuadTreeGPU.h b/Engine/Plugins/Experimental/Water/Source/Runtime/Private/WaterQuadTreeGPU.h
--- a/Engine/Plugins/Experimental/Water/Source/Runtime/Private/WaterQuadTreeGPU.h
+++ b/Engine/Plugins/Experimental/Water/Source/Runtime/Private/WaterQuadTreeGPU.h
@@ -90,6 +90,14 @@ public:
 	bool IsInitialized() const { return bInitialized; }
 
 private:
+	enum class EOcclusionQueryMode
+	{
+		Disabled = 0,
+		HZB = 1,
+		PixelPrecise = 2,
+		HZBAndPixelPrecise = 3
+	};
+
 	/** Mipped texture representing the water quadtree on the GPU */
 	TRefCountPtr<FRHITexture> QuadTreeTexture;
 
