Water: Supports scale on the water zone.

diff --git a/Engine/Plugins/Experimental/Water/Source/Runtime/Private/WaterZoneActor.cpp b/Engine/Plugins/Experimental/Water/Source/Runtime/Private/WaterZoneActor.cpp
--- a/Engine/Plugins/Experimental/Water/Source/Runtime/Private/WaterZoneActor.cpp
+++ b/Engine/Plugins/Experimental/Water/Source/Runtime/Private/WaterZoneActor.cpp
@@ -262,6 +262,11 @@ void AWaterZone::RemoveWaterBodyComponent(UWaterBodyComponent* WaterBodyComponen
 	}
 }
 
+FVector2D AWaterZone::GetZoneExtent() const
+{
+	return ZoneExtent * FVector2D(GetActorScale());
+}
+
 void AWaterZone::Update()
 {
 	if (bNeedsWaterInfoRebuild || (ForceUpdateWaterInfoNextFrames != 0))
@@ -346,6 +351,10 @@ void AWaterZone::PostEditChangeProperty(FPropertyChangedEvent& PropertyChangedEv
 	{
 		MarkForRebuild(EWaterZoneRebuildFlags::All);
 	}
+	else if (PropertyName == USceneComponent::GetRelativeScale3DPropertyName())
+	{
+		MarkForRebuild(EWaterZoneRebuildFlags::All);
+	}
 }
 
 void AWaterZone::OnActorSelectionChanged(const TArray<UObject*>& NewSelection, bool bForceRefresh)
@@ -615,7 +624,7 @@ FVector AWaterZone::GetDynamicWaterInfoExtent() const
 {
 	if (IsLocalOnlyTessellationEnabled())
 	{
-		return LocalTessellationExtent;
+		return LocalTessellationExtent * GetActorScale();
 	}
 	else
 	{
diff --git a/Engine/Plugins/Experimental/Water/Source/Runtime/Public/WaterZoneActor.h b/Engine/Plugins/Experimental/Water/Source/Runtime/Public/WaterZoneActor.h
--- a/Engine/Plugins/Experimental/Water/Source/Runtime/Public/WaterZoneActor.h
+++ b/Engine/Plugins/Experimental/Water/Source/Runtime/Public/WaterZoneActor.h
@@ -44,7 +44,7 @@ public:
 	void AddWaterBodyComponent(UWaterBodyComponent* WaterBodyComponent);
 	void RemoveWaterBodyComponent(UWaterBodyComponent* WaterBodyComponent);
 
-	FVector2D GetZoneExtent() const { return ZoneExtent; }
+	FVector2D GetZoneExtent() const;
 	void SetZoneExtent(FVector2D NewExtents);
 
 	FBox2D GetZoneBounds2D() const;
@@ -162,7 +162,7 @@ private:
 	UPROPERTY(VisibleAnywhere, Category = Water, BlueprintReadOnly, meta = (AllowPrivateAccess = "true"))
 	TObjectPtr<UWaterMeshComponent> WaterMesh;
 
-	/** Width of the zone bounding box */
+	/** The maximum size in local space of the water zone. */
 	UPROPERTY(Category = Water, EditAnywhere, BlueprintReadWrite, meta = (AllowPrivateAccess = "true"))
 	FVector2D ZoneExtent;
 
@@ -190,7 +190,7 @@ private:
 	UPROPERTY(Category = LocalTessellation, EditAnywhere)
 	bool bEnableLocalOnlyTessellation = false;
 
-	/** The diameters in world space units for the region within which local dynamic tessellation occurs. A smaller value increases the effective pixel density of the water info texture. */
+	/** The diameters in local space units for the region within which dynamic tessellation occurs. A smaller value increases the effective pixel density of the water info texture. */
 	UPROPERTY(Category = LocalTessellation, EditAnywhere, meta = (EditCondition = "bEnableLocalOnlyTessellation"))
 	FVector LocalTessellationExtent;
 
