Moved FixedZHeight, FixedWaterDepth and FixedVelocity water uniforms into the view water buffer. Replaced FWaterBodyHeader with FWaterBodyData, essentially removing an indirection when looking up the number of gerstner waves of a water body.

#rnx

#virtualized

diff --git a/Engine/Plugins/Experimental/Water/Shaders/Private/GerstnerWaveFunctions.ush b/Engine/Plugins/Experimental/Water/Shaders/Private/GerstnerWaveFunctions.ush
--- a/Engine/Plugins/Experimental/Water/Shaders/Private/GerstnerWaveFunctions.ush
+++ b/Engine/Plugins/Experimental/Water/Shaders/Private/GerstnerWaveFunctions.ush
@@ -286,8 +286,8 @@ float3 FinalizeNormalNew(float3 inNormal)
 
 FWaveParams GetWaveDataByWaveIndexNew(int InWaterBodyIndex, int InWaveIndex)
 {
-	const FWaterBodyHeader WaterBodyHeader = DecodeWaterBodyHeader(View.WaterIndirection[InWaterBodyIndex]);
-	return GetWaveDataNew(InWaveIndex, WaterBodyHeader);
+	const FWaterBodyData WaterBodyData = GetWaterBodyData(InWaterBodyIndex);
+	return GetWaveDataNew(InWaveIndex, WaterBodyData);
 }
 
 float3 PackNormalAndWPONew(WaveOutput inWave)
@@ -300,11 +300,11 @@ float3 PackNormalAndWPONew(WaveOutput inWave)
 	return packedoutput;
 }
 
-WaveOutput GetSingleGerstnerWaveNew(int InWaveIndex, FWaterBodyHeader InWaterBodyHeader, float InTime, FLWCVector2 InWorldPos)
+WaveOutput GetSingleGerstnerWaveNew(int InWaveIndex, FWaterBodyData InWaterBodyData, float InTime, FLWCVector2 InWorldPos)
 {
 	WaveOutput OutWave = (WaveOutput)0;
 		
-	FWaveParams CurrentWave = GetWaveDataNew(InWaveIndex, InWaterBodyHeader);
+	FWaveParams CurrentWave = GetWaveDataNew(InWaveIndex, InWaterBodyData);
 		
 	float Dispersion = 2 * PI / CurrentWave.Wavelength;
 	float2 WaveVector = CurrentWave.Direction * Dispersion;
@@ -359,13 +359,13 @@ WaveOutput GetAllGerstnerWavesNew(int InWaterBodyIndex, float InTime, FLWCVector
 {		
 	WaveOutput OutWaves = (WaveOutput)0;
 	WaveOutput CurrentWave = (WaveOutput)0;
-		
-	const FWaterBodyHeader WaterBodyHeader = DecodeWaterBodyHeader(View.WaterIndirection[InWaterBodyIndex]);
-	const int NumWaves = GetWaterBodyData(InWaterBodyIndex).NumWaves;
+	
+	const FWaterBodyData WaterBodyData = GetWaterBodyData(InWaterBodyIndex);
+	const int NumWaves = WaterBodyData.NumWaves;
 
 	for (int i = 0; i < NumWaves; i++)
 	{
-		CurrentWave = GetSingleGerstnerWaveNew(i, WaterBodyHeader, InTime, InWorldPos);
+		CurrentWave = GetSingleGerstnerWaveNew(i, WaterBodyData, InTime, InWorldPos);
 		OutWaves = AddWavesNew(OutWaves, CurrentWave);
 	}
 		
@@ -379,15 +379,15 @@ WaveOutput GetRangeGerstnerWavesNew(int InMinWaveIndex, int InMaxWaveIndex, int
 	WaveOutput OutWaves = (WaveOutput)0;
 	WaveOutput CurrentWave = (WaveOutput)0;
 		
-	const FWaterBodyHeader WaterBodyHeader = DecodeWaterBodyHeader(View.WaterIndirection[InWaterBodyIndex]);
-	const int NumWaves = GetWaterBodyData(InWaterBodyIndex).NumWaves;
+	const FWaterBodyData WaterBodyData = GetWaterBodyData(InWaterBodyIndex);
+	const int NumWaves = WaterBodyData.NumWaves;
 
 	const int MaxWaveRange = min(InMaxWaveIndex + 1, NumWaves);
 	const int MinWaveRange = max(0, InMinWaveIndex);
 
 	for (int i = MinWaveRange; i < MaxWaveRange; i++)
 	{
-		CurrentWave = GetSingleGerstnerWaveNew(i, WaterBodyHeader, InTime, InWorldPos);
+		CurrentWave = GetSingleGerstnerWaveNew(i, WaterBodyData, InTime, InWorldPos);
 		OutWaves = AddWavesNew(OutWaves, CurrentWave);
 	}
 		
diff --git a/Engine/Plugins/Experimental/Water/Shaders/Private/WaterDataFunctions.ush b/Engine/Plugins/Experimental/Water/Shaders/Private/WaterDataFunctions.ush
--- a/Engine/Plugins/Experimental/Water/Shaders/Private/WaterDataFunctions.ush
+++ b/Engine/Plugins/Experimental/Water/Shaders/Private/WaterDataFunctions.ush
@@ -5,10 +5,10 @@
 
 #define PER_WAVE_DATA_SIZE 2
 #define PER_WATERZONE_DATA_SIZE 2
-#define PER_WATERBODY_DATA_SIZE 1
+#define PER_WATERBODY_DATA_SIZE 2
 
-/** Struct containing the decoded water body header data  */	
-struct FWaterBodyHeader
+/** Struct containing the decoded water body data  */	
+struct FWaterBodyData
 {
 	// Index of the water zone which owns this body. Points to the position of the water zone in the zone sub-buffer.
 	int WaterZoneIndex;
@@ -16,8 +16,15 @@ struct FWaterBodyHeader
 	// Index into the wave data sub-buffer.
 	int WaveDataIndex;
 
-	// Index into the water body data sub-buffer
-	int WaterBodyDataIndex;
+	// Number of waves to be read from the wave data buffer
+	int NumWaves;
+
+	// Target Wave Mask Depth
+	float TargetWaveMaskDepth;
+	
+	float3 FixedVelocity;
+	float FixedZHeight;
+	float FixedWaterDepth;
 };
 
 /** Struct containing the decoded per-water zone data. */
@@ -29,16 +36,6 @@ struct FWaterZoneData
 	float GroundZMin;
 };
 
-/** Struct containing the decoded per-water body data. */
-struct FWaterBodyData
-{
-	// Number of waves to be read from the wave data buffer
-	int NumWaves;
-	
-	// Target Wave Mask Depth
-	float TargetWaveMaskDepth;
-};
-
 struct FWaveParams
 {
 	float2 Direction;
@@ -57,20 +54,28 @@ int GetWaveDataOffset()
 	return (int)View.WaterData[0].y;
 }
 
-int GetWaterBodyDataOffset()
+/** Function to decode the raw data from the water indirection buffer */
+FWaterBodyData DecodeWaterBodyData(float4 InDataToDecode0, float4 InDataToDecode1)
 {
-	return (int)View.WaterData[0].z;
+	FWaterBodyData OutWaterBodyData = (FWaterBodyData)0;
+	OutWaterBodyData.WaterZoneIndex = InDataToDecode0.x;
+	OutWaterBodyData.WaveDataIndex = InDataToDecode0.y;
+	OutWaterBodyData.NumWaves = InDataToDecode0.z;
+	OutWaterBodyData.TargetWaveMaskDepth = InDataToDecode0.w;
+	OutWaterBodyData.FixedVelocity = float3(f16tof32(asuint(InDataToDecode1.x)), f16tof32(asuint(InDataToDecode1.x) >> 16u), InDataToDecode1.y);
+	OutWaterBodyData.FixedZHeight = InDataToDecode1.z;
+	OutWaterBodyData.FixedWaterDepth = InDataToDecode1.w;
+	
+	return OutWaterBodyData;
 }
 
-/** Function to decode the raw data from the water indirection buffer */
-FWaterBodyHeader DecodeWaterBodyHeader(float4 InDataToDecode)
+FWaterBodyData GetWaterBodyData(int InWaterBodyIndex)
 {
-	FWaterBodyHeader OutWaterBodyHeader = (FWaterBodyHeader)0;
-	OutWaterBodyHeader.WaterZoneIndex = InDataToDecode.x;
-	OutWaterBodyHeader.WaveDataIndex = InDataToDecode.y;
-	OutWaterBodyHeader.WaterBodyDataIndex = InDataToDecode.z;
-	
-	return OutWaterBodyHeader;
+	const int WaterBodyOffset = InWaterBodyIndex * PER_WATERBODY_DATA_SIZE;
+	const float4 Data0 = View.WaterIndirection[WaterBodyOffset];
+	const float4 Data1 = View.WaterIndirection[WaterBodyOffset + 1];
+
+	return DecodeWaterBodyData(Data0, Data1);
 }
 
 /** Function to decode the data for an individual gerstner wave from the water data buffer */
@@ -85,36 +90,16 @@ FWaveParams DecodeWaveParams(float4 InDataToDecode0, float4 InDataToDecode1)
 	return OutWaveParams;
 }
 
-FWaveParams GetWaveDataNew(int InWaveIndex, FWaterBodyHeader InWaterBodyHeader)
+FWaveParams GetWaveDataNew(int InWaveIndex, FWaterBodyData InWaterBodyData)
 {
 	const int WaveDataOffset = GetWaveDataOffset();
-	const int AbsoluteWaveDataIndex = WaveDataOffset + (InWaterBodyHeader.WaveDataIndex + InWaveIndex) * PER_WAVE_DATA_SIZE;
-
-	float4 Data0 = View.WaterData[AbsoluteWaveDataIndex];
-	float4 Data1 = View.WaterData[AbsoluteWaveDataIndex + 1];
+	const int AbsoluteWaveDataIndex = WaveDataOffset + (InWaterBodyData.WaveDataIndex + InWaveIndex) * PER_WAVE_DATA_SIZE;
+	const float4 Data0 = View.WaterData[AbsoluteWaveDataIndex];
+	const float4 Data1 = View.WaterData[AbsoluteWaveDataIndex + 1];
 		
 	return DecodeWaveParams(Data0, Data1);
 }
 
-FWaterBodyData DecodeWaterBodyData(float4 Data)
-{
-	FWaterBodyData OutWaterBodyData = (FWaterBodyData)0;
-	OutWaterBodyData.NumWaves = Data.x;
-	OutWaterBodyData.TargetWaveMaskDepth = Data.y;
-
-	return OutWaterBodyData;
-}
-
-FWaterBodyData GetWaterBodyData(int InWaterBodyIndex)
-{
-	const int WaterBodyDataBufferOffset = GetWaterBodyDataOffset();
-	const FWaterBodyHeader WaterBodyHeader = DecodeWaterBodyHeader(View.WaterIndirection[InWaterBodyIndex]);
-
-	float4 Data = View.WaterData[WaterBodyDataBufferOffset + WaterBodyHeader.WaterBodyDataIndex];
-
-	return DecodeWaterBodyData(Data);
-}
-
 FWaterZoneData DecodeWaterZoneData(float4 Data0, float4 Data1)
 {
 	FWaterZoneData OutWaterZoneData = (FWaterZoneData)0;
@@ -131,15 +116,15 @@ FWaterZoneData GetWaterZoneData(int InWaterZoneIndex)
 	const int WaterZoneDataBufferOffset = GetWaterZoneDataOffset();
 	const int WaterZoneOffset = WaterZoneDataBufferOffset + (InWaterZoneIndex * PER_WATERZONE_DATA_SIZE);
 
-	float4 Data0 = View.WaterData[WaterZoneOffset];
-	float4 Data1 = View.WaterData[WaterZoneOffset + 1];
+	const float4 Data0 = View.WaterData[WaterZoneOffset];
+	const float4 Data1 = View.WaterData[WaterZoneOffset + 1];
 
 	return DecodeWaterZoneData(Data0, Data1);
 }
 
 FWaterZoneData GetWaterZoneDataForWaterBody(int InWaterBodyIndex)
 {
-	const int WaterZoneIndex = DecodeWaterBodyHeader(View.WaterIndirection[InWaterBodyIndex]).WaterZoneIndex;
+	const int WaterZoneIndex = GetWaterBodyData(InWaterBodyIndex).WaterZoneIndex;
 
 	return GetWaterZoneData(WaterZoneIndex);
 }
diff --git a/Engine/Plugins/Experimental/Water/Source/Runtime/Private/WaterBodyComponent.cpp b/Engine/Plugins/Experimental/Water/Source/Runtime/Private/WaterBodyComponent.cpp
--- a/Engine/Plugins/Experimental/Water/Source/Runtime/Private/WaterBodyComponent.cpp
+++ b/Engine/Plugins/Experimental/Water/Source/Runtime/Private/WaterBodyComponent.cpp
@@ -71,14 +71,7 @@ const FName UWaterBodyComponent::WaterZoneIndexParamName(TEXT("WaterZoneIndex"))
 const FName UWaterBodyComponent::WaterBodyZOffsetParamName(TEXT("WaterBodyZOffset"));
 const FName UWaterBodyComponent::WaterVelocityAndHeightName(TEXT("WaterVelocityAndHeight"));
 const FName UWaterBodyComponent::GlobalOceanHeightName(TEXT("GlobalOceanHeight"));
-const FName UWaterBodyComponent::FixedZHeightName(TEXT("FixedZHeight"));
-const FName UWaterBodyComponent::FixedVelocityName(TEXT("FixedVelocity"));
-const FName UWaterBodyComponent::FixedWaterDepthName(TEXT("FixedWaterDepth"));
-const FName UWaterBodyComponent::WaterAreaParamName(TEXT("WaterArea"));
 const FName UWaterBodyComponent::MaxFlowVelocityParamName(TEXT("MaxFlowVelocity"));
-const FName UWaterBodyComponent::WaterZMinParamName(TEXT("WaterZMin"));
-const FName UWaterBodyComponent::WaterZMaxParamName(TEXT("WaterZMax"));
-const FName UWaterBodyComponent::GroundZMinParamName(TEXT("GroundZMin"));
 
 UWaterBodyComponent::UWaterBodyComponent(const FObjectInitializer& ObjectInitializer)
 	: Super(ObjectInitializer)
@@ -1763,10 +1756,6 @@ bool UWaterBodyComponent::SetDynamicParametersOnMID(UMaterialInstanceDynamic* In
 	const float GlobalOceanHeight = WaterSubsystem->GetOceanTotalHeight();
 	InMID->SetScalarParameterValue(WaterBodyIndexParamName, WaterBodyIndex);
 	InMID->SetScalarParameterValue(GlobalOceanHeightName, GlobalOceanHeight);
-	InMID->SetScalarParameterValue(FixedZHeightName, GetConstantSurfaceZ());
-	InMID->SetScalarParameterValue(FixedWaterDepthName, GetConstantDepth());
-
-	InMID->SetVectorParameterValue(FixedVelocityName, GetConstantVelocity());
 
 	if (const AWaterZone* WaterZone = GetWaterZone())
 	{
diff --git a/Engine/Plugins/Experimental/Water/Source/Runtime/Private/WaterViewExtension.cpp b/Engine/Plugins/Experimental/Water/Source/Runtime/Private/WaterViewExtension.cpp
--- a/Engine/Plugins/Experimental/Water/Source/Runtime/Private/WaterViewExtension.cpp
+++ b/Engine/Plugins/Experimental/Water/Source/Runtime/Private/WaterViewExtension.cpp
@@ -81,19 +81,22 @@ void FWaterViewExtension::UpdateGPUBuffers()
 		FWaterBodyManager* WaterBodyManager = UWaterSubsystem::GetWaterBodyManager(WorldPtr);
 		check(WaterBodyManager);
 
-		// Shrink the water manager storage to avoid over-preallocating in the WaterIndirectionBuffer.
+		// Shrink the water manager storage to avoid over-preallocating in the WaterBodyDataBuffer.
 		WaterBodyManager->Shrink();
 
 
-		struct FWaterIndirection
+		struct FWaterBodyData
 		{
 			float WaterZoneIndex;
 			float WaveDataIndex;
-			float WaterBodyDataIndex;
-
-			float _Padding; // Unused 
+			float NumWaves;
+			float TargetWaveMaskDepth;
+			float FixedVelocityXY; // Packed as two 16 bit floats. X is in the lower 16 bits.
+			float FixedVelocityZ;
+			float FixedZHeight;
+			float FixedWaterDepth;
 		};
-		static_assert(sizeof(FWaterIndirection) == 1 * sizeof(FVector4f));
+		static_assert(sizeof(FWaterBodyData) == 2 * sizeof(FVector4f));
 
 		struct FWaterZoneData
 		{
@@ -123,23 +126,14 @@ void FWaterViewExtension::UpdateGPUBuffers()
 		};
 		static_assert(sizeof(FGerstnerWaveData) == 2 * sizeof(FVector4f));
 
-		struct FWaterBodyData
-		{
-			float NumWaves;
-			float TargetWaveMaskDepth;
-			
-			float _Padding[2]; // Unused;
-		};
-		static_assert(sizeof(FWaterBodyData) == sizeof(FVector4f));
-
-		// Water Indirection Buffer layout:
+		// Water Body Data Buffer layout:
 		// -------------------------------------------------------------------------------
-		// || WaterZoneIndex | WaveDataIndex | WaterBodyDataIndex | *Unused* ||   ...   ||
+		// || WaterZoneIndex | WaveDataIndex | NumWaves | (Other members) ||   ...   ||
 		// -------------------------------------------------------------------------------
 		//
-		// Water Data Buffer layout:
+		// Water Aux Data Buffer layout:
 		// -----------------------------------------------------------------------------
-		// ||| WaterZone Data | ..  || GerstnerWaveData | ... || WaterBodyData | ... |||
+		// ||| WaterZone Data | ... || GerstnerWaveData | ... |||
 		// -----------------------------------------------------------------------------
 		//
 
@@ -162,28 +156,29 @@ void FWaterViewExtension::UpdateGPUBuffers()
 		}
 
 
-		TArray<FWaterIndirection> WaterIndirection;
-		TArray<FGerstnerWaveData> WaveData;
 		TArray<FWaterBodyData> WaterBodyData;
+		TArray<FGerstnerWaveData> WaveData;
 		{
 			const int32 NumWaterBodies =  WaterBodyManager->NumWaterBodies();
 			// Pre-set up to the max water body index. Some entries may be empty and NumWaterBodies != MaxIndex
-			WaterIndirection.SetNumZeroed(WaterBodyManager->MaxWaterBodyIndex());
-			WaterBodyData.Reserve(NumWaterBodies);
+			WaterBodyData.SetNumZeroed(WaterBodyManager->MaxWaterBodyIndex());
 
 			TMap<const UGerstnerWaterWaves*, int32> GerstnerWavesIndices;
 
-			WaterBodyManager->ForEachWaterBodyComponent([&WaterIndirection, &WaveData, &WaterBodyData, &GerstnerWavesIndices](UWaterBodyComponent* WaterBodyComponent)
+			WaterBodyManager->ForEachWaterBodyComponent([&WaterBodyData, &WaveData, &GerstnerWavesIndices](UWaterBodyComponent* WaterBodyComponent)
 			{
 				const int32 WaterZoneIndex = WaterBodyComponent->GetWaterZone() ? WaterBodyComponent->GetWaterZone()->GetWaterZoneIndex() : -1;
 
-				check(WaterBodyComponent->GetWaterBodyIndex() < WaterIndirection.Num());
-				FWaterIndirection& WaterIndirectionEntry = WaterIndirection[WaterBodyComponent->GetWaterBodyIndex()];
-				WaterIndirectionEntry.WaterZoneIndex = WaterZoneIndex;
-				WaterIndirectionEntry.WaterBodyDataIndex = WaterBodyData.Num();
+				const FVector FixedVelocity = WaterBodyComponent->GetConstantVelocity();
 
-				FWaterBodyData& WaterBodyDataEntry = WaterBodyData.AddZeroed_GetRef();
+				check(WaterBodyComponent->GetWaterBodyIndex() < WaterBodyData.Num());
+				FWaterBodyData& WaterBodyDataEntry = WaterBodyData[WaterBodyComponent->GetWaterBodyIndex()];
+				WaterBodyDataEntry.WaterZoneIndex = WaterZoneIndex;
 				WaterBodyDataEntry.TargetWaveMaskDepth = WaterBodyComponent->TargetWaveMaskDepth;
+				WaterBodyDataEntry.FixedVelocityXY = FMath::AsFloat(static_cast<uint32>(FFloat16(FixedVelocity.X).Encoded) | static_cast<uint32>(FFloat16(FixedVelocity.Y).Encoded) << 16u);
+				WaterBodyDataEntry.FixedVelocityZ = static_cast<float>(FixedVelocity.Z);
+				WaterBodyDataEntry.FixedZHeight = WaterBodyComponent->GetConstantSurfaceZ();
+				WaterBodyDataEntry.FixedWaterDepth = WaterBodyComponent->GetConstantDepth();
 
 				if (WaterBodyComponent->HasWaves())
 				{
@@ -220,7 +215,7 @@ void FWaterViewExtension::UpdateGPUBuffers()
 
 						check(WaveDataIndex);
 
-						WaterIndirectionEntry.WaveDataIndex = *WaveDataIndex;
+						WaterBodyDataEntry.WaveDataIndex = *WaveDataIndex;
 						WaterBodyDataEntry.NumWaves = Waves.Num();
 					}
 				}
@@ -228,15 +223,15 @@ void FWaterViewExtension::UpdateGPUBuffers()
 			});
 		}
 
-		TResourceArray<FVector4f> WaterIndirectionBuffer;
-		TResourceArray<FVector4f> WaterDataBuffer;
+		TResourceArray<FVector4f> WaterBodyDataBuffer;
+		TResourceArray<FVector4f> WaterAuxDataBuffer;
 
 		// The first element of the WaterDataBuffer contains the offsets to each of the sub-buffers.
 		// X = WaterZoneDataOffset
 		// Y = WaterWaveDataOffset
-		// Z = WaterBodyDataOffset
+		// Z = Unused
 		// W = Unused
-		WaterDataBuffer.AddZeroed();
+		WaterAuxDataBuffer.AddZeroed();
 
 		// Transform the individual arrays into the single buffer:
 		{
@@ -253,37 +248,36 @@ void FWaterViewExtension::UpdateGPUBuffers()
 				return StartOffset;
 			};
 
-			AppendDataToFloat4Buffer(WaterIndirectionBuffer, WaterIndirection);
+			AppendDataToFloat4Buffer(WaterBodyDataBuffer, WaterBodyData);
 
-			const int32 ZoneDataOffset = AppendDataToFloat4Buffer(WaterDataBuffer, WaterZoneData);
-			const int32 WaveDataOffset = AppendDataToFloat4Buffer(WaterDataBuffer, WaveData);
-			const int32 WaterBodyDataOffset = AppendDataToFloat4Buffer(WaterDataBuffer, WaterBodyData);
+			const int32 ZoneDataOffset = AppendDataToFloat4Buffer(WaterAuxDataBuffer, WaterZoneData);
+			const int32 WaveDataOffset = AppendDataToFloat4Buffer(WaterAuxDataBuffer, WaveData);
 
 			// Store the offsets to each sub-buffer in the first entry.
 			// If this layout ever changes, corresponding decode functions must be updated in GerstnerWaveFunctions.ush!
-			FVector4f& OffsetData = WaterDataBuffer[0];
+			FVector4f& OffsetData = WaterAuxDataBuffer[0];
 			OffsetData.X = ZoneDataOffset;
 			OffsetData.Y = WaveDataOffset;
-			OffsetData.Z = WaterBodyDataOffset;
-			OffsetData.W = 0.f;
+			OffsetData.Z = 0.0f;
+			OffsetData.W = 0.0f;
 		}
 
-		if (WaterIndirectionBuffer.Num() == 0)
+		if (WaterBodyDataBuffer.Num() == 0)
 		{
-			WaterIndirectionBuffer.AddZeroed();
+			WaterBodyDataBuffer.AddZeroed();
 		}
 
 		ENQUEUE_RENDER_COMMAND(AllocateWaterInstanceDataBuffer)
 		(
-			[WaterGPUData=WaterGPUData, WaterDataBuffer, WaterIndirectionBuffer](FRHICommandListImmediate& RHICmdList) mutable
+			[WaterGPUData=WaterGPUData, WaterAuxDataBuffer, WaterBodyDataBuffer](FRHICommandListImmediate& RHICmdList) mutable
 			{
-				FRHIResourceCreateInfo CreateInfoData(TEXT("WaterDataBuffer"), &WaterDataBuffer);
-				WaterGPUData->DataBuffer = RHICmdList.CreateBuffer(WaterDataBuffer.GetResourceDataSize(), BUF_VertexBuffer | BUF_ShaderResource | BUF_Static, sizeof(FVector4f), ERHIAccess::SRVMask, CreateInfoData);
-				WaterGPUData->DataSRV = RHICmdList.CreateShaderResourceView(WaterGPUData->DataBuffer, sizeof(FVector4f), PF_A32B32G32R32F);
+				FRHIResourceCreateInfo AuxDataCreateInfo(TEXT("WaterAuxDataBuffer"), &WaterAuxDataBuffer);
+				WaterGPUData->AuxDataBuffer = RHICmdList.CreateBuffer(WaterAuxDataBuffer.GetResourceDataSize(), BUF_VertexBuffer | BUF_ShaderResource | BUF_Static, sizeof(FVector4f), ERHIAccess::SRVMask, AuxDataCreateInfo);
+				WaterGPUData->AuxDataSRV = RHICmdList.CreateShaderResourceView(WaterGPUData->AuxDataBuffer, sizeof(FVector4f), PF_A32B32G32R32F);
 
-				FRHIResourceCreateInfo CreateInfoIndirection(TEXT("WaterIndirectionBuffer"), &WaterIndirectionBuffer);
-				WaterGPUData->IndirectionBuffer = RHICmdList.CreateBuffer(WaterIndirectionBuffer.GetResourceDataSize(), BUF_VertexBuffer | BUF_ShaderResource | BUF_Static, sizeof(FVector4f), ERHIAccess::SRVMask, CreateInfoIndirection);
-				WaterGPUData->IndirectionSRV = RHICmdList.CreateShaderResourceView(WaterGPUData->IndirectionBuffer, sizeof(FVector4f), PF_A32B32G32R32F);
+				FRHIResourceCreateInfo WaterBodyDataCreateInfo(TEXT("WaterBodyDataBuffer"), &WaterBodyDataBuffer);
+				WaterGPUData->WaterBodyDataBuffer = RHICmdList.CreateBuffer(WaterBodyDataBuffer.GetResourceDataSize(), BUF_VertexBuffer | BUF_ShaderResource | BUF_Static, sizeof(FVector4f), ERHIAccess::SRVMask, WaterBodyDataCreateInfo);
+				WaterGPUData->WaterBodyDataSRV = RHICmdList.CreateShaderResourceView(WaterGPUData->WaterBodyDataBuffer, sizeof(FVector4f), PF_A32B32G32R32F);
 			}
 		);
 
@@ -413,10 +407,11 @@ void FWaterViewExtension::PreRenderViewFamily_RenderThread(FRDGBuilder& GraphBui
 
 void FWaterViewExtension::PreRenderView_RenderThread(FRDGBuilder& GraphBuilder, FSceneView& InView)
 {
-	if (WaterGPUData->DataSRV && WaterGPUData->IndirectionSRV)
+	if (WaterGPUData->WaterBodyDataSRV && WaterGPUData->AuxDataSRV)
 	{
-		InView.WaterDataBuffer = WaterGPUData->DataSRV;
-		InView.WaterIndirectionBuffer = WaterGPUData->IndirectionSRV;
+		// TODO: Rename members on FSceneView in a separate CL. This will invalidate almost all shaders.
+		InView.WaterDataBuffer = WaterGPUData->AuxDataSRV;
+		InView.WaterIndirectionBuffer = WaterGPUData->WaterBodyDataSRV;
 	}
 }
 
diff --git a/Engine/Plugins/Experimental/Water/Source/Runtime/Public/WaterBodyComponent.h b/Engine/Plugins/Experimental/Water/Source/Runtime/Public/WaterBodyComponent.h
--- a/Engine/Plugins/Experimental/Water/Source/Runtime/Public/WaterBodyComponent.h
+++ b/Engine/Plugins/Experimental/Water/Source/Runtime/Public/WaterBodyComponent.h
@@ -566,14 +566,7 @@ public:
 	static const FName WaterBodyZOffsetParamName;
 	static const FName WaterVelocityAndHeightName;
 	static const FName GlobalOceanHeightName;
-	static const FName FixedZHeightName;
-	static const FName FixedVelocityName;
-	static const FName FixedWaterDepthName;
-	static const FName WaterAreaParamName;
 	static const FName MaxFlowVelocityParamName;
-	static const FName WaterZMinParamName;
-	static const FName WaterZMaxParamName;
-	static const FName GroundZMinParamName;
 
 	UPROPERTY(EditDefaultsOnly, Category = Collision, meta = (EditCondition = "bGenerateCollisions"))
 	TObjectPtr<UPhysicalMaterial> PhysicalMaterial;
diff --git a/Engine/Plugins/Experimental/Water/Source/Runtime/Public/WaterViewExtension.h b/Engine/Plugins/Experimental/Water/Source/Runtime/Public/WaterViewExtension.h
--- a/Engine/Plugins/Experimental/Water/Source/Runtime/Public/WaterViewExtension.h
+++ b/Engine/Plugins/Experimental/Water/Source/Runtime/Public/WaterViewExtension.h
@@ -42,11 +42,11 @@ private:
 
 	struct FWaterGPUResources
 	{
-		FBufferRHIRef DataBuffer;
-		FShaderResourceViewRHIRef DataSRV;
+		FBufferRHIRef WaterBodyDataBuffer;
+		FShaderResourceViewRHIRef WaterBodyDataSRV;
 
-		FBufferRHIRef IndirectionBuffer;
-		FShaderResourceViewRHIRef IndirectionSRV;
+		FBufferRHIRef AuxDataBuffer;
+		FShaderResourceViewRHIRef AuxDataSRV;
 	};
 
 	TSharedRef<FWaterGPUResources, ESPMode::ThreadSafe> WaterGPUData;
