Fixes for occlusion culling on the new GPU water quadtree. HZB had incorrect texture and view rect dimensions passed as uniforms. Raster based occlusion culling had incorrect winding order specified as front face and used waveops in an incorrect way: Waves can actually span multiple instances, so only writing to the result buffer with the first lane might miss certain instances.
Also slightly cleaned up the FWaterQuadTreeGPU::Traverse interface. It now directly takes an array of FSceneView to render.

#rnx

diff --git a/Engine/Plugins/Experimental/Water/Shaders/Private/WaterQuadTreeDraws.usf b/Engine/Plugins/Experimental/Water/Shaders/Private/WaterQuadTreeDraws.usf
--- a/Engine/Plugins/Experimental/Water/Shaders/Private/WaterQuadTreeDraws.usf
+++ b/Engine/Plugins/Experimental/Water/Shaders/Private/WaterQuadTreeDraws.usf
@@ -141,8 +141,6 @@ RWStructuredBuffer<FWaterTileDrawIndexedIndirectArgs> OcclusionQueryIndirectArgs
 Texture2D<float4> QuadTreeTexture;
 Texture2D<float4> WaterZBoundsTexture;
 StructuredBuffer<FWaterBodyRenderData> WaterBodyRenderData;
-float4x4 TranslatedWorldToClip;
-float4x4 ViewToClip;
 float4 CullingBoundsAABB;
 float3 QuadTreePosition;
 float3 ObserverPosition;
@@ -211,6 +209,8 @@ void MainCS(uint3 DispatchThreadID : SV_DispatchThreadID)
 		return;
 	}
 
+	ResolvedView = ResolveView();
+
 	// Compute which LOD level this thread belongs to and then find the node index relative to that LOD level
 	uint3 NodeCoord = 0;
 	{
@@ -250,8 +250,8 @@ void MainCS(uint3 DispatchThreadID : SV_DispatchThreadID)
 		const float3 NodeCenter = (NodeAABBMin + NodeAABBMax) * 0.5f;
 		const float3 NodeExtent = float3((NodeAABBMax.xy - NodeAABBMin.xy) * 0.5f, (NodeAABBMax.z - NodeAABBMin.z) * 0.5f + 1e-5f);
 	
-		const bool bIsOrtho = IsOrthoProjection(ViewToClip);
-		FFrustumCullData Cull = BoxCullFrustum(NodeCenter, NodeExtent, TranslatedWorldToClip, ViewToClip, bIsOrtho, true /* near clip */, false /* skip culling */);
+		const bool bIsOrtho = IsOrthoProjection(ResolvedView.ViewToClip);
+		FFrustumCullData Cull = BoxCullFrustum(NodeCenter, NodeExtent, ResolvedView.TranslatedWorldToClip, ResolvedView.ViewToClip, bIsOrtho, true /* near clip */, false /* skip culling */);
 
 		if (bHZBOcclusionCullingEnabled && !Cull.bCrossesNearPlane)
 		{
@@ -357,7 +357,6 @@ void MainCS(uint3 DispatchThreadID : SV_DispatchThreadID)
 #ifdef OCCLUSION_QUERY_RASTER_VS
 
 Buffer<float4> OcclusionQueryBoxes;
-float4x4 ViewProjection;
 
 void MainVS(
 	in float3 InPosition : ATTRIBUTE0,
@@ -366,6 +365,8 @@ void MainVS(
 	out float4 OutPosition : SV_Position
 )
 {
+	ResolvedView = ResolveView();
+
 	OutQueryIndex = InInstanceId;
 
 	const float4 CenterAndCrossesNearPlane = OcclusionQueryBoxes[InInstanceId * 2 + 0];
@@ -376,7 +377,7 @@ void MainVS(
 		const float3 Center = CenterAndCrossesNearPlane.xyz;
 		const float3 Extent = OcclusionQueryBoxes[InInstanceId * 2 + 1].xyz;
 		const float3 TranslatedWorldPosition = (InPosition * Extent) + Center;
-		OutPosition = mul(float4(TranslatedWorldPosition, 1.0f), ViewProjection);
+		OutPosition = mul(float4(TranslatedWorldPosition, 1.0f), ResolvedView.TranslatedWorldToClip);
 	}
 	else
 	{
@@ -397,14 +398,7 @@ void MainPS(
 	out float4 OutColor : SV_Target0
 )
 {
-#if COMPILER_SUPPORTS_WAVE_VOTE
-	if (WaveIsFirstLane()) // Try to reduce the number of threads writing the same value to the same location.
-	{
-		Visibility[InQueryIndex] = 1;
-	}
-#else
 	Visibility[InQueryIndex] = 1;
-#endif
 	OutColor = float4(1.0f, 0.0f, 0.0f, 1.0f);
 }
 
diff --git a/Engine/Plugins/Experimental/Water/Source/Runtime/Private/WaterMeshSceneProxy.cpp b/Engine/Plugins/Experimental/Water/Source/Runtime/Private/WaterMeshSceneProxy.cpp
--- a/Engine/Plugins/Experimental/Water/Source/Runtime/Private/WaterMeshSceneProxy.cpp
+++ b/Engine/Plugins/Experimental/Water/Source/Runtime/Private/WaterMeshSceneProxy.cpp
@@ -358,9 +358,7 @@ void FWaterMeshSceneProxy::GetDynamicMeshElements(const TArray<const FSceneView*
 	{
 		bool bEncounteredISRView = false;
 		int32 InstanceFactor = 1;
-		int32 NumVisibleViews = 0;
-
-		TArray<FWaterQuadTreeGPU::FPerViewInfo> PerViewInfoArray;
+		TArray<const FSceneView*> VisibleViews;
 
 		// Gather per view data for all renderable views
 		for (int32 ViewIndex = 0; ViewIndex < Views.Num(); ViewIndex++)
@@ -375,21 +373,11 @@ void FWaterMeshSceneProxy::GetDynamicMeshElements(const TArray<const FSceneView*
 			// skip gathering visible tiles from instanced right eye views
 			if ((VisibilityMap & (1 << ViewIndex)) && (!bEncounteredISRView || View->IsPrimarySceneView()))
 			{
-				const FVector QuadTreePosition = FVector(WaterQuadTree.GetTileRegion().Min, WaterQuadTreeMinHeight);
-				const FVector PreViewTranslation = View->ViewMatrices.GetPreViewTranslation();
-				const FViewInfo* ViewInfo = View->bIsViewInfo ? reinterpret_cast<const FViewInfo*>(View) : nullptr;
-
-				FWaterQuadTreeGPU::FPerViewInfo& PerViewInfo = PerViewInfoArray.AddDefaulted_GetRef();
-				PerViewInfo.QuadTreePositionTranslatedWorldSpace = FVector3f(QuadTreePosition + PreViewTranslation);
-				PerViewInfo.ObserverPositionTranslatedWorldSpace = FVector3f(View->ViewMatrices.GetViewOrigin() + PreViewTranslation);
-				PerViewInfo.TranslatedWorldToClip = FMatrix44f(View->ViewMatrices.GetTranslatedViewProjectionMatrix());
-				PerViewInfo.ViewToClip = FMatrix44f(View->ViewMatrices.GetProjectionMatrix());
-				PerViewInfo.ViewInfo = ViewInfo;
-				
-				++NumVisibleViews;
+				VisibleViews.Add(View);
 			}
 		}
 
+		const int32 NumVisibleViews = VisibleViews.Num();
 		const int32 NumIndirectDrawCalls = NumBuckets * NumVisibleViews;
 		const int32 NumInstances = WaterQuadTree.GetMaxLeafCount() * NumVisibleViews;
 
@@ -422,7 +410,8 @@ void FWaterMeshSceneProxy::GetDynamicMeshElements(const TArray<const FSceneView*
 			WaterQuadTreeGPUTraverseParams.OutInstanceData0Buffer = IndirectDrawResources.InstanceData0;
 			WaterQuadTreeGPUTraverseParams.OutInstanceData1Buffer = IndirectDrawResources.InstanceData1;
 			WaterQuadTreeGPUTraverseParams.OutInstanceData2Buffer = IndirectDrawResources.InstanceData2;
-			WaterQuadTreeGPUTraverseParams.PerViewInfo = MoveTemp(PerViewInfoArray);
+			WaterQuadTreeGPUTraverseParams.Views = MoveTemp(VisibleViews);
+			WaterQuadTreeGPUTraverseParams.QuadTreePosition = FVector(WaterQuadTree.GetTileRegion().Min, WaterQuadTreeMinHeight);
 			WaterQuadTreeGPUTraverseParams.CullingBounds = WaterInfoBounds.ShiftBy(-WaterQuadTree.GetTileRegion().Min);
 			WaterQuadTreeGPUTraverseParams.NumDensities = DensityCount;
 			WaterQuadTreeGPUTraverseParams.NumMaterials = NumWaterMaterials;
diff --git a/Engine/Plugins/Experimental/Water/Source/Runtime/Private/WaterQuadTreeGPU.cpp b/Engine/Plugins/Experimental/Water/Source/Runtime/Private/WaterQuadTreeGPU.cpp
--- a/Engine/Plugins/Experimental/Water/Source/Runtime/Private/WaterQuadTreeGPU.cpp
+++ b/Engine/Plugins/Experimental/Water/Source/Runtime/Private/WaterQuadTreeGPU.cpp
@@ -169,6 +169,7 @@ public:
 	using FPermutationDomain = TShaderPermutationDomain<FPreciseOcclusionQueries>;
 
 	BEGIN_SHADER_PARAMETER_STRUCT(FParameters, )
+		SHADER_PARAMETER_STRUCT_REF(FViewUniformShaderParameters, View)
 		SHADER_PARAMETER_RDG_BUFFER_UAV(RWByteAddressBuffer, PackedNodes)
 		SHADER_PARAMETER_RDG_BUFFER_UAV(RWBuffer, OcclusionQueryBoxes)
 		SHADER_PARAMETER_RDG_BUFFER_UAV(RWBuffer, OcclusionVisibility)
@@ -178,8 +179,6 @@ public:
 		SHADER_PARAMETER_RDG_BUFFER_SRV(StructuredBuffer, WaterBodyRenderData)
 		SHADER_PARAMETER_RDG_TEXTURE(Texture2D, HZBTexture)
 		SHADER_PARAMETER_SAMPLER(SamplerState, HZBSampler)
-		SHADER_PARAMETER(FMatrix44f, TranslatedWorldToClip)
-		SHADER_PARAMETER(FMatrix44f, ViewToClip)
 		SHADER_PARAMETER(FVector4f, CullingBoundsAABB)
 		SHADER_PARAMETER(FVector2f, HZBSize)
 		SHADER_PARAMETER(FVector2f, HZBViewSize)
@@ -214,8 +213,8 @@ public:
 	SHADER_USE_PARAMETER_STRUCT(FWaterQuadTreeOcclusionQueryVS, FGlobalShader);
 
 	BEGIN_SHADER_PARAMETER_STRUCT(FParameters, )
+		SHADER_PARAMETER_STRUCT_REF(FViewUniformShaderParameters, View)
 		SHADER_PARAMETER_RDG_BUFFER_SRV(Buffer, OcclusionQueryBoxes)
-		SHADER_PARAMETER(FMatrix44f, ViewProjection)
 	END_SHADER_PARAMETER_STRUCT()
 
 	static void ModifyCompilationEnvironment(const FGlobalShaderPermutationParameters& Parameters, FShaderCompilerEnvironment& OutEnvironment)
@@ -633,7 +632,12 @@ void FWaterQuadTreeGPU::Traverse(FRDGBuilder& GraphBuilder, const FTraverseParam
 	// Iterate over all views for which to create indirect water draws
 	for (uint32 ViewIndex = 0; ViewIndex < Params.NumViews; ++ViewIndex)
 	{
-		const FPerViewInfo& PerViewInfo = Params.PerViewInfo[ViewIndex];
+		const FSceneView* View = Params.Views[ViewIndex];
+		const FViewInfo* ViewInfo = View->bIsViewInfo ? reinterpret_cast<const FViewInfo*>(View) : nullptr;
+
+		const FVector PreViewTranslation = View->ViewMatrices.GetPreViewTranslation();
+		const FVector3f QuadTreePositionTranslatedWorldSpace = FVector3f(Params.QuadTreePosition + PreViewTranslation);
+		const FVector3f ObserverPositionTranslatedWorldSpace = FVector3f(View->ViewMatrices.GetViewOrigin() + PreViewTranslation);
 
 		const uint32 MaxNumDraws = QuadTreeResolution.X * QuadTreeResolution.Y;
 		FRDGBuffer* PackedNodes = GraphBuilder.CreateBuffer(FRDGBufferDesc::CreateByteAddressDesc((1 + MaxNumDraws) * sizeof(uint32)), TEXT("WaterQuadTree.PackedNodes"));
@@ -643,13 +647,15 @@ void FWaterQuadTreeGPU::Traverse(FRDGBuilder& GraphBuilder, const FTraverseParam
 		FRDGBufferSRV* PackedNodesSRV = GraphBuilder.CreateSRV(FRDGBufferSRVDesc(PackedNodes, PF_R32_UINT));
 		FRDGBufferSRV* BucketCountsSRV = GraphBuilder.CreateSRV(FRDGBufferSRVDesc(BucketCounts, PF_R32_UINT));
 
+		const bool bPixelPreciseOQForThisView = bPixelPreciseOcclusionQueries && View->bIsViewInfo;
+
 		FRDGBuffer* OcclusionQueryBoxes = nullptr;
 		FRDGBuffer* OcclusionQueryResults = nullptr;
 		FRDGBufferUAV* OcclusionQueryBoxesUAV = nullptr;
 		FRDGBufferUAV* OcclusionQueryResultsUAV = nullptr;
 		FRDGBufferSRV* OcclusionQueryBoxesSRV = nullptr;
 		FRDGBufferSRV* OcclusionQueryResultsSRV = nullptr;
-		if (bPixelPreciseOcclusionQueries)
+		if (bPixelPreciseOQForThisView)
 		{
 			OcclusionQueryBoxes = GraphBuilder.CreateBuffer(FRDGBufferDesc::CreateBufferDesc(sizeof(float) * 4, MaxNumDraws * 2), TEXT("WaterQuadTree.OcclusionQueryBoxes"));
 			OcclusionQueryResults = GraphBuilder.CreateBuffer(FRDGBufferDesc::CreateBufferDesc(sizeof(uint32), MaxNumDraws), TEXT("WaterQuadTree.OcclusionQueryResults"));
@@ -667,7 +673,7 @@ void FWaterQuadTreeGPU::Traverse(FRDGBuilder& GraphBuilder, const FTraverseParam
 		// Traverse quadtree
 		{
 			FWaterQuadTreeTraverseCS::FPermutationDomain PermutationDomain;
-			PermutationDomain.Set<FWaterQuadTreeTraverseCS::FPreciseOcclusionQueries>(bPixelPreciseOcclusionQueries);
+			PermutationDomain.Set<FWaterQuadTreeTraverseCS::FPreciseOcclusionQueries>(bPixelPreciseOQForThisView);
 			TShaderMapRef<FWaterQuadTreeTraverseCS> ComputeShader(ShaderMap, PermutationDomain);
 
 			const uint32 NumMipLevels = QuadTreeTextureRDG->Desc.NumMips;
@@ -679,6 +685,7 @@ void FWaterQuadTreeGPU::Traverse(FRDGBuilder& GraphBuilder, const FTraverseParam
 			}
 
 			FWaterQuadTreeTraverseCS::FParameters* PassParameters = GraphBuilder.AllocParameters<FWaterQuadTreeTraverseCS::FParameters>();
+			PassParameters->View = GetShaderBinding(View->ViewUniformBuffer);
 			PassParameters->PackedNodes = PackedNodesUAV;
 			PassParameters->OcclusionQueryBoxes = OcclusionQueryBoxesUAV;
 			PassParameters->OcclusionVisibility = OcclusionQueryResultsUAV;
@@ -686,15 +693,13 @@ void FWaterQuadTreeGPU::Traverse(FRDGBuilder& GraphBuilder, const FTraverseParam
 			PassParameters->QuadTreeTexture = QuadTreeTextureRDG;
 			PassParameters->WaterZBoundsTexture = WaterZBoundsTextureRDG;
 			PassParameters->WaterBodyRenderData = WaterBodyRenderDataBufferSRV;
-			PassParameters->HZBTexture = (PerViewInfo.ViewInfo && PerViewInfo.ViewInfo->HZB) ? PerViewInfo.ViewInfo->HZB : GSystemTextures.GetBlackDummy(GraphBuilder);
+			PassParameters->HZBTexture = (ViewInfo && ViewInfo->HZB) ? ViewInfo->HZB : GSystemTextures.GetBlackDummy(GraphBuilder);
 			PassParameters->HZBSampler = TStaticSamplerState<SF_Point, AM_Clamp, AM_Clamp, AM_Clamp>::GetRHI();
-			PassParameters->TranslatedWorldToClip = PerViewInfo.TranslatedWorldToClip;
-			PassParameters->ViewToClip = PerViewInfo.ViewToClip;
 			PassParameters->CullingBoundsAABB = FVector4f(Params.CullingBounds.Min.X, Params.CullingBounds.Min.Y, Params.CullingBounds.Max.X, Params.CullingBounds.Max.Y);
-			PassParameters->HZBSize = PerViewInfo.ViewInfo ? FVector2f(PerViewInfo.ViewInfo->ViewRect.Size()) : FVector2f::ZeroVector;
-			PassParameters->HZBViewSize = PerViewInfo.ViewInfo ? FVector2f(PerViewInfo.ViewInfo->HZBMipmap0Size) : FVector2f::ZeroVector;
-			PassParameters->QuadTreePosition = PerViewInfo.QuadTreePositionTranslatedWorldSpace;
-			PassParameters->ObserverPosition = PerViewInfo.ObserverPositionTranslatedWorldSpace;
+			PassParameters->HZBSize = ViewInfo ? FVector2f(ViewInfo->HZBMipmap0Size) : FVector2f::ZeroVector;
+			PassParameters->HZBViewSize = ViewInfo ? FVector2f(ViewInfo->ViewRect.Size()) : FVector2f::ZeroVector;
+			PassParameters->QuadTreePosition = QuadTreePositionTranslatedWorldSpace;
+			PassParameters->ObserverPosition = ObserverPositionTranslatedWorldSpace;
 			PassParameters->QuadTreeResolutionX = QuadTreeResolution.X;
 			PassParameters->QuadTreeResolutionY = QuadTreeResolution.Y;
 			PassParameters->NumDensities = Params.NumDensities;
@@ -706,27 +711,27 @@ void FWaterQuadTreeGPU::Traverse(FRDGBuilder& GraphBuilder, const FTraverseParam
 			PassParameters->ForceCollapseDensityLevel = -1; // Params.ForceCollapseDensityLevel; // TODO: Properly implement support for this parameter
 			PassParameters->NumLODs = NumMipLevels;
 			PassParameters->NumDispatchedThreads = NumTexelsTotal;
-			PassParameters->bHZBOcclusionCullingEnabled = bHZBOcclusionQueries && (PerViewInfo.ViewInfo && PerViewInfo.ViewInfo->HZB);
+			PassParameters->bHZBOcclusionCullingEnabled = bHZBOcclusionQueries && (ViewInfo && ViewInfo->HZB);
 
 			FComputeShaderUtils::AddPass(GraphBuilder, RDG_EVENT_NAME("WaterQuadTreeTraverse(View: %i)", ViewIndex), ComputeShader, PassParameters, FComputeShaderUtils::GetGroupCount(NumTexelsTotal, 64));
 		}
 
 		// Raster occlusion queries
-		if (bPixelPreciseOcclusionQueries)
+		if (bPixelPreciseOQForThisView)
 		{
 			TShaderMapRef<FWaterQuadTreeOcclusionQueryVS> VertexShader(ShaderMap);
 			TShaderMapRef<FWaterQuadTreeOcclusionQueryPS> PixelShader(ShaderMap);
 
-			FRDGTextureRef DepthTexture = PerViewInfo.ViewInfo->GetSceneTextures().Depth.Target;
+			FRDGTextureRef DepthTexture = ViewInfo->GetSceneTextures().Depth.Target;
 
 			FWaterQuadTreeOcclusionQueryParameters* PassParameters = GraphBuilder.AllocParameters<FWaterQuadTreeOcclusionQueryParameters>();
 			PassParameters->IndirectDrawArgsBuffer = OcclusionQueryIndirectArgsBuffer;
+			PassParameters->VS.View = GetShaderBinding(View->ViewUniformBuffer);
 			PassParameters->VS.OcclusionQueryBoxes = OcclusionQueryBoxesSRV;
-			PassParameters->VS.ViewProjection = PerViewInfo.TranslatedWorldToClip;
 			PassParameters->PS.Visibility = OcclusionQueryResultsUAV;
 			PassParameters->PS.RenderTargets.DepthStencil = FDepthStencilBinding(DepthTexture, ERenderTargetLoadAction::ELoad, ERenderTargetLoadAction::ENoAction, FExclusiveDepthStencil::DepthRead_StencilNop);
 
-			const FIntVector4 ViewRect = FIntVector4(PerViewInfo.ViewInfo->ViewRect.Min.X, PerViewInfo.ViewInfo->ViewRect.Min.Y, PerViewInfo.ViewInfo->ViewRect.Max.X, PerViewInfo.ViewInfo->ViewRect.Max.Y);
+			const FIntVector4 ViewRect = FIntVector4(ViewInfo->ViewRect.Min.X, ViewInfo->ViewRect.Min.Y, ViewInfo->ViewRect.Max.X, ViewInfo->ViewRect.Max.Y);
 
 			GraphBuilder.AddPass(
 				RDG_EVENT_NAME("WaterQuadTreeOcclusionCulling(View: %i)", ViewIndex),
@@ -742,7 +747,7 @@ void FWaterQuadTreeGPU::Traverse(FRDGBuilder& GraphBuilder, const FTraverseParam
 					GraphicsPSOInit.BoundShaderState.VertexShaderRHI = VertexShader.GetVertexShader();
 					GraphicsPSOInit.BoundShaderState.PixelShaderRHI = PixelShader.GetPixelShader();
 					GraphicsPSOInit.BlendState = TStaticBlendState<>::GetRHI();
-					GraphicsPSOInit.RasterizerState = TStaticRasterizerState<FM_Solid, CM_CW>::GetRHI();
+					GraphicsPSOInit.RasterizerState = TStaticRasterizerState<FM_Solid, CM_CCW>::GetRHI();
 					GraphicsPSOInit.DepthStencilState = TStaticDepthStencilState<false, CF_DepthNearOrEqual>::GetRHI();
 					GraphicsPSOInit.PrimitiveType = PT_TriangleList;
 
@@ -761,7 +766,7 @@ void FWaterQuadTreeGPU::Traverse(FRDGBuilder& GraphBuilder, const FTraverseParam
 		// Compute bucket counts
 		{
 			FWaterQuadTreeBucketCountsCS::FPermutationDomain PermutationDomain;
-			PermutationDomain.Set<FWaterQuadTreeBucketCountsCS::FPreciseOcclusionQueries>(bPixelPreciseOcclusionQueries);
+			PermutationDomain.Set<FWaterQuadTreeBucketCountsCS::FPreciseOcclusionQueries>(bPixelPreciseOQForThisView);
 			TShaderMapRef<FWaterQuadTreeBucketCountsCS> ComputeShader(ShaderMap, PermutationDomain);
 
 			FWaterQuadTreeBucketCountsCS::FParameters* PassParameters = GraphBuilder.AllocParameters<FWaterQuadTreeBucketCountsCS::FParameters>();
@@ -795,7 +800,7 @@ void FWaterQuadTreeGPU::Traverse(FRDGBuilder& GraphBuilder, const FTraverseParam
 		// Generate instance data and fill indirect args
 		{
 			FWaterQuadTreeGenerateInstanceDataCS::FPermutationDomain PermutationDomain;
-			PermutationDomain.Set<FWaterQuadTreeGenerateInstanceDataCS::FPreciseOcclusionQueries>(bPixelPreciseOcclusionQueries);
+			PermutationDomain.Set<FWaterQuadTreeGenerateInstanceDataCS::FPreciseOcclusionQueries>(bPixelPreciseOQForThisView);
 			TShaderMapRef<FWaterQuadTreeGenerateInstanceDataCS> ComputeShader(ShaderMap, PermutationDomain);
 
 			FWaterQuadTreeGenerateInstanceDataCS::FParameters* PassParameters = GraphBuilder.AllocParameters<FWaterQuadTreeGenerateInstanceDataCS::FParameters>();
@@ -809,8 +814,8 @@ void FWaterQuadTreeGPU::Traverse(FRDGBuilder& GraphBuilder, const FTraverseParam
 			PassParameters->PackedNodes = PackedNodesSRV;
 			PassParameters->InstanceDataOffsets = InstanceDataOffsetsBufferSRV;
 			PassParameters->OcclusionResults = OcclusionQueryResultsSRV;
-			PassParameters->QuadTreePosition = PerViewInfo.QuadTreePositionTranslatedWorldSpace;
-			PassParameters->ObserverPosition = PerViewInfo.ObserverPositionTranslatedWorldSpace;
+			PassParameters->QuadTreePosition = QuadTreePositionTranslatedWorldSpace;
+			PassParameters->ObserverPosition = ObserverPositionTranslatedWorldSpace;
 			PassParameters->QuadTreeResolutionX = QuadTreeResolution.X;
 			PassParameters->QuadTreeResolutionY = QuadTreeResolution.Y;
 			PassParameters->NumDensities = Params.NumDensities;
@@ -832,10 +837,10 @@ void FWaterQuadTreeGPU::Traverse(FRDGBuilder& GraphBuilder, const FTraverseParam
 		{
 			ShaderPrint::SetEnabled(true);
 
-			const bool bEnableShaderPrint = PerViewInfo.ViewInfo
+			const bool bEnableShaderPrint = ViewInfo != nullptr
 				&& ShaderPrint::IsEnabled()
-				&& ShaderPrint::IsEnabled(PerViewInfo.ViewInfo->ShaderPrintData)
-				&& ShaderPrint::IsValid(PerViewInfo.ViewInfo->ShaderPrintData)
+				&& ShaderPrint::IsEnabled(ViewInfo->ShaderPrintData)
+				&& ShaderPrint::IsValid(ViewInfo->ShaderPrintData)
 				&& ShaderPrint::IsSupported(GetFeatureLevelShaderPlatform(GMaxRHIFeatureLevel));
 
 			if (bEnableShaderPrint)
@@ -844,17 +849,17 @@ void FWaterQuadTreeGPU::Traverse(FRDGBuilder& GraphBuilder, const FTraverseParam
 				ShaderPrint::RequestSpaceForLines(12 * MaxNumDraws);
 
 				FWaterQuadTreeDebugCS::FPermutationDomain PermutationDomain;
-				PermutationDomain.Set<FWaterQuadTreeDebugCS::FPreciseOcclusionQueries>(bPixelPreciseOcclusionQueries);
+				PermutationDomain.Set<FWaterQuadTreeDebugCS::FPreciseOcclusionQueries>(bPixelPreciseOQForThisView);
 				TShaderMapRef<FWaterQuadTreeDebugCS> ComputeShader(ShaderMap, PermutationDomain);
 
 				FWaterQuadTreeDebugCS::FParameters* PassParameters = GraphBuilder.AllocParameters<FWaterQuadTreeDebugCS::FParameters>();
-				ShaderPrint::SetParameters(GraphBuilder, PerViewInfo.ViewInfo->ShaderPrintData, PassParameters->ShaderPrintUniformBuffer);
+				ShaderPrint::SetParameters(GraphBuilder, ViewInfo->ShaderPrintData, PassParameters->ShaderPrintUniformBuffer);
 				PassParameters->QuadTreeTexture = QuadTreeTextureRDG;
 				PassParameters->WaterZBoundsTexture = WaterZBoundsTextureRDG;
 				PassParameters->WaterBodyRenderData = WaterBodyRenderDataBufferSRV;
 				PassParameters->PackedNodes = PackedNodesSRV;
 				PassParameters->OcclusionResults = OcclusionQueryResultsSRV;
-				PassParameters->QuadTreePosition = PerViewInfo.QuadTreePositionTranslatedWorldSpace;
+				PassParameters->QuadTreePosition = QuadTreePositionTranslatedWorldSpace;
 				PassParameters->NumDispatchedThreads = FMath::Min(FMath::DivideAndRoundUp(MaxNumDraws, 64u), 65535u) * 64u; // The maximum number of dispatched groups is 64k
 				PassParameters->LeafSize = Params.LeafSize;
 				PassParameters->CaptureDepthRange = CaptureDepthRange;
diff --git a/Engine/Plugins/Experimental/Water/Source/Runtime/Private/WaterQuadTreeGPU.h b/Engine/Plugins/Experimental/Water/Source/Runtime/Private/WaterQuadTreeGPU.h
--- a/Engine/Plugins/Experimental/Water/Source/Runtime/Private/WaterQuadTreeGPU.h
+++ b/Engine/Plugins/Experimental/Water/Source/Runtime/Private/WaterQuadTreeGPU.h
@@ -7,7 +7,7 @@ class FRHITexture;
 class FRDGBuilder;
 class FRDGTexture;
 class FRDGPooledBuffer;
-class FViewInfo;
+class FSceneView;
 struct FShaderPrintData;
 
 class FWaterQuadTreeGPU
@@ -53,15 +53,6 @@ public:
 		float CaptureDepthRange = 1.0f;
 	};
 
-	struct FPerViewInfo
-	{
-		FMatrix44f TranslatedWorldToClip = FMatrix44f::Identity;
-		FMatrix44f ViewToClip = FMatrix44f::Identity;
-		FVector3f QuadTreePositionTranslatedWorldSpace = FVector3f::ZeroVector;
-		FVector3f ObserverPositionTranslatedWorldSpace = FVector3f::ZeroVector;
-		const FViewInfo* ViewInfo = nullptr;
-	};
-
 	struct FTraverseParams
 	{
 		FRDGPooledBuffer* OutIndirectArgsBuffer = nullptr;
@@ -69,7 +60,8 @@ public:
 		FRDGPooledBuffer* OutInstanceData0Buffer = nullptr;
 		FRDGPooledBuffer* OutInstanceData1Buffer = nullptr;
 		FRDGPooledBuffer* OutInstanceData2Buffer = nullptr;
-		TArray<FPerViewInfo> PerViewInfo;
+		TArray<const FSceneView*> Views;
+		FVector QuadTreePosition;
 		FBox2D CullingBounds = FBox2D(ForceInit);
 		uint32 NumDensities = 0;
 		uint32 NumMaterials = 0;
