Water: fixed the r.water.waterinfo.showsceneproxies cvar so that it works with the new rendermethod path. Additionally, restored the functionality that allows showing non-dilated and dilated meshes independently.

diff --git a/Engine/Plugins/Experimental/Water/Source/Runtime/Private/WaterBodyInfoMeshComponent.cpp b/Engine/Plugins/Experimental/Water/Source/Runtime/Private/WaterBodyInfoMeshComponent.cpp
--- a/Engine/Plugins/Experimental/Water/Source/Runtime/Private/WaterBodyInfoMeshComponent.cpp
+++ b/Engine/Plugins/Experimental/Water/Source/Runtime/Private/WaterBodyInfoMeshComponent.cpp
@@ -15,10 +15,14 @@ void OnCVarWaterInfoSceneProxiesValueChanged(IConsoleVariable*)
 	}
 }
 
-static TAutoConsoleVariable<bool> CVarShowWaterInfoSceneProxies(
+static TAutoConsoleVariable<int32> CVarShowWaterInfoSceneProxies(
 	TEXT("r.Water.WaterInfo.ShowSceneProxies"),
-	false,
-	TEXT("When enabled, always shows the water scene proxies in the main viewport. Useful for debugging only"),
+	0,
+	TEXT("When enabled, shows the water scene proxies in the main viewport. Useful for debugging only.\n")
+	TEXT("0: Disabled\n")
+	TEXT("1: Only show water info mesh\n")
+	TEXT("2: Only show water info dilated mesh\n")
+	TEXT("3: Show both water info meshes"),
 	FConsoleVariableDelegate::CreateStatic(OnCVarWaterInfoSceneProxiesValueChanged),
 	ECVF_RenderThreadSafe
 );
@@ -73,9 +77,19 @@ bool FWaterBodyInfoMeshSceneProxy::GetMeshElement(int32 LODIndex, int32 BatchInd
 	return bResult;
 }
 
+static bool ShouldShowOutsideWaterInfoPass(bool bIsDilatedMesh)
+{
+#if !UE_BUILD_SHIPPING
+	const int32 ShowWaterInfoSceneProxiesValue = CVarShowWaterInfoSceneProxies.GetValueOnAnyThread();
+	return (ShowWaterInfoSceneProxiesValue == 1 && !bIsDilatedMesh) || (ShowWaterInfoSceneProxiesValue == 2 && bIsDilatedMesh) || (ShowWaterInfoSceneProxiesValue > 2);
+#endif
+
+	return false;
+}
+
 void FWaterBodyInfoMeshSceneProxy::SetEnabled(bool bInEnabled)
 {
-	SetForceHidden(!(bInEnabled || CVarShowWaterInfoSceneProxies.GetValueOnAnyThread()));
+	SetForceHidden(!(bInEnabled || ShouldShowOutsideWaterInfoPass(bIsDilatedMesh)));
 }
 
 FPrimitiveViewRelevance FWaterBodyInfoMeshSceneProxy::GetViewRelevance(const FSceneView* View) const
@@ -85,10 +99,12 @@ FPrimitiveViewRelevance FWaterBodyInfoMeshSceneProxy::GetViewRelevance(const FSc
 	// When water info mesh is rendered with custom render passes, enable the mesh for drawing
 	if (GetWaterInfoRenderingMethod() == 2)
 	{
-		FString PassName = View->CustomRenderPass ? const_cast<FSceneView*>(View)->CustomRenderPass->Name : TEXT("");
+		const FString PassName = View->CustomRenderPass ? const_cast<FSceneView*>(View)->CustomRenderPass->Name : TEXT("");
 		Result.bDrawRelevance = (PassName == TEXT("WaterInfoDepthPass") || PassName == TEXT("WaterInfoColorPass") || PassName == TEXT("WaterInfoDilationPass"));
 		Result.bShadowRelevance = false;
 	}
 
+	Result.bDrawRelevance |= ShouldShowOutsideWaterInfoPass(bIsDilatedMesh);
+
 	return Result;
 }
\ No newline at end of file
