Fix C4702
#rnx

diff --git a/Engine/Plugins/Experimental/Water/Source/Runtime/Private/WaterBodyInfoMeshComponent.cpp b/Engine/Plugins/Experimental/Water/Source/Runtime/Private/WaterBodyInfoMeshComponent.cpp
--- a/Engine/Plugins/Experimental/Water/Source/Runtime/Private/WaterBodyInfoMeshComponent.cpp
+++ b/Engine/Plugins/Experimental/Water/Source/Runtime/Private/WaterBodyInfoMeshComponent.cpp
@@ -79,12 +79,12 @@ bool FWaterBodyInfoMeshSceneProxy::GetMeshElement(int32 LODIndex, int32 BatchInd
 
 static bool ShouldShowOutsideWaterInfoPass(bool bIsDilatedMesh)
 {
-#if !UE_BUILD_SHIPPING
+#if UE_BUILD_SHIPPING
+	return false;
+#else
 	const int32 ShowWaterInfoSceneProxiesValue = CVarShowWaterInfoSceneProxies.GetValueOnAnyThread();
 	return (ShowWaterInfoSceneProxiesValue == 1 && !bIsDilatedMesh) || (ShowWaterInfoSceneProxiesValue == 2 && bIsDilatedMesh) || (ShowWaterInfoSceneProxiesValue > 2);
 #endif
-
-	return false;
 }
 
 void FWaterBodyInfoMeshSceneProxy::SetEnabled(bool bInEnabled)
