Water: fix multiple transactions registered when modifying the water spline metadata details. Additionally, moved the transaction scope inside SetValues to narrow down what changes are captured in the transaction.

diff --git a/Engine/Plugins/Experimental/Water/Source/Editor/Private/WaterSplineMetadataDetails.cpp b/Engine/Plugins/Experimental/Water/Source/Editor/Private/WaterSplineMetadataDetails.cpp
--- a/Engine/Plugins/Experimental/Water/Source/Editor/Private/WaterSplineMetadataDetails.cpp
+++ b/Engine/Plugins/Experimental/Water/Source/Editor/Private/WaterSplineMetadataDetails.cpp
@@ -96,12 +96,16 @@ void FWaterSplineMetadataDetails::Update(USplineComponent* InSplineComponent, co
 }
 
 template<class T>
-void SetValues(FWaterSplineMetadataDetails& Details, TArray<FInterpCurvePoint<T>>& Points, const T& NewValue, ETextCommit::Type CommitInfo)
+void SetValues(FWaterSplineMetadataDetails& Details, TArray<FInterpCurvePoint<T>>& Points, const T& NewValue, ETextCommit::Type CommitInfo, const FText& TransactionText)
 {
-	Details.SplineComp->GetSplinePointsMetadata()->Modify();
-	for (int32 Index : Details.SelectedKeys)
+	// Scope the transaction to only include the value change and none of the derived data changes that might arise from NotifyPropertyModified
 	{
-		Points[Index].OutVal = NewValue;
+		const FScopedTransaction Transaction(TransactionText);
+		Details.SplineComp->GetSplinePointsMetadata()->Modify();
+		for (int32 Index : Details.SelectedKeys)
+		{
+			Points[Index].OutVal = NewValue;
+		}
 	}
 
 	Details.SplineComp->UpdateSpline();
@@ -114,6 +118,18 @@ void SetValues(FWaterSplineMetadataDetails& Details, TArray<FInterpCurvePoint<T>
 	GEditor->RedrawLevelEditingViewports(true);
 }
 
+template <typename T>
+bool CheckIfDifferent(const FWaterSplineMetadataDetails& Details, const TArray<FInterpCurvePoint<T>>& Points, const T& NewValue)
+{
+	bool bIsModified = false;
+	for (int32 Index : Details.SelectedKeys)
+	{
+		bIsModified |= Points[Index].OutVal != NewValue;
+	}
+
+	return bIsModified;
+}
+
 void FWaterSplineMetadataDetails::OnBeginSliderMovement()
 {
 	EditSliderValueTransaction = MakeUnique<FScopedTransaction>(LOCTEXT("EditWaterProperty", "Edit Water Property"));
@@ -134,8 +150,10 @@ void FWaterSplineMetadataDetails::OnSetDepth(float NewValue, ETextCommit::Type C
 {
 	if (UWaterSplineMetadata* Metadata = GetMetadata())
 	{
-		const FScopedTransaction Transaction(LOCTEXT("SetSplineDepth", "Set spline point water depth"));
-		SetValues<float>(*this, Metadata->Depth.Points, NewValue, CommitInfo);
+		if (CheckIfDifferent(*this, Metadata->Depth.Points, NewValue))
+		{
+			SetValues(*this, Metadata->Depth.Points, NewValue, CommitInfo, LOCTEXT("SetSplineDepth", "Set spline point water depth"));
+		}
 	}
 }
 
@@ -143,8 +161,10 @@ void FWaterSplineMetadataDetails::OnSetRiverWidth(float NewValue, ETextCommit::T
 {
 	if (UWaterSplineMetadata* Metadata = GetMetadata())
 	{
-		const FScopedTransaction Transaction(LOCTEXT("SetSplineWaterWidth", "Set spline point river width"));
-		SetValues<float>(*this, Metadata->RiverWidth.Points, NewValue, CommitInfo);
+		if (CheckIfDifferent(*this, Metadata->Depth.Points, NewValue))
+		{
+			SetValues(*this, Metadata->RiverWidth.Points, NewValue, CommitInfo, LOCTEXT("SetSplineWaterWidth", "Set spline point river width"));
+		}
 	}
 }
 
@@ -152,8 +172,10 @@ void FWaterSplineMetadataDetails::OnSetVelocity(float NewValue, ETextCommit::Typ
 {
 	if (UWaterSplineMetadata* Metadata = GetMetadata())
 	{
-		const FScopedTransaction Transaction(LOCTEXT("SetSplineWaterVelocity", "Set spline point water velocity"));
-		SetValues(*this, Metadata->WaterVelocityScalar.Points, NewValue, CommitInfo);
+		if (CheckIfDifferent(*this, Metadata->Depth.Points, NewValue))
+		{
+			SetValues(*this, Metadata->WaterVelocityScalar.Points, NewValue, CommitInfo, LOCTEXT("SetSplineWaterVelocity", "Set spline point water velocity"));
+		}
 	}
 }
 
@@ -161,8 +183,10 @@ void FWaterSplineMetadataDetails::OnSetAudioIntensity(float NewValue, ETextCommi
 {
 	if (UWaterSplineMetadata* Metadata = GetMetadata())
 	{
-		const FScopedTransaction Transaction(LOCTEXT("SetSpline point audio intensity", "Set spline point audio intensity"));
-		SetValues(*this, Metadata->AudioIntensity.Points, NewValue, CommitInfo);
+		if (CheckIfDifferent(*this, Metadata->Depth.Points, NewValue))
+		{
+			SetValues(*this, Metadata->AudioIntensity.Points, NewValue, CommitInfo, LOCTEXT("SetSpline point audio intensity", "Set spline point audio intensity"));
+		}
 	}
 }
 
