Water: fixed issue where WaterBodymeshComponent was blocking its static mesh's build immediately in PostLoad to fixup the body setup.

diff --git a/Engine/Plugins/Experimental/Water/Source/Runtime/Private/WaterBodyMeshComponent.cpp b/Engine/Plugins/Experimental/Water/Source/Runtime/Private/WaterBodyMeshComponent.cpp
--- a/Engine/Plugins/Experimental/Water/Source/Runtime/Private/WaterBodyMeshComponent.cpp
+++ b/Engine/Plugins/Experimental/Water/Source/Runtime/Private/WaterBodyMeshComponent.cpp
@@ -87,12 +87,36 @@ bool UWaterBodyMeshComponent::CanCreateSceneProxy() const
 #if WITH_EDITOR
 void UWaterBodyMeshComponent::PostLoad()
 {
+	TRACE_CPUPROFILER_EVENT_SCOPE(WaterBodyMeshComponent::PostLoad);
+
 	Super::PostLoad();
 
+	// If the static mesh is currently compiling, we should not call GetBodySetup.
+	// This would block the async compilation immediately after it starts for each water body mesh and slows down editor load.
+	UStaticMesh* Mesh = GetStaticMesh();
+	if (Mesh && !Mesh->IsCompiling())
+	{
+		FixupCollisionOnBodySetup();
+	}
+}
+
+void UWaterBodyMeshComponent::PostStaticMeshCompilation()
+{
+	TRACE_CPUPROFILER_EVENT_SCOPE(WaterBodyMeshComponent::PostStaticMeshCompilation);
+
+	Super::PostStaticMeshCompilation();
+
+	FixupCollisionOnBodySetup();
+}
+
+void UWaterBodyMeshComponent::FixupCollisionOnBodySetup()
+{
+	TRACE_CPUPROFILER_EVENT_SCOPE(WaterBodyMeshComponent::FixupCollisionOnBodySetup);
+
 	// Fix for a bug that existing where the body setup was not created with bNeverNeedsCookedCollisionData, and also address an issue where bNeverNeedsCookedCollisionData was allowing meshes to generate cooked data.
 	UBodySetup* BodySetup = GetBodySetup();
 	UStaticMesh* Mesh = GetStaticMesh();
-	bool bResetBodySetup = Mesh != nullptr && (BodySetup == nullptr || !BodySetup->bNeverNeedsCookedCollisionData || BodySetup->bHasCookedCollisionData);
+	const bool bResetBodySetup = Mesh != nullptr && (BodySetup == nullptr || !BodySetup->bNeverNeedsCookedCollisionData || BodySetup->bHasCookedCollisionData);
 	if (bResetBodySetup)
 	{
 		Mesh->CreateBodySetup();
diff --git a/Engine/Plugins/Experimental/Water/Source/Runtime/Public/WaterBodyMeshComponent.h b/Engine/Plugins/Experimental/Water/Source/Runtime/Public/WaterBodyMeshComponent.h
--- a/Engine/Plugins/Experimental/Water/Source/Runtime/Public/WaterBodyMeshComponent.h
+++ b/Engine/Plugins/Experimental/Water/Source/Runtime/Public/WaterBodyMeshComponent.h
@@ -32,6 +32,8 @@ protected:
 
 #if WITH_EDITOR
 	virtual void PostLoad() override;
+	virtual void PostStaticMeshCompilation() override;
+	void FixupCollisionOnBodySetup();
 #endif
 };
 
