Always call Super::EndPlay.

diff --git a/Engine/Plugins/Experimental/Water/Source/Runtime/Private/WaterZoneActor.cpp b/Engine/Plugins/Experimental/Water/Source/Runtime/Private/WaterZoneActor.cpp
--- a/Engine/Plugins/Experimental/Water/Source/Runtime/Private/WaterZoneActor.cpp
+++ b/Engine/Plugins/Experimental/Water/Source/Runtime/Private/WaterZoneActor.cpp
@@ -124,6 +124,8 @@ void AWaterZone::EndPlay(const EEndPlayReason::Type EndPlayReason)
 {
 	FWorldDelegates::LevelAddedToWorld.RemoveAll(this);
 	FWorldDelegates::LevelRemovedFromWorld.RemoveAll(this);
+
+	Super::EndPlay(EndPlayReason);
 }
 
 void AWaterZone::PostLoadSubobjects(FObjectInstancingGraph* OuterInstanceGraph)
