Added custom clear shader to the GPU water quadtree which clears both the draw buckets counts buffer and the counter at index  0 of the PackedNodes buffer in a single dispatch. Fixed indirect draw argument buffers being used as StructuredBuffers without ever declaring the buffers as being used with StructuredBuffers. Minor adjustment in WaterQuadTreeDraws.usf to get the case of DispatchThreadId in line with the style used in the rest of the engine (ThreadId instead of ThreadID).

#rnx

diff --git a/Engine/Plugins/Experimental/Water/Shaders/Private/WaterQuadTreeDraws.usf b/Engine/Plugins/Experimental/Water/Shaders/Private/WaterQuadTreeDraws.usf
--- a/Engine/Plugins/Experimental/Water/Shaders/Private/WaterQuadTreeDraws.usf
+++ b/Engine/Plugins/Experimental/Water/Shaders/Private/WaterQuadTreeDraws.usf
@@ -4,15 +4,6 @@
 #include "/Engine/Private/PackUnpack.ush"
 #include "WaterQuadTreeCommon.ush"
 
-struct FWaterTileDrawIndexedIndirectArgs
-{
-	uint IndexCountPerInstance;
-	uint InstanceCount;
-	uint StartIndexLocation;
-	int  BaseVertexLocation;
-	uint StartInstanceLocation;
-};
-
 uint PackNodeCoord(uint4 InNode)
 {
 	uint Result = 0;
@@ -66,9 +57,9 @@ float4 GetNodeAABB2D(uint3 InNodeCoord, float3 InQuadTreePosition, float InLeafS
 #define PRECISE_OCCLUSION_QUERIES 0
 #endif
 
-RWStructuredBuffer<FWaterTileDrawIndexedIndirectArgs> IndirectArgs;
+RWBuffer<uint> IndirectArgs;
 #if PRECISE_OCCLUSION_QUERIES
-RWStructuredBuffer<FWaterTileDrawIndexedIndirectArgs> OcclusionQueryIndirectArgs;
+RWBuffer<uint> OcclusionQueryIndirectArgs;
 #endif
 uint NumDensities;
 uint NumMaterials;
@@ -76,10 +67,10 @@ uint NumViews;
 uint NumQuadsLOD0;
 
 [numthreads(64, 1, 1)]
-void MainCS(uint3 DispatchThreadID : SV_DispatchThreadID)
+void MainCS(uint3 DispatchThreadId : SV_DispatchThreadID)
 {
 	const uint NumBuckets = NumDensities * NumMaterials;
-	const uint BucketIndex = DispatchThreadID.x;
+	const uint BucketIndex = DispatchThreadId.x;
 
 	if (BucketIndex >= NumBuckets)
 	{
@@ -89,32 +80,28 @@ void MainCS(uint3 DispatchThreadID : SV_DispatchThreadID)
 	const uint DensityIndex = BucketIndex % NumDensities;
 	const uint MaterialIndex = BucketIndex / NumDensities;
 	const uint NumQuads = NumQuadsLOD0 >> DensityIndex;
-
-	FWaterTileDrawIndexedIndirectArgs BucketArgs;
-	BucketArgs.IndexCountPerInstance = NumQuads * NumQuads * 6;
-	BucketArgs.InstanceCount = 0;
-	BucketArgs.StartIndexLocation = 0;
-	BucketArgs.BaseVertexLocation = 0;
-	BucketArgs.StartInstanceLocation = 0;
+	const uint IndexCountPerInstance = NumQuads * NumQuads * 6;
 
 	for (uint ViewIndex = 0; ViewIndex < NumViews; ++ViewIndex)
 	{
-		IndirectArgs[NumBuckets * ViewIndex + BucketIndex] = BucketArgs;
+		const uint IndirectArgIndex = NumBuckets * ViewIndex + BucketIndex;
+		IndirectArgs[IndirectArgIndex * 5 + 0] = IndexCountPerInstance;
+		IndirectArgs[IndirectArgIndex * 5 + 1] = 0; // InstanceCount
+		IndirectArgs[IndirectArgIndex * 5 + 2] = 0; // StartIndexLocation
+		IndirectArgs[IndirectArgIndex * 5 + 3] = 0; // BaseVertexLocation
+		IndirectArgs[IndirectArgIndex * 5 + 4] = 0; // StartInstanceLocation
 	}
 
 #if PRECISE_OCCLUSION_QUERIES
-	if (all(DispatchThreadID == 0))
+	if (all(DispatchThreadId == 0))
 	{
-		FWaterTileDrawIndexedIndirectArgs OcclusionQueryArgs;
-		OcclusionQueryArgs.IndexCountPerInstance = 36; // 12 triangles per cube
-		OcclusionQueryArgs.InstanceCount = 0;
-		OcclusionQueryArgs.StartIndexLocation = 0;
-		OcclusionQueryArgs.BaseVertexLocation = 0;
-		OcclusionQueryArgs.StartInstanceLocation = 0;
-
 		for (uint ViewIndex = 0; ViewIndex < NumViews; ++ViewIndex)
 		{
-			OcclusionQueryIndirectArgs[ViewIndex] = OcclusionQueryArgs;
+			OcclusionQueryIndirectArgs[ViewIndex * 5 + 0] = 36;// IndexCountPerInstance; 12 triangles per cube
+			OcclusionQueryIndirectArgs[ViewIndex * 5 + 1] = 0; // InstanceCount
+			OcclusionQueryIndirectArgs[ViewIndex * 5 + 2] = 0; // StartIndexLocation
+			OcclusionQueryIndirectArgs[ViewIndex * 5 + 3] = 0; // BaseVertexLocation
+			OcclusionQueryIndirectArgs[ViewIndex * 5 + 4] = 0; // StartInstanceLocation
 		}
 	}
 #endif // PRECISE_OCCLUSION_QUERIES
@@ -122,6 +109,29 @@ void MainCS(uint3 DispatchThreadID : SV_DispatchThreadID)
 
 #endif // INITIALIZE_INDIRECT_ARGS
 
+#ifdef CLEAR_PER_VIEW_BUFFERS
+
+RWByteAddressBuffer BucketCounts;
+RWByteAddressBuffer PackedNodes;
+uint NumDrawBuckets;
+
+[numthreads(64, 1, 1)]
+void MainCS(uint3 DispatchThreadId : SV_DispatchThreadID)
+{
+	if (all(DispatchThreadId == 0))
+	{
+		// PackedNodes stores a counter at index 0 which we need to clear.
+		PackedNodes.Store(0, 0);
+	}
+
+	if (DispatchThreadId.x < NumDrawBuckets)
+	{
+		BucketCounts.Store(DispatchThreadId.x << 2u, 0);
+	}
+}
+
+#endif // CLEAR_PER_VIEW_BUFFERS
+
 #ifdef QUAD_TREE_TRAVERSE
 
 #include "/Engine/Private/Nanite/NaniteHZBCull.ush" // BoxCullFrustum()
@@ -135,7 +145,7 @@ RWByteAddressBuffer PackedNodes;
 #if PRECISE_OCCLUSION_QUERIES
 RWBuffer<float4> OcclusionQueryBoxes;
 RWBuffer<uint> OcclusionVisibility;
-RWStructuredBuffer<FWaterTileDrawIndexedIndirectArgs> OcclusionQueryIndirectArgs;
+RWBuffer<uint> OcclusionQueryIndirectArgs;
 #endif // PRECISE_OCCLUSION_QUERIES
 
 Texture2D<float4> QuadTreeTexture;
@@ -202,9 +212,9 @@ bool CanRender(FWaterQuadTreeNode Node, int InDensityLevel, int InForceCollapseD
 }
 
 [numthreads(64, 1, 1)]
-void MainCS(uint3 DispatchThreadID : SV_DispatchThreadID)
+void MainCS(uint3 DispatchThreadId : SV_DispatchThreadID)
 {
-	if (DispatchThreadID.x >= NumDispatchedThreads)
+	if (DispatchThreadId.x >= NumDispatchedThreads)
 	{
 		return;
 	}
@@ -217,14 +227,14 @@ void MainCS(uint3 DispatchThreadID : SV_DispatchThreadID)
 		uint LODLevelIndex = 0;
 		uint2 Resolution = uint2(QuadTreeResolutionX, QuadTreeResolutionY);
 		uint NumPreviousNodes = 0;
-		while (LODLevelIndex < NumLODs && DispatchThreadID.x >= (Resolution.x * Resolution.y + NumPreviousNodes))
+		while (LODLevelIndex < NumLODs && DispatchThreadId.x >= (Resolution.x * Resolution.y + NumPreviousNodes))
 		{
 			NumPreviousNodes += Resolution.x * Resolution.y;
 			Resolution = uint2(max(1u, Resolution.x / 2), max(1u, Resolution.y / 2));
 			++LODLevelIndex;
 		}
 
-		const uint LocalNodeIndex = DispatchThreadID.x - NumPreviousNodes;
+		const uint LocalNodeIndex = DispatchThreadId.x - NumPreviousNodes;
 		NodeCoord = uint3(LocalNodeIndex % Resolution.x, LocalNodeIndex / Resolution.x, LODLevelIndex);
 	}
 
@@ -335,7 +345,7 @@ void MainCS(uint3 DispatchThreadID : SV_DispatchThreadID)
 #if PRECISE_OCCLUSION_QUERIES
 			// Increase the InstanceCount by 1.
 			uint Dummy = 0;
-			InterlockedAdd(OcclusionQueryIndirectArgs[ViewIndex].InstanceCount, 1, Dummy);
+			InterlockedAdd(OcclusionQueryIndirectArgs[ViewIndex * 5 + 1], 1, Dummy);
 
 			// Write out bounding box and store a flag indicating if the box crosses the near plane.
 			const float3 NodeAABBMin = float3(NodeAABB2D.xy, WaterZBounds.x);
@@ -422,12 +432,12 @@ uint NumMaterials;
 uint NumDispatchedThreads;
 
 [numthreads(64, 1, 1)]
-void MainCS(uint3 DispatchThreadID : SV_DispatchThreadID)
+void MainCS(uint3 DispatchThreadId : SV_DispatchThreadID)
 {
 	const uint NumDraws = PackedNodes.Load(0);
 
 	// We are limited to 65535 threads, but might end up with more draws than that
-	for (uint DrawIndex = DispatchThreadID.x; DrawIndex < NumDraws; DrawIndex += NumDispatchedThreads)
+	for (uint DrawIndex = DispatchThreadId.x; DrawIndex < NumDraws; DrawIndex += NumDispatchedThreads)
 	{
 #if PRECISE_OCCLUSION_QUERIES
 		const bool bIsVisible = OcclusionResults[DrawIndex] != 0;
@@ -518,11 +528,11 @@ void WriteToBuffer(uint InIndex, uint InWriteOffset, uint InGlobalPrefixSum)
 #else
 [numthreads(1, 1, 1)]
 #endif
-void MainCS(uint3 DispatchThreadID : SV_DispatchThreadID)
+void MainCS(uint3 DispatchThreadId : SV_DispatchThreadID)
 {
 #if PARALLEL_PREFIX_SUM
 
-	const uint ThreadID = DispatchThreadID.x;
+	const uint ThreadID = DispatchThreadId.x;
 	const uint NumBlocks = (NumBuckets + (ARRAY_SIZE - 1)) / ARRAY_SIZE;
 
 	if (ThreadID == 0)
@@ -618,7 +628,7 @@ void MainCS(uint3 DispatchThreadID : SV_DispatchThreadID)
 #define PRECISE_OCCLUSION_QUERIES 0
 #endif
 
-RWStructuredBuffer<FWaterTileDrawIndexedIndirectArgs> IndirectArgs;
+RWBuffer<uint> IndirectArgs;
 RWBuffer<float4> InstanceData0;
 RWBuffer<float4> InstanceData1;
 RWBuffer<float4> InstanceData2;
@@ -647,11 +657,11 @@ uint bWithWaterSelectionSupport;
 uint bLODMorphingEnabled;
 
 [numthreads(64, 1, 1)]
-void MainCS(uint3 DispatchThreadID : SV_DispatchThreadID)
+void MainCS(uint3 DispatchThreadId : SV_DispatchThreadID)
 {
 	const uint NumDraws = PackedNodes.Load(0);
 	// We are limited to 65535 threads, but might end up with more draws than that
-	for (uint DrawIndex = DispatchThreadID.x; DrawIndex < NumDraws; DrawIndex += NumDispatchedThreads)
+	for (uint DrawIndex = DispatchThreadId.x; DrawIndex < NumDraws; DrawIndex += NumDispatchedThreads)
 	{
 #if PRECISE_OCCLUSION_QUERIES
 		const bool bIsVisible = OcclusionResults[DrawIndex] != 0;
@@ -696,7 +706,7 @@ void MainCS(uint3 DispatchThreadID : SV_DispatchThreadID)
 
 		const uint BucketInstanceDataOffset = InstanceDataOffsets.Load(BucketIndex);
 		uint LocalInstanceDataOffset;
-		InterlockedAdd(IndirectArgs[BucketIndex].InstanceCount, 1, LocalInstanceDataOffset);
+		InterlockedAdd(IndirectArgs[BucketIndex * 5 + 1], 1, LocalInstanceDataOffset); // Increment InstanceCount
 		const uint InstanceDataIndex = BucketInstanceDataOffset + LocalInstanceDataOffset;
 		
 		const float Scale = (1u << LODLevel) * LeafSize;
@@ -778,11 +788,11 @@ float LeafSize;
 float CaptureDepthRange;
 
 [numthreads(64, 1, 1)]
-void MainCS(uint3 DispatchThreadID : SV_DispatchThreadID)
+void MainCS(uint3 DispatchThreadId : SV_DispatchThreadID)
 {
 	const uint NumDraws = PackedNodes.Load(0);
 	// We are limited to 65535 threads, but might end up with more draws than that
-	for (uint DrawIndex = DispatchThreadID.x; DrawIndex < NumDraws; DrawIndex += NumDispatchedThreads)
+	for (uint DrawIndex = DispatchThreadId.x; DrawIndex < NumDraws; DrawIndex += NumDispatchedThreads)
 	{
 #if PRECISE_OCCLUSION_QUERIES
 		const bool bIsVisible = OcclusionResults[DrawIndex] != 0;
diff --git a/Engine/Plugins/Experimental/Water/Source/Runtime/Private/WaterQuadTreeGPU.cpp b/Engine/Plugins/Experimental/Water/Source/Runtime/Private/WaterQuadTreeGPU.cpp
--- a/Engine/Plugins/Experimental/Water/Source/Runtime/Private/WaterQuadTreeGPU.cpp
+++ b/Engine/Plugins/Experimental/Water/Source/Runtime/Private/WaterQuadTreeGPU.cpp
@@ -159,6 +159,27 @@ public:
 };
 IMPLEMENT_GLOBAL_SHADER(FWaterQuadTreeInitializeIndirectArgsCS, "/Plugin/Water/Private/WaterQuadTreeDraws.usf", "MainCS", SF_Compute);
 
+class FWaterQuadTreeClearPerViewBuffersCS : public FGlobalShader
+{
+public:
+	DECLARE_GLOBAL_SHADER(FWaterQuadTreeClearPerViewBuffersCS);
+	SHADER_USE_PARAMETER_STRUCT(FWaterQuadTreeClearPerViewBuffersCS, FGlobalShader);
+
+	BEGIN_SHADER_PARAMETER_STRUCT(FParameters, )
+		SHADER_PARAMETER_RDG_BUFFER_UAV(RWByteAddressBuffer, BucketCounts)
+		SHADER_PARAMETER_RDG_BUFFER_UAV(RWByteAddressBuffer, PackedNodes)
+		SHADER_PARAMETER(uint32, NumDrawBuckets)
+	END_SHADER_PARAMETER_STRUCT()
+
+	static void ModifyCompilationEnvironment(const FGlobalShaderPermutationParameters& Parameters, FShaderCompilerEnvironment& OutEnvironment)
+	{
+		FGlobalShader::ModifyCompilationEnvironment(Parameters, OutEnvironment);
+		OutEnvironment.SetDefine(TEXT("CLEAR_PER_VIEW_BUFFERS"), 1);
+	}
+};
+IMPLEMENT_GLOBAL_SHADER(FWaterQuadTreeClearPerViewBuffersCS, "/Plugin/Water/Private/WaterQuadTreeDraws.usf", "MainCS", SF_Compute);
+
+
 class FWaterQuadTreeTraverseCS : public FGlobalShader
 {
 public:
@@ -591,7 +612,7 @@ void FWaterQuadTreeGPU::Traverse(FRDGBuilder& GraphBuilder, const FTraverseParam
 		GraphBuilder.RegisterExternalBuffer(Params.OutInstanceData2Buffer) : 
 		GraphBuilder.CreateBuffer(FRDGBufferDesc::CreateBufferDesc(sizeof(float) * 4, 1), TEXT("WaterQuadTree.InstanceData2Dummy"));
 	
-	FRDGBufferUAV* IndirectArgsBufferUAV = GraphBuilder.CreateUAV(FRDGBufferUAVDesc(IndirectArgsBuffer));
+	FRDGBufferUAV* IndirectArgsBufferUAV = GraphBuilder.CreateUAV(FRDGBufferUAVDesc(IndirectArgsBuffer, PF_R32_UINT));
 	FRDGBufferUAV* InstanceDataOffsetsBufferUAV = GraphBuilder.CreateUAV(FRDGBufferUAVDesc(InstanceDataOffsetsBuffer, PF_R32_UINT));
 	FRDGBufferSRV* InstanceDataOffsetsBufferSRV = GraphBuilder.CreateSRV(FRDGBufferSRVDesc(InstanceDataOffsetsBuffer, PF_R32_UINT));
 	FRDGBufferSRV* WaterBodyRenderDataBufferSRV = GraphBuilder.CreateSRV(FRDGBufferSRVDesc(WaterBodyRenderDataBufferRDG));
@@ -605,7 +626,7 @@ void FWaterQuadTreeGPU::Traverse(FRDGBuilder& GraphBuilder, const FTraverseParam
 	if (bPixelPreciseOcclusionQueries)
 	{
 		OcclusionQueryIndirectArgsBuffer = GraphBuilder.CreateBuffer(FRDGBufferDesc::CreateIndirectDesc<FRHIDrawIndexedIndirectParameters>(Params.NumViews), TEXT("WaterQuadTree.OcclusionQueryIndirectArgs"));
-		OcclusionQueryIndirectArgsBufferUAV = GraphBuilder.CreateUAV(FRDGBufferUAVDesc(OcclusionQueryIndirectArgsBuffer));
+		OcclusionQueryIndirectArgsBufferUAV = GraphBuilder.CreateUAV(FRDGBufferUAVDesc(OcclusionQueryIndirectArgsBuffer, PF_R32_UINT));
 	}
 
 	const uint32 NumBucketsPerView = Params.NumDensities * Params.NumMaterials;
@@ -665,10 +686,17 @@ void FWaterQuadTreeGPU::Traverse(FRDGBuilder& GraphBuilder, const FTraverseParam
 			OcclusionQueryResultsSRV = GraphBuilder.CreateSRV(FRDGBufferSRVDesc(OcclusionQueryResults, PF_R32_UINT));
 		}
 
-		// Clear counts and packed nodes counter
-		// TODO: only clear the counter at index 0, not the entire buffer
-		AddClearUAVPass(GraphBuilder, BucketCountsUAV, 0);
-		AddClearUAVPass(GraphBuilder, PackedNodesUAV, 0);
+		// Clear bucket counts buffer and packed nodes counter
+		{
+			TShaderMapRef<FWaterQuadTreeClearPerViewBuffersCS> ComputeShader(ShaderMap);
+
+			FWaterQuadTreeClearPerViewBuffersCS::FParameters* PassParameters = GraphBuilder.AllocParameters<FWaterQuadTreeClearPerViewBuffersCS::FParameters>();
+			PassParameters->BucketCounts = BucketCountsUAV;
+			PassParameters->PackedNodes = PackedNodesUAV;
+			PassParameters->NumDrawBuckets = NumBucketsPerView;
+
+			FComputeShaderUtils::AddPass(GraphBuilder, RDG_EVENT_NAME("WaterQuadTreeClearPerViewBuffers"), ComputeShader, PassParameters, FComputeShaderUtils::GetGroupCount(NumBucketsPerView, 64));
+		}
 
 		// Traverse quadtree
 		{
