Water: fixed some cases where users could author invalid transforms for water bodies (weird rotations, 0 scale).

diff --git a/Engine/Plugins/Experimental/Water/Source/Runtime/Private/WaterBodyActor.cpp b/Engine/Plugins/Experimental/Water/Source/Runtime/Private/WaterBodyActor.cpp
--- a/Engine/Plugins/Experimental/Water/Source/Runtime/Private/WaterBodyActor.cpp
+++ b/Engine/Plugins/Experimental/Water/Source/Runtime/Private/WaterBodyActor.cpp
@@ -149,6 +149,8 @@ void AWaterBody::PostEditMove(bool bFinished)
 		return;
 	}
 
+	WaterBodyComponent->FixupEditorTransform();
+
 	if (bFinished)
 	{
 		WaterBodyComponent->UpdateWaterHeight();
diff --git a/Engine/Plugins/Experimental/Water/Source/Runtime/Private/WaterBodyComponent.cpp b/Engine/Plugins/Experimental/Water/Source/Runtime/Private/WaterBodyComponent.cpp
--- a/Engine/Plugins/Experimental/Water/Source/Runtime/Private/WaterBodyComponent.cpp
+++ b/Engine/Plugins/Experimental/Water/Source/Runtime/Private/WaterBodyComponent.cpp
@@ -1128,17 +1128,10 @@ void UWaterBodyComponent::OnPostEditChangeProperty(FOnWaterBodyChangedParams& In
 		// Waves data affect the navigation :
 		InOutOnWaterBodyChangedParams.bShapeOrPositionChanged = true;
 	}
-	else if (PropertyChangedEvent.MemberProperty && PropertyChangedEvent.MemberProperty->GetFName() == FName(TEXT("RelativeScale3D")))
+	else if (PropertyChangedEvent.MemberProperty && (PropertyChangedEvent.MemberProperty->GetFName() == FName(TEXT("RelativeScale3D"))
+												|| PropertyChangedEvent.MemberProperty->GetFName() == FName(TEXT("RelativeRotation"))))
 	{
-		// All water bodies which can ever be rendered by the water mesh shouldn't have a z-scale.
-		// Custom meshes also can't have a z scale of 0 or they will render NaN normals into the GBuffer.
-		FVector CurrentScale = GetRelativeScale3D();
-		if (CanEverAffectWaterMesh() || FMath::IsNearlyZero(CurrentScale.Z))
-		{
-			FVector NewScale = CurrentScale;
-			NewScale.Z = 1.f;
-			SetRelativeScale3D(NewScale);
-		}
+		FixupEditorTransform();
 		InOutOnWaterBodyChangedParams.bShapeOrPositionChanged = true;
 	}
 	else if (PropertyChangedEvent.GetMemberPropertyName() == GET_MEMBER_NAME_CHECKED(UWaterBodyComponent, StaticMeshSettings))
@@ -1696,29 +1689,6 @@ void UWaterBodyComponent::OnComponentDestroyed(bool bDestroyingHierarchy)
 	Super::OnComponentDestroyed(bDestroyingHierarchy);
 }
 
-bool UWaterBodyComponent::MoveComponentImpl(const FVector& Delta, const FQuat& NewRotation, bool bSweep, FHitResult* Hit, EMoveComponentFlags MoveFlags, ETeleportType Teleport)
-{
-	FQuat CorrectedRotation = NewRotation;
-
-	// All water bodies which can ever be rendered by the water mesh shouldn't have a z-scale or non-z rotation
-	// Custom meshes also can't have a z scale of 0 or they will render NaN normals into the GBuffer.
-	FVector Scale = GetRelativeScale3D();
-	if (CanEverAffectWaterMesh() || FMath::IsNearlyZero(Scale.Z))
-	{
-		Scale.Z = 1.f;
-		SetRelativeScale3D(Scale);
-	}
-
-	if (CanEverAffectWaterMesh())
-	{
-		// Restrict rotation to the Z-axis only
-		CorrectedRotation.X = 0.f;
-		CorrectedRotation.Y = 0.f;
-	}
-
-	return Super::MoveComponentImpl(Delta, CorrectedRotation, bSweep, Hit, MoveFlags, Teleport);
-}
-
 void UWaterBodyComponent::OnComponentCollisionSettingsChanged(bool bUpdateOverlaps)
 {
 	if (IsRegistered() && !IsTemplate())			// not for CDOs
@@ -2012,6 +1982,37 @@ void UWaterBodyComponent::UpdateWaterBodyStaticMeshComponents()
 	}
 }
 
+void UWaterBodyComponent::FixupEditorTransform()
+{
+	// Water bodies should not have a scale of 0 on any component or they can generate NaN/Infs.
+	// Any water body that is rendered into the water mesh should only have a z scale of 1.
+	FVector CurrentScale = GetRelativeScale3D();
+	if (CanEverAffectWaterMesh() || FMath::IsNearlyZero(CurrentScale.Z))
+	{
+		CurrentScale.Z = 1.;
+	}
+	if (FMath::IsNearlyZero(CurrentScale.X))
+	{
+		CurrentScale.X = 1.;
+	}
+	if (FMath::IsNearlyZero(CurrentScale.Y))
+	{
+		CurrentScale.Y = 1.;
+	}
+	SetRelativeScale3D(CurrentScale);
+	
+	// All water bodies which can ever be rendered by the water mesh should only have yaw rotation.
+	if (CanEverAffectWaterMesh())
+	{
+		FRotator CorrectedRotation = GetRelativeRotation();
+		CorrectedRotation.Pitch = 0.f;
+		CorrectedRotation.Roll = 0.f;
+		SetRelativeRotation(CorrectedRotation);
+	}
+
+}
+
+
 void UWaterBodyComponent::CreateWaterSpriteComponent()
 {
 	UTexture2D* Texture = nullptr;
diff --git a/Engine/Plugins/Experimental/Water/Source/Runtime/Public/WaterBodyComponent.h b/Engine/Plugins/Experimental/Water/Source/Runtime/Public/WaterBodyComponent.h
--- a/Engine/Plugins/Experimental/Water/Source/Runtime/Public/WaterBodyComponent.h
+++ b/Engine/Plugins/Experimental/Water/Source/Runtime/Public/WaterBodyComponent.h
@@ -505,6 +505,9 @@ protected:
 	virtual FVector GetWaterSpriteLocation() const { return GetComponentLocation(); }
 
 	virtual void OnWaterBodyRenderDataUpdated();
+
+	/** Fixup any invalid transformations made to the water body in the editor that may have been made through the transform gizmo or property changes. */
+	void FixupEditorTransform();
 #endif // WITH_EDITOR
 
 	EWaterBodyQueryFlags CheckAndAjustQueryFlags(EWaterBodyQueryFlags InQueryFlags) const;
@@ -525,7 +528,6 @@ protected:
 	virtual void Serialize(FArchive& Ar) override;
 	virtual void PostLoad() override;
 	virtual void OnComponentDestroyed(bool bDestroyingHierarchy) override;
-	virtual bool MoveComponentImpl(const FVector& Delta, const FQuat& NewRotation, bool bSweep, FHitResult* Hit = NULL, EMoveComponentFlags MoveFlags = MOVECOMP_NoFlags, ETeleportType Teleport = ETeleportType::None) override;
 	virtual void OnComponentCollisionSettingsChanged(bool bUpdateOverlaps) override;
 	virtual void OnGenerateOverlapEventsChanged() override;
 	virtual void GetResourceSizeEx(FResourceSizeEx& CumulativeResourceSize) override;
