Added optional jittering to the rasterization of water body meshes when creating the GPU water quadtree. Every draw has multiple instances, each with a different jitter applied. This simulates the effect of multisampling and can be used to work around platform specific limitations/problems with the MSAA path.

#rnx

diff --git a/Engine/Plugins/Experimental/Water/Shaders/Private/WaterQuadTreeVertexShader.usf b/Engine/Plugins/Experimental/Water/Shaders/Private/WaterQuadTreeVertexShader.usf
--- a/Engine/Plugins/Experimental/Water/Shaders/Private/WaterQuadTreeVertexShader.usf
+++ b/Engine/Plugins/Experimental/Water/Shaders/Private/WaterQuadTreeVertexShader.usf
@@ -2,12 +2,29 @@
 
 #include "/Engine/Private/Common.ush"
 
+#ifndef APPLY_JITTER
+#define APPLY_JITTER 0
+#endif
+
 float4x4 Transform;
 
+#if APPLY_JITTER
+float2 JitterScale;
+#endif
+
 void Main(
 	in float3 InPosition : ATTRIBUTE0,
+#if APPLY_JITTER
+	in float2 InJitter : ATTRIBUTE1,
+#endif
 	out float4 OutPosition : SV_Position
 )
 {
+#if APPLY_JITTER
+	float4x4 JitteredTransform = Transform;
+	JitteredTransform[3].xy += InJitter * JitterScale;
+	OutPosition = mul(float4(InPosition, 1.0f), JitteredTransform);
+#else
 	OutPosition = mul(float4(InPosition, 1.0f), Transform);
+#endif
 }
diff --git a/Engine/Plugins/Experimental/Water/Source/Runtime/Private/WaterMeshSceneProxy.cpp b/Engine/Plugins/Experimental/Water/Source/Runtime/Private/WaterMeshSceneProxy.cpp
--- a/Engine/Plugins/Experimental/Water/Source/Runtime/Private/WaterMeshSceneProxy.cpp
+++ b/Engine/Plugins/Experimental/Water/Source/Runtime/Private/WaterMeshSceneProxy.cpp
@@ -39,8 +39,23 @@ static TAutoConsoleVariable<int32> CVarWaterMeshGPUQuadTreeSuperSampling(
 	ECVF_RenderThreadSafe);
 
 static TAutoConsoleVariable<int32> CVarWaterMeshGPUQuadTreeMultiSampling(
-	TEXT("r.Water.WaterMesh.GPUQuadTree.MultiSampling"), 4,
-	TEXT("Rasterizes water meshes into the GPU water quadtree with a multisampled rendertarget, reducing missing water tile artifacts near the edges of water bodies. Default: 4, Min: 1, Max : 8"),
+	TEXT("r.Water.WaterMesh.GPUQuadTree.MultiSampling"), 1,
+	TEXT("Rasterizes water meshes into the GPU water quadtree with a multisampled rendertarget, reducing missing water tile artifacts near the edges of water bodies. Default: 1, Min: 1, Max : 8"),
+	ECVF_RenderThreadSafe);
+
+static TAutoConsoleVariable<int32> CVarWaterMeshGPUQuadTreeNumJitterSamples(
+	TEXT("r.Water.WaterMesh.GPUQuadTree.NumJitterSamples"), 4,
+	TEXT("Rasterizes water meshes into the GPU water quadtree with multiple jittered draw calls, reducing missing water tile artifacts near the edges of water bodies. Default: 4, Min: 1, Max : 16. 1 disables this feature."),
+	ECVF_RenderThreadSafe);
+
+static TAutoConsoleVariable<int32> CVarWaterMeshGPUQuadTreeJitterPattern(
+	TEXT("r.Water.WaterMesh.GPUQuadTree.JitterPattern"), 1,
+	TEXT("Jitter pattern when using multiple jittered draw calls to rasterize water meshes into the GPU water quadtree. 0: Halton, 1: MSAA. Default 1"),
+	ECVF_RenderThreadSafe);
+
+static TAutoConsoleVariable<float> CVarWaterMeshGPUQuadTreeJitterSampleFootprint(
+	TEXT("r.Water.WaterMesh.GPUQuadTree.JitterSampleFootprint"), 1.5f,
+	TEXT("Pixel footprint of the jitter sample pattern. Values greater than 1.0 can cause the water mesh to raster into neighboring pixels not normally covered by the mesh. Default: 1.5, Min 0.0"),
 	ECVF_RenderThreadSafe);
 
 static TAutoConsoleVariable<int32> CVarWaterMeshOcclusionCullingMaxQueries(
@@ -1163,7 +1178,10 @@ void FWaterMeshSceneProxy::BuildGPUQuadTree(FRDGBuilder& GraphBuilder)
 	Params.RequestedQuadTreeResolution = Resolution;
 	Params.SuperSamplingFactor = FMath::Clamp(CVarWaterMeshGPUQuadTreeSuperSampling.GetValueOnRenderThread(), 1, 8);
 	Params.NumMSAASamples = FMath::Clamp(CVarWaterMeshGPUQuadTreeMultiSampling.GetValueOnRenderThread(), 1, 8);
+	Params.NumJitterSamples = FMath::Clamp(CVarWaterMeshGPUQuadTreeNumJitterSamples.GetValueOnRenderThread(), 1, 16);
+	Params.JitterSampleFootprint = FMath::Max(CVarWaterMeshGPUQuadTreeJitterSampleFootprint.GetValueOnRenderThread(), 0.0f);
 	Params.CaptureDepthRange = WaterQuadTreeDepthRange;
+	Params.bUseMSAAJitterPattern = CVarWaterMeshGPUQuadTreeJitterPattern.GetValueOnRenderThread() == 1;
 
 	QuadTreeGPU.Init(GraphBuilder, Params, Draws);
 }
diff --git a/Engine/Plugins/Experimental/Water/Source/Runtime/Private/WaterQuadTreeGPU.cpp b/Engine/Plugins/Experimental/Water/Source/Runtime/Private/WaterQuadTreeGPU.cpp
--- a/Engine/Plugins/Experimental/Water/Source/Runtime/Private/WaterQuadTreeGPU.cpp
+++ b/Engine/Plugins/Experimental/Water/Source/Runtime/Private/WaterQuadTreeGPU.cpp
@@ -14,6 +14,7 @@
 #include "SystemTextures.h"
 #include "ShaderPrintParameters.h"
 #include "SceneRendering.h"
+#include "Math/Halton.h"
 
 DECLARE_GPU_STAT(FWaterQuadTreeGPU_Init);
 DECLARE_GPU_STAT(FWaterQuadTreeGPU_Traverse);
@@ -46,8 +47,12 @@ public:
 	DECLARE_GLOBAL_SHADER(FWaterQuadTreeVS);
 	SHADER_USE_PARAMETER_STRUCT(FWaterQuadTreeVS, FGlobalShader);
 
+	class FApplyJitter : SHADER_PERMUTATION_BOOL("APPLY_JITTER");
+	using FPermutationDomain = TShaderPermutationDomain<FApplyJitter>;
+
 	BEGIN_SHADER_PARAMETER_STRUCT(FParameters, )
 		SHADER_PARAMETER(FMatrix44f, Transform)
+		SHADER_PARAMETER(FVector2f, JitterScale)
 	END_SHADER_PARAMETER_STRUCT()
 
 	static void ModifyCompilationEnvironment(const FGlobalShaderPermutationParameters& Parameters, FShaderCompilerEnvironment& OutEnvironment)
@@ -404,6 +409,83 @@ public:
 IMPLEMENT_GLOBAL_SHADER(FWaterQuadTreeDebugCS, "/Plugin/Water/Private/WaterQuadTreeDraws.usf", "MainCS", SF_Compute);
 
 
+class FJitterOffsetVertexBuffer : public FVertexBuffer
+{
+public:
+	static constexpr int32 NumHaltonSamples = 16;
+
+	virtual void InitRHI(FRHICommandListBase& RHICmdList) override
+	{
+		// Create the texture RHI.
+		FRHIResourceCreateInfo CreateInfo(TEXT("JitterOffsetVertexBuffer"));
+
+		const int32 NumMSAASamples = 2 + 4 + 8 + 16;
+		const int32 NumFloats = 2 * (NumHaltonSamples + NumMSAASamples);
+		VertexBufferRHI = RHICmdList.CreateVertexBuffer(sizeof(float) * NumFloats, BUF_Static | BUF_ShaderResource, CreateInfo);
+
+		float* BufferData = (float*)RHICmdList.LockBuffer(VertexBufferRHI, 0, sizeof(float) * NumFloats, RLM_WriteOnly);
+		
+		// Halton
+		for (int32 i = 0; i < NumHaltonSamples; ++i)
+		{
+			BufferData[i * 2 + 0] = Halton(i + 1, 2) - 0.5f;
+			BufferData[i * 2 + 1] = Halton(i + 1, 3) - 0.5f;
+		}
+		BufferData += NumHaltonSamples * 2;
+
+		// MSAA
+		{
+			const float Offsets[] = 
+			{
+				/*2x*/ 4, 4, -4, -4, 
+				/*4x*/ -2, -6, 6, -2, -6, 2, 2, 6, 
+				/*8x*/ 1, -3, -1, 3, 5, 1, -3, -5, -5, 5, -7, -1, 3, 7, 7, -7, 
+				/*16x*/ 1, 1, -1, -3, -3, 2, 4, -1, -5, -2, 2, 5, 5, 3, 3, -5, 2, 6, 0, -7, -4, -6, -6, 4, -8, 0, 7, -4, 6, 7, -7, -8
+			};
+
+			for (int32 i = 0; i < NumMSAASamples; ++i)
+			{
+				// Remap from (-8) - 7 to (-0.5) - 0.5
+				BufferData[i * 2 + 0] = Offsets[i * 2 + 0] / 16.0f;
+				BufferData[i * 2 + 1] = Offsets[i * 2 + 1] / 16.0f;
+			}
+		}
+
+		RHICmdList.UnlockBuffer(VertexBufferRHI);
+	}
+
+	static constexpr int32 GetMSAADataOffset(int32 NumSamples)
+	{
+		int32 SampleOffset = NumHaltonSamples;
+		SampleOffset += NumSamples > 2 ? 2 : 0;
+		SampleOffset += NumSamples > 4 ? 4 : 0;
+		SampleOffset += NumSamples > 8 ? 8 : 0;
+		return SampleOffset * 2 * sizeof(float);
+	}
+};
+
+TGlobalResource<FJitterOffsetVertexBuffer> GJitterOffsetVertexBuffer;
+
+class FVector3AndInstancedVector2VertexDeclaration : public FRenderResource
+{
+public:
+	FVertexDeclarationRHIRef VertexDeclarationRHI;
+	virtual void InitRHI(FRHICommandListBase& RHICmdList) override
+	{
+		FVertexDeclarationElementList Elements;
+		Elements.Add(FVertexElement(0, 0, VET_Float3, 0, sizeof(FVector3f)));
+		Elements.Add(FVertexElement(1, 0, VET_Float2, 1, sizeof(FVector2f), true /*bInUseInstanceIndex*/));
+		VertexDeclarationRHI = PipelineStateCache::GetOrCreateVertexDeclaration(Elements);
+	}
+	virtual void ReleaseRHI() override
+	{
+		VertexDeclarationRHI.SafeRelease();
+	}
+};
+
+TGlobalResource<FVector3AndInstancedVector2VertexDeclaration> GVector3AndInstancedVector2VertexDeclaration;
+
+
 void FWaterQuadTreeGPU::Init(FRDGBuilder& GraphBuilder, const FInitParams& Params, TArray<FDraw>& Draws)
 {
 	TRACE_CPUPROFILER_EVENT_SCOPE(FWaterQuadTreeGPU::Init);
@@ -442,24 +524,13 @@ void FWaterQuadTreeGPU::Init(FRDGBuilder& GraphBuilder, const FInitParams& Param
 
 	auto GetNumMSAASamples = [](int32 RequestedNumSamples)
 	{
-		if (RequestedNumSamples <= 1)
-		{
-			return 1;
-		}
-		else if (RequestedNumSamples < 4)
-		{
-			return 2;
-		}
-		else if (RequestedNumSamples < 8)
-		{
-			return 4;
-		}
-		else
-		{
-			return 8;
-		}
+		int32 Result = 16;
+		Result = RequestedNumSamples < 16 ? 8 : Result;
+		Result = RequestedNumSamples < 8 ? 4 : Result;
+		Result = RequestedNumSamples < 4 ? 2 : Result;
+		Result = RequestedNumSamples < 2 ? 1 : Result;
+		return Result;
 	};
-	
 
 	FGlobalShaderMap* ShaderMap = GetGlobalShaderMap(GMaxRHIFeatureLevel);
 
@@ -474,17 +545,20 @@ void FWaterQuadTreeGPU::Init(FRDGBuilder& GraphBuilder, const FInitParams& Param
 	// Padding this texture is not required as we will not need to create a mip chain for it; but rather use it to initialize mip0 of our quadtree texture.
 	const FIntPoint RasterResolution = Params.RequestedQuadTreeResolution * Params.SuperSamplingFactor;
 	const int32 NumMipLevels = QuadTreeTextureRDG->Desc.NumMips;
-	const uint8 NumSamples = GetNumMSAASamples(Params.NumMSAASamples);
+	const uint8 NumMSAASamples = FMath::Min(GetNumMSAASamples(Params.NumMSAASamples), 8);
+	const int32 NumJitterSamples = Params.bUseMSAAJitterPattern ? GetNumMSAASamples(Params.NumJitterSamples) : FMath::Clamp(Params.NumJitterSamples, 1, 16);
 
-	FRDGTextureDesc WaterBodyRasterTextureDesc = FRDGTextureDesc::Create2D(RasterResolution, PF_A2B10G10R10, FClearValueBinding::Black, TexCreate_RenderTargetable | TexCreate_ShaderResource, 1, NumSamples);
+	FRDGTextureDesc WaterBodyRasterTextureDesc = FRDGTextureDesc::Create2D(RasterResolution, PF_A2B10G10R10, FClearValueBinding::Black, TexCreate_RenderTargetable | TexCreate_ShaderResource, 1, NumMSAASamples);
 	FRDGTexture* WaterBodyRasterTexture = GraphBuilder.CreateTexture(WaterBodyRasterTextureDesc, TEXT("WaterQuadTree.WaterBodyRasterTexture"));
-	FRDGTextureDesc ZBoundsRasterTextureDesc = FRDGTextureDesc::Create2D(RasterResolution, PF_A2B10G10R10, FClearValueBinding::Black, TexCreate_RenderTargetable | TexCreate_ShaderResource, 1, NumSamples);
+	FRDGTextureDesc ZBoundsRasterTextureDesc = FRDGTextureDesc::Create2D(RasterResolution, PF_A2B10G10R10, FClearValueBinding::Black, TexCreate_RenderTargetable | TexCreate_ShaderResource, 1, NumMSAASamples);
 	FRDGTexture* ZBoundsRasterTexture = GraphBuilder.CreateTexture(ZBoundsRasterTextureDesc, TEXT("WaterQuadTree.ZBoundsRasterTexture"));
 
 
 	// Raster water body meshes
 	{
-		TShaderMapRef<FWaterQuadTreeVS> VertexShader(ShaderMap);
+		FWaterQuadTreeVS::FPermutationDomain VSPermutationDomain;
+		VSPermutationDomain.Set<FWaterQuadTreeVS::FApplyJitter>(NumJitterSamples > 1);
+		TShaderMapRef<FWaterQuadTreeVS> VertexShader(ShaderMap, VSPermutationDomain);
 		TShaderMapRef<FWaterQuadTreePS> PixelShader(ShaderMap);
 
 		FWaterQuadTreeParameters* PassParameters = GraphBuilder.AllocParameters<FWaterQuadTreeParameters>();
@@ -495,7 +569,7 @@ void FWaterQuadTreeGPU::Init(FRDGBuilder& GraphBuilder, const FInitParams& Param
 			RDG_EVENT_NAME("WaterQuadTreeRaster"),
 			PassParameters,
 			ERDGPassFlags::Raster,
-			[PassParameters, RasterResolution, Draws = MoveTemp(Draws), VertexShader, PixelShader](FRHICommandList& RHICmdList)
+			[PassParameters, RasterResolution, NumJitterSamples, bMSAAPattern = Params.bUseMSAAJitterPattern, JitterFootprint = Params.JitterSampleFootprint, Draws = MoveTemp(Draws), VertexShader, PixelShader](FRHICommandList& RHICmdList)
 			{
 				RHICmdList.SetViewport(0.0f, 0.0f, 0.0f, RasterResolution.X, RasterResolution.Y, 1.0f);
 
@@ -507,15 +581,23 @@ void FWaterQuadTreeGPU::Init(FRDGBuilder& GraphBuilder, const FInitParams& Param
 					CW_RGBA, BO_Max, BF_One, BF_One, BO_Max, BF_One, BF_One,
 					CW_RGBA, BO_Max, BF_One, BF_One, BO_Max, BF_One, BF_One>::GetRHI(); // MAX blending
 				GraphicsPSOInit.PrimitiveType = PT_TriangleList;
-				GraphicsPSOInit.BoundShaderState.VertexDeclarationRHI = GetVertexDeclarationFVector3();
+				GraphicsPSOInit.BoundShaderState.VertexDeclarationRHI = NumJitterSamples > 1 ? GVector3AndInstancedVector2VertexDeclaration.VertexDeclarationRHI : GetVertexDeclarationFVector3();
 				GraphicsPSOInit.BoundShaderState.VertexShaderRHI = VertexShader.GetVertexShader();
 				GraphicsPSOInit.BoundShaderState.PixelShaderRHI = PixelShader.GetPixelShader();
 
 				SetGraphicsPipelineState(RHICmdList, GraphicsPSOInit, 0);
 
+				const FVector2f JitterScale = JitterFootprint * FVector2f(2.0f / RasterResolution.X, 2.0f / RasterResolution.Y);
+				if (NumJitterSamples > 1)
+				{
+					const int32 JitterVertexBufferOffset = bMSAAPattern ? FJitterOffsetVertexBuffer::GetMSAADataOffset(NumJitterSamples) : 0;
+					RHICmdList.SetStreamSource(1, GJitterOffsetVertexBuffer.GetRHI(), JitterVertexBufferOffset);
+				}
+
 				for (const FDraw& Draw : Draws)
 				{
 					PassParameters->VS.Transform = Draw.Transform;
+					PassParameters->VS.JitterScale = JitterScale;
 					PassParameters->PS.WaterBodyRenderDataIndex = Draw.WaterBodyRenderDataIndex;
 					PassParameters->PS.Priority = Draw.Priority;
 					PassParameters->PS.WaterBodyMinZ = Draw.MinZ;
@@ -526,7 +608,7 @@ void FWaterQuadTreeGPU::Init(FRDGBuilder& GraphBuilder, const FInitParams& Param
 					SetShaderParameters(RHICmdList, VertexShader, VertexShader.GetVertexShader(), PassParameters->VS);
 					SetShaderParameters(RHICmdList, PixelShader, PixelShader.GetPixelShader(), PassParameters->PS);
 					RHICmdList.SetStreamSource(0, Draw.VertexBuffer, 0);
-					RHICmdList.DrawIndexedPrimitive(Draw.IndexBuffer, Draw.BaseVertexIndex, 0, Draw.NumVertices, Draw.FirstIndex, Draw.NumPrimitives, 1);
+					RHICmdList.DrawIndexedPrimitive(Draw.IndexBuffer, Draw.BaseVertexIndex, 0, Draw.NumVertices, Draw.FirstIndex, Draw.NumPrimitives, NumJitterSamples);
 				}
 			});
 	}
@@ -534,11 +616,11 @@ void FWaterQuadTreeGPU::Init(FRDGBuilder& GraphBuilder, const FInitParams& Param
 	// Merge river and non-river water bodies and downsample to quad tree LOD0 resolution
 	{
 		FWaterQuadTreeMergePS::FPermutationDomain PermutationDomain;
-		PermutationDomain.Set<FWaterQuadTreeMergePS::FNumMSAASamples>(NumSamples);
+		PermutationDomain.Set<FWaterQuadTreeMergePS::FNumMSAASamples>(NumMSAASamples);
 		TShaderMapRef<FWaterQuadTreeMergePS> PixelShader(ShaderMap, PermutationDomain);
 
 		FWaterQuadTreeMergePS::FParameters* PassParameters = GraphBuilder.AllocParameters<FWaterQuadTreeMergePS::FParameters>();
-		if (NumSamples > 1)
+		if (NumMSAASamples > 1)
 		{
 			PassParameters->WaterBodyRasterTextureMS = WaterBodyRasterTexture;
 			PassParameters->ZBoundsRasterTextureMS = ZBoundsRasterTexture;
diff --git a/Engine/Plugins/Experimental/Water/Source/Runtime/Private/WaterQuadTreeGPU.h b/Engine/Plugins/Experimental/Water/Source/Runtime/Private/WaterQuadTreeGPU.h
--- a/Engine/Plugins/Experimental/Water/Source/Runtime/Private/WaterQuadTreeGPU.h
+++ b/Engine/Plugins/Experimental/Water/Source/Runtime/Private/WaterQuadTreeGPU.h
@@ -49,8 +49,11 @@ public:
 		TArrayView<FWaterBodyRenderDataGPU> WaterBodyRenderData;
 		FIntPoint RequestedQuadTreeResolution = FIntPoint::ZeroValue;
 		int32 SuperSamplingFactor = 1;
-		int32 NumMSAASamples = 4;
+		int32 NumMSAASamples = 1;
+		int32 NumJitterSamples = 1;
+		float JitterSampleFootprint = 1.0f;
 		float CaptureDepthRange = 1.0f;
+		bool bUseMSAAJitterPattern = false;
 	};
 
 	struct FTraverseParams
