Deduplicated water quadtree MIDs. All water bodies with the same material will use a single shared MID which is also cached across quadtree rebuilds.

#rnx

diff --git a/Engine/Plugins/Experimental/Water/Source/Runtime/Private/WaterBodyRiverComponent.cpp b/Engine/Plugins/Experimental/Water/Source/Runtime/Private/WaterBodyRiverComponent.cpp
--- a/Engine/Plugins/Experimental/Water/Source/Runtime/Private/WaterBodyRiverComponent.cpp
+++ b/Engine/Plugins/Experimental/Water/Source/Runtime/Private/WaterBodyRiverComponent.cpp
@@ -400,6 +400,16 @@ UMaterialInstanceDynamic* UWaterBodyRiverComponent::GetRiverToOceanTransitionMat
 	return OceanTransitionMID;
 }
 
+UMaterialInterface* UWaterBodyRiverComponent::GetRiverToLakeTransitionMaterial() const
+{
+	return LakeTransitionMaterial;
+}
+
+UMaterialInterface* UWaterBodyRiverComponent::GetRiverToOceanTransitionMaterial() const
+{
+	return OceanTransitionMaterial;
+}
+
 #if WITH_EDITOR
 TArray<UPrimitiveComponent*> UWaterBodyRiverComponent::GetBrushRenderableComponents() const
 {
diff --git a/Engine/Plugins/Experimental/Water/Source/Runtime/Private/WaterMeshComponent.cpp b/Engine/Plugins/Experimental/Water/Source/Runtime/Private/WaterMeshComponent.cpp
--- a/Engine/Plugins/Experimental/Water/Source/Runtime/Private/WaterMeshComponent.cpp
+++ b/Engine/Plugins/Experimental/Water/Source/Runtime/Private/WaterMeshComponent.cpp
@@ -7,6 +7,7 @@
 #include "Materials/Material.h"
 #include "PhysicsEngine/BodySetup.h"
 #include "Engine/StaticMesh.h"
+#include "Engine/TextureRenderTarget2D.h"
 #include "WaterBodyComponent.h"
 #include "WaterMeshSceneProxy.h"
 #include "WaterModule.h"
@@ -62,6 +63,13 @@ TAutoConsoleVariable<int32> CVarWaterMeshEnabled(
 	ECVF_RenderThreadSafe
 );
 
+TAutoConsoleVariable<int32> CVarWaterMeshMIDDeduplication(
+	TEXT("r.Water.WaterMesh.MIDDeduplication"),
+	1,
+	TEXT("Deduplicate per-water body MIDs"),
+	ECVF_RenderThreadSafe
+);
+
 extern TAutoConsoleVariable<float> CVarWaterSplineResampleMaxDistance;
 
 
@@ -321,11 +329,26 @@ void UWaterMeshComponent::RebuildWaterMesh(float InTileSize, const FIntPoint& In
 			UE_LOG(LogWater, Warning, TEXT("WaterZone has more unique water body priorities (%i) than can be supported with GPU driven water quadtree rendering (%i)!"), SortedPriorities.Num(), GPUQuadTreeMaxNumPriorities);
 		}
 	}
+
+	// Lambda for setting parameters on shared/deduplicated MIDs
+	auto SetDynamicParametersOnSharedMID = [&](UMaterialInstanceDynamic* InMID)
+	{
+		if ((InMID == nullptr) || (WaterSubsystem == nullptr))
+		{
+			return false;
+		}
+
+		InMID->SetScalarParameterValue(UWaterBodyComponent::WaterBodyIndexParamName, -1);
+		InMID->SetScalarParameterValue(UWaterBodyComponent::GlobalOceanHeightName, WaterSubsystem->GetOceanTotalHeight());
+		InMID->SetScalarParameterValue(UWaterBodyComponent::WaterZoneIndexParamName, OwningZone->GetWaterZoneIndex());
+		InMID->SetTextureParameterValue(UWaterBodyComponent::WaterVelocityAndHeightName, OwningZone->WaterInfoTexture);
+		return true;
+	};
 	
 
 	bool bAnyWaterMeshesNotReady = false;
 
-	OwningZone->ForEachWaterBodyComponent([this, bIsGPUQuadTree, WaterWorldBox, bIsFlooded, GlobalOceanHeight, OceanFlood, &FarMeshHeight, &bHasOcean, &SortedPriorities, &bAnyWaterMeshesNotReady](UWaterBodyComponent* WaterBodyComponent)
+	OwningZone->ForEachWaterBodyComponent([this, bIsGPUQuadTree, WaterWorldBox, bIsFlooded, GlobalOceanHeight, OceanFlood, &FarMeshHeight, &bHasOcean, &SortedPriorities, &bAnyWaterMeshesNotReady, &SetDynamicParametersOnSharedMID](UWaterBodyComponent* WaterBodyComponent)
 	{
 		check(WaterBodyComponent);
 		AActor* Actor = WaterBodyComponent->GetOwner();
@@ -367,63 +390,144 @@ void UWaterMeshComponent::RebuildWaterMesh(float InTileSize, const FIntPoint& In
 			}
 		}
 
-		// Assign material instance
-		UMaterialInstanceDynamic* WaterMaterial = WaterBodyComponent->GetWaterMaterialInstance();
-		RenderData.Material = WaterMaterial;
+		const bool bDeduplicateMIDs = CVarWaterMeshMIDDeduplication.GetValueOnGameThread() != 0;
 
-		if (RenderData.Material)
+		if (bDeduplicateMIDs)
 		{
-			if (!IsMaterialUsedWithWater(RenderData.Material))
+			// Assign material instance(s)
+			UMaterialInterface* WaterMaterial = WaterBodyComponent->GetWaterMaterial();
+			if (WaterMaterial)
 			{
-				RenderData.Material = UMaterial::GetDefaultMaterial(MD_Surface);
+				const bool bIsRiver = WaterBodyType == EWaterBodyType::River;
+
+				UMaterialInterface* Materials[3] = {};
+				Materials[0] = WaterBodyComponent->GetWaterMaterial();
+				Materials[1] = bIsRiver ? WaterBodyComponent->GetRiverToLakeTransitionMaterial() : nullptr;
+				Materials[2] = bIsRiver ? WaterBodyComponent->GetRiverToOceanTransitionMaterial() : nullptr;
+
+				FName Names[] = { TEXT("WaterMID"), TEXT("LakeTransitionMID"), TEXT("OceanTransitionMID") };
+
+				UMaterialInterface* OutMaterials[3] = {};
+				for (int i = 0; i < 3; ++i)
+				{
+					UMaterialInterface* Material = Materials[i];
+					if (Material)
+					{
+						if (IsMaterialUsedWithWater(Material))
+						{
+							const bool bHasExistingMID = MaterialToMID.Contains(Material);
+							UMaterialInstanceDynamic* MID = nullptr;
+							if (!bHasExistingMID)
+							{
+								MID = FWaterUtils::GetOrCreateTransientMID(nullptr, Names[i], Material, RF_Transient | RF_NonPIEDuplicateTransient | RF_TextExportTransient);
+								SetDynamicParametersOnSharedMID(MID);
+								MaterialToMID.Add(Material, MID);
+							}
+							else
+							{
+								MID = MaterialToMID[Material];
+								// Update the parameters on this MID if this is the first time this material is seen in this rebuild.
+								// This is necessary to handle cases where GlobalOceanHeight or the water info texture pointer change
+								// TODO: GlobalOceanHeight can probably be removed, leaving only the texture parameter.
+								// Can we then skip this and set the parameters only on creation?
+								if (!UsedMaterials.Contains(MID))
+								{
+									SetDynamicParametersOnSharedMID(MID);
+								}
+							}
+							OutMaterials[i] = MID;
+						}
+						else
+						{
+							OutMaterials[i] = UMaterial::GetDefaultMaterial(MD_Surface);
+						}
+						UsedMaterials.Add(OutMaterials[i]);
+					}
+				}
+
+				RenderData.Material = OutMaterials[0];
+				RenderData.RiverToLakeMaterial = OutMaterials[1];
+				RenderData.RiverToOceanMaterial = OutMaterials[2];
 			}
-			else
+
+			RenderData.Priority = static_cast<int16>(FMath::Clamp(WaterBodyComponent->GetOverlapMaterialPriority(), MinWaterBodyPriority, MaxWaterBodyPriority));
+			RenderData.WaterBodyIndex = static_cast<int16>(WaterBodyComponent->GetWaterBodyIndex());
+			RenderData.SurfaceBaseHeight = WaterBodyComponent->GetComponentLocation().Z;
+			RenderData.MaxWaveHeight = WaterBodyComponent->GetMaxWaveHeight();
+			RenderData.BoundsMinZ = WaterBodyComponent->Bounds.GetBox().Min.Z;
+			RenderData.BoundsMaxZ = WaterBodyComponent->Bounds.GetBox().Max.Z;
+			RenderData.WaterBodyType = static_cast<int8>(WaterBodyType);
+#if WITH_WATER_SELECTION_SUPPORT
+			RenderData.HitProxy = new HActor(/*InActor = */Actor, /*InPrimComponent = */nullptr);
+			RenderData.bWaterBodySelected = Actor->IsSelected();
+#endif // WITH_WATER_SELECTION_SUPPORT
+
+			if (WaterBodyType == EWaterBodyType::Ocean && bIsFlooded)
 			{
-				// Add ocean height as a scalar parameter
-				WaterBodyComponent->SetDynamicParametersOnMID(WaterMaterial);
+				RenderData.SurfaceBaseHeight += OceanFlood;
+				RenderData.Priority -= 1;
 			}
-
-			// Add material so that the component keeps track of all potential materials used
-			UsedMaterials.Add(RenderData.Material);
 		}
+		else
+		{
+			// Assign material instance
+			UMaterialInstanceDynamic* WaterMaterial = WaterBodyComponent->GetWaterMaterialInstance();
+			RenderData.Material = WaterMaterial;
+
+			if (RenderData.Material)
+			{
+				if (!IsMaterialUsedWithWater(RenderData.Material))
+				{
+					RenderData.Material = UMaterial::GetDefaultMaterial(MD_Surface);
+				}
+				else
+				{
+					// Add ocean height as a scalar parameter
+					WaterBodyComponent->SetDynamicParametersOnMID(WaterMaterial);
+				}
 
-		RenderData.Priority = static_cast<int16>(FMath::Clamp(WaterBodyComponent->GetOverlapMaterialPriority(), MinWaterBodyPriority, MaxWaterBodyPriority));
-		RenderData.WaterBodyIndex = static_cast<int16>(WaterBodyComponent->GetWaterBodyIndex());
-		RenderData.SurfaceBaseHeight = WaterBodyComponent->GetComponentLocation().Z;
-		RenderData.MaxWaveHeight = WaterBodyComponent->GetMaxWaveHeight();
-		RenderData.BoundsMinZ = WaterBodyComponent->Bounds.GetBox().Min.Z;
-		RenderData.BoundsMaxZ = WaterBodyComponent->Bounds.GetBox().Max.Z;
-		RenderData.WaterBodyType = static_cast<int8>(WaterBodyType);
+				// Add material so that the component keeps track of all potential materials used
+				UsedMaterials.Add(RenderData.Material);
+			}
+
+			RenderData.Priority = static_cast<int16>(FMath::Clamp(WaterBodyComponent->GetOverlapMaterialPriority(), MinWaterBodyPriority, MaxWaterBodyPriority));
+			RenderData.WaterBodyIndex = static_cast<int16>(WaterBodyComponent->GetWaterBodyIndex());
+			RenderData.SurfaceBaseHeight = WaterBodyComponent->GetComponentLocation().Z;
+			RenderData.MaxWaveHeight = WaterBodyComponent->GetMaxWaveHeight();
+			RenderData.BoundsMinZ = WaterBodyComponent->Bounds.GetBox().Min.Z;
+			RenderData.BoundsMaxZ = WaterBodyComponent->Bounds.GetBox().Max.Z;
+			RenderData.WaterBodyType = static_cast<int8>(WaterBodyType);
 #if WITH_WATER_SELECTION_SUPPORT
-		RenderData.HitProxy = new HActor(/*InActor = */Actor, /*InPrimComponent = */nullptr);
-		RenderData.bWaterBodySelected = Actor->IsSelected();
+			RenderData.HitProxy = new HActor(/*InActor = */Actor, /*InPrimComponent = */nullptr);
+			RenderData.bWaterBodySelected = Actor->IsSelected();
 #endif // WITH_WATER_SELECTION_SUPPORT
 
-		if (WaterBodyType == EWaterBodyType::Ocean && bIsFlooded)
-		{
-			RenderData.SurfaceBaseHeight += OceanFlood;
-			RenderData.Priority -= 1;
-		}
-
-		// For rivers, set up transition materials if they exist
-		if (WaterBodyType == EWaterBodyType::River)
-		{
-			UMaterialInstanceDynamic* RiverToLakeMaterial = WaterBodyComponent->GetRiverToLakeTransitionMaterialInstance();
-			if (IsMaterialUsedWithWater(RiverToLakeMaterial))
+			if (WaterBodyType == EWaterBodyType::Ocean && bIsFlooded)
 			{
-				RenderData.RiverToLakeMaterial = RiverToLakeMaterial;
-				UsedMaterials.Add(RenderData.RiverToLakeMaterial);
-				// Add ocean height as a scalar parameter
-				WaterBodyComponent->SetDynamicParametersOnMID(RiverToLakeMaterial);
+				RenderData.SurfaceBaseHeight += OceanFlood;
+				RenderData.Priority -= 1;
 			}
-			
-			UMaterialInstanceDynamic* RiverToOceanMaterial = WaterBodyComponent->GetRiverToOceanTransitionMaterialInstance();
-			if (IsMaterialUsedWithWater(RiverToOceanMaterial))
+
+			// For rivers, set up transition materials if they exist
+			if (WaterBodyType == EWaterBodyType::River)
 			{
-				RenderData.RiverToOceanMaterial = RiverToOceanMaterial;
-				UsedMaterials.Add(RenderData.RiverToOceanMaterial);
-				// Add ocean height as a scalar parameter
-				WaterBodyComponent->SetDynamicParametersOnMID(RiverToOceanMaterial);
+				UMaterialInstanceDynamic* RiverToLakeMaterial = WaterBodyComponent->GetRiverToLakeTransitionMaterialInstance();
+				if (IsMaterialUsedWithWater(RiverToLakeMaterial))
+				{
+					RenderData.RiverToLakeMaterial = RiverToLakeMaterial;
+					UsedMaterials.Add(RenderData.RiverToLakeMaterial);
+					// Add ocean height as a scalar parameter
+					WaterBodyComponent->SetDynamicParametersOnMID(RiverToLakeMaterial);
+				}
+
+				UMaterialInstanceDynamic* RiverToOceanMaterial = WaterBodyComponent->GetRiverToOceanTransitionMaterialInstance();
+				if (IsMaterialUsedWithWater(RiverToOceanMaterial))
+				{
+					RenderData.RiverToOceanMaterial = RiverToOceanMaterial;
+					UsedMaterials.Add(RenderData.RiverToOceanMaterial);
+					// Add ocean height as a scalar parameter
+					WaterBodyComponent->SetDynamicParametersOnMID(RiverToOceanMaterial);
+				}
 			}
 		}
 
@@ -688,6 +792,9 @@ void UWaterMeshComponent::RebuildWaterMesh(float InTileSize, const FIntPoint& In
 		WaterQuadTree.AddFarMesh(FarDistanceMaterial, FarMeshBounds, FarDistanceMeshExtent, FarMeshHeight);
 	}
 
+	// Remove all materials from the MaterialToMID map that aren't currently used
+	MaterialToMID = MaterialToMID.FilterByPredicate([this](const auto& Pair) { return UsedMaterials.Contains(Pair.Value); });
+
 	WaterQuadTree.Unlock(true);
 
 	MarkRenderStateDirty();
diff --git a/Engine/Plugins/Experimental/Water/Source/Runtime/Public/WaterBodyComponent.h b/Engine/Plugins/Experimental/Water/Source/Runtime/Public/WaterBodyComponent.h
--- a/Engine/Plugins/Experimental/Water/Source/Runtime/Public/WaterBodyComponent.h
+++ b/Engine/Plugins/Experimental/Water/Source/Runtime/Public/WaterBodyComponent.h
@@ -218,6 +218,14 @@ public:
 	UFUNCTION(BlueprintCallable, Category = Rendering)
 	UMaterialInterface* GetWaterMaterial() const { return WaterMaterial; }
 
+	/** Returns river to lake transition water material */
+	UFUNCTION(BlueprintCallable, Category = Rendering)
+	virtual UMaterialInterface* GetRiverToLakeTransitionMaterial() const { return nullptr; }
+
+	/** Returns river to ocean transition water material */
+	UFUNCTION(BlueprintCallable, Category = Rendering)
+	virtual UMaterialInterface* GetRiverToOceanTransitionMaterial() const { return nullptr; }
+
 	/** Returns material used to render the water as a static mesh */
 	UMaterialInterface* GetWaterStaticMeshMaterial() const { return WaterStaticMeshMaterial; }
 
diff --git a/Engine/Plugins/Experimental/Water/Source/Runtime/Public/WaterBodyRiverComponent.h b/Engine/Plugins/Experimental/Water/Source/Runtime/Public/WaterBodyRiverComponent.h
--- a/Engine/Plugins/Experimental/Water/Source/Runtime/Public/WaterBodyRiverComponent.h
+++ b/Engine/Plugins/Experimental/Water/Source/Runtime/Public/WaterBodyRiverComponent.h
@@ -22,6 +22,8 @@ public:
 	virtual TArray<UPrimitiveComponent*> GetStandardRenderableComponents() const override;
 	virtual UMaterialInstanceDynamic* GetRiverToLakeTransitionMaterialInstance() override;
 	virtual UMaterialInstanceDynamic* GetRiverToOceanTransitionMaterialInstance() override;
+	virtual UMaterialInterface* GetRiverToLakeTransitionMaterial() const override;
+	virtual UMaterialInterface* GetRiverToOceanTransitionMaterial() const override;
 
 #if WITH_EDITOR
 	virtual TArray<UPrimitiveComponent*> GetBrushRenderableComponents() const override;
diff --git a/Engine/Plugins/Experimental/Water/Source/Runtime/Public/WaterMeshComponent.h b/Engine/Plugins/Experimental/Water/Source/Runtime/Public/WaterMeshComponent.h
--- a/Engine/Plugins/Experimental/Water/Source/Runtime/Public/WaterMeshComponent.h
+++ b/Engine/Plugins/Experimental/Water/Source/Runtime/Public/WaterMeshComponent.h
@@ -113,6 +113,10 @@ private:
 	UPROPERTY(Transient, NonPIEDuplicateTransient, TextExportTransient)
 	TSet<TObjectPtr<UMaterialInterface>> UsedMaterials;
 
+	/** Maps from materials assigned to each water body to actually used MIDs. Persists across rebuilds in order to cache MIDs */
+	UPROPERTY(Transient, NonPIEDuplicateTransient, TextExportTransient)
+	TMap<TObjectPtr<UMaterialInterface>, TObjectPtr<UMaterialInstanceDynamic>> MaterialToMID;
+
 	/** Forces the water mesh to always render the far mesh, regardless if there is an ocean or not.*/
 	UPROPERTY(Category = "Rendering|FarDistance", EditAnywhere)
 	bool bUseFarMeshWithoutOcean = false;
