Implemented optional software conservative rasterization for the GPU water quadtree. When building the quadtree by rasterizing the meshes, actual (software) CR can now be used instead of approximating CR with supersampling/multisampling/multiple jittered draws. The new CR path dilates each triangle to conservatively cover all pixel centers of pixels touched by the input triangle. It needs a special water mesh which stores the positions of the other two vertices of each triangle. This additional data is stored in the UV channels of the existing water info mesh. Even though this makes each vertex unique and therefore increases the total vertex count, testing didn't show any performance difference when rasterizing the water info meshes to the water info texture. Whether the additional data is built into the water info mesh can be controlled with the "r.Water.BuildConservativeRasterizationMesh" read-only cvar. This approach saves memory compared to introducing another mesh specifically for conservative rasterization.

#rnx

diff --git a/Engine/Plugins/Experimental/Water/Shaders/Private/WaterQuadTreePixelShader.usf b/Engine/Plugins/Experimental/Water/Shaders/Private/WaterQuadTreePixelShader.usf
--- a/Engine/Plugins/Experimental/Water/Shaders/Private/WaterQuadTreePixelShader.usf
+++ b/Engine/Plugins/Experimental/Water/Shaders/Private/WaterQuadTreePixelShader.usf
@@ -2,6 +2,10 @@
 
 #include "/Engine/Private/Common.ush"
 
+#ifndef CLIP_CONSERVATIVE_TRIANGLE
+#define CLIP_CONSERVATIVE_TRIANGLE 0
+#endif
+
 uint WaterBodyRenderDataIndex;
 uint Priority;
 float WaterBodyMinZ;
@@ -10,10 +14,22 @@ float MaxWaveHeight;
 uint bIsRiver;
 
 void Main(
-	in float4 InPosition : SV_Position, 
+	in float4 InPosition : SV_Position,
+#if CLIP_CONSERVATIVE_TRIANGLE
+	in float4 InClipPosition : CLIP_POSITION,
+	in float4 InTriangleAABB : TRIANGLE_AABB,
+#endif
 	out float4 OutWaterBody : SV_Target0,
 	out float4 OutZBounds : SV_Target1)
 {
+#if CLIP_CONSERVATIVE_TRIANGLE
+	const float2 Pos = InClipPosition.xy / InClipPosition.w;
+	if (any(Pos < InTriangleAABB.xy) || any(Pos > InTriangleAABB.zw))
+	{
+		discard;
+	}
+#endif
+
 	uint PackedPriorityAndWaterBodyRenderDataIndex = 0;
 	PackedPriorityAndWaterBodyRenderDataIndex |= WaterBodyRenderDataIndex & 0x7Fu; // 7bits
 	PackedPriorityAndWaterBodyRenderDataIndex |= min(Priority, 7u) << 7u; // 3bits
diff --git a/Engine/Plugins/Experimental/Water/Shaders/Private/WaterQuadTreeVertexShader.usf b/Engine/Plugins/Experimental/Water/Shaders/Private/WaterQuadTreeVertexShader.usf
--- a/Engine/Plugins/Experimental/Water/Shaders/Private/WaterQuadTreeVertexShader.usf
+++ b/Engine/Plugins/Experimental/Water/Shaders/Private/WaterQuadTreeVertexShader.usf
@@ -2,29 +2,100 @@
 
 #include "/Engine/Private/Common.ush"
 
-#ifndef APPLY_JITTER
-#define APPLY_JITTER 0
+#define RASTER_MODE_REGULAR 0
+#define RASTER_MODE_JITTER 1
+#define RASTER_MODE_CONSERVATIVE 2
+
+#ifndef RASTER_MODE
+#define RASTER_MODE RASTER_MODE_REGULAR
 #endif
 
 float4x4 Transform;
 
-#if APPLY_JITTER
+#if RASTER_MODE == RASTER_MODE_JITTER
 float2 JitterScale;
+#elif RASTER_MODE == RASTER_MODE_CONSERVATIVE
+float2 HalfPixelSize;
 #endif
 
+float4 AddPoint(float4 InPos, float4 InAABB)
+{
+	if (InPos.w < 0.0f)
+	{
+		InAABB.xy = min(InAABB.xy, InPos.xy);
+		InAABB.zw = max(InAABB.zw, InPos.xy);
+	}
+	else
+	{
+		InAABB.xy = min(InAABB.xy, InPos.xy / InPos.w);
+		InAABB.zw = max(InAABB.zw, InPos.xy / InPos.w);
+	}
+	return InAABB;
+}
+
 void Main(
 	in float3 InPosition : ATTRIBUTE0,
-#if APPLY_JITTER
+#if RASTER_MODE == RASTER_MODE_JITTER
 	in float2 InJitter : ATTRIBUTE1,
+#elif RASTER_MODE == RASTER_MODE_CONSERVATIVE
+	in float2 InPrevNextVertex0 : ATTRIBUTE1,
+	in float2 InPrevNextVertex1 : ATTRIBUTE2,
+	in float2 InPrevNextVertex2 : ATTRIBUTE3,
 #endif
 	out float4 OutPosition : SV_Position
+#if RASTER_MODE == RASTER_MODE_CONSERVATIVE
+	,out float4 OutClipPosition : CLIP_POSITION
+	,out float4 OutTriangleAABB : TRIANGLE_AABB
+#endif
 )
 {
-#if APPLY_JITTER
+#if RASTER_MODE == RASTER_MODE_REGULAR
+	
+	OutPosition = mul(float4(InPosition, 1.0f), Transform);
+
+#elif RASTER_MODE == RASTER_MODE_JITTER
+	
 	float4x4 JitteredTransform = Transform;
 	JitteredTransform[3].xy += InJitter * JitterScale;
 	OutPosition = mul(float4(InPosition, 1.0f), JitteredTransform);
-#else
-	OutPosition = mul(float4(InPosition, 1.0f), Transform);
+
+#elif RASTER_MODE == RASTER_MODE_CONSERVATIVE
+
+	// See "GPU Gems 2 Chapter 42. Conservative Rasterization".
+	// We construct a slightly bigger triangle by displacing the edges by HalfPixelSize. 
+	// This new triangle should then conservatively cover all pixels touched by the original triangle.
+	// We pass a triangle AABB to the pixel shader to discard false positive pixels.
+
+	const float3 PrevVertexPos = float3(InPrevNextVertex0.x, InPrevNextVertex0.y, InPrevNextVertex1.x);
+	const float3 NextVertexPos = float3(InPrevNextVertex1.y, InPrevNextVertex2.x, InPrevNextVertex2.y);
+
+	const float4 PrevPosition = mul(float4(PrevVertexPos, 1.0f), Transform);
+	const float4 CurrPosition = mul(float4(InPosition, 1.0f), Transform);
+	const float4 NextPosition = mul(float4(NextVertexPos, 1.0f), Transform);
+
+	// Compute the triangle AABB
+	OutTriangleAABB = float4(9999999.0f, 9999999.0f, -9999999.0f, -9999999.0f);
+	OutTriangleAABB = AddPoint(PrevPosition, OutTriangleAABB);
+	OutTriangleAABB = AddPoint(CurrPosition, OutTriangleAABB);
+	OutTriangleAABB = AddPoint(NextPosition, OutTriangleAABB);
+	OutTriangleAABB += float4(-HalfPixelSize, HalfPixelSize);
+
+	// Compute equations of the planes through the two edges  
+	float3 Planes[2];
+	Planes[0] = cross(CurrPosition.xyw - PrevPosition.xyw, PrevPosition.xyw);
+	Planes[1] = cross(NextPosition.xyw - CurrPosition.xyw, CurrPosition.xyw);
+
+	// Move the planes by the appropriate semidiagonal  
+	Planes[0].z -= dot(HalfPixelSize, abs(Planes[0].xy));
+	Planes[1].z -= dot(HalfPixelSize, abs(Planes[1].xy));
+
+	// Compute the intersection point of the planes.
+	OutPosition = CurrPosition;
+	OutPosition.xyw = cross(Planes[0], Planes[1]);
+	OutPosition /= abs(OutPosition.w);
+	OutPosition.z = CurrPosition.z / CurrPosition.w; // Keep the depth of the non-dilated vertex. Slightly incorrect, but seems to work.
+
+	OutClipPosition = OutPosition;
+
 #endif
 }
diff --git a/Engine/Plugins/Experimental/Water/Source/Runtime/Private/WaterBodyComponent.cpp b/Engine/Plugins/Experimental/Water/Source/Runtime/Private/WaterBodyComponent.cpp
--- a/Engine/Plugins/Experimental/Water/Source/Runtime/Private/WaterBodyComponent.cpp
+++ b/Engine/Plugins/Experimental/Water/Source/Runtime/Private/WaterBodyComponent.cpp
@@ -66,6 +66,11 @@ TAutoConsoleVariable<float> CVarWaterOceanFallbackDepth(
 	TEXT("Depth to report for the ocean when no terrain is found under the query location. Not used when <= 0."),
 	ECVF_Default);
 
+TAutoConsoleVariable<int32> CVarWaterBodyBuildConservativeRasterizationMesh(
+	TEXT("r.Water.BuildConservativeRasterizationMesh"), 1,
+	TEXT("Enables additional data in the UV channels of the water mesh, which is used for software conservative rasterization when creating the GPU water quadtree."),
+	ECVF_ReadOnly);
+
 const FName UWaterBodyComponent::WaterBodyIndexParamName(TEXT("WaterBodyIndex"));
 const FName UWaterBodyComponent::WaterZoneIndexParamName(TEXT("WaterZoneIndex"));
 const FName UWaterBodyComponent::WaterBodyZOffsetParamName(TEXT("WaterBodyZOffset"));
@@ -1486,7 +1491,10 @@ void UWaterBodyComponent::OnPostRegisterAllComponents()
 		}
 	}
 
-	if (GetLinkerCustomVersion(FFortniteMainBranchObjectVersion::GUID) < FFortniteMainBranchObjectVersion::WaterBodyStaticMeshFixup)
+	UWaterBodyInfoMeshComponent* WaterInfoMeshComponent = GetWaterInfoMeshComponent();
+	const bool bHasConservativeRasterMesh = IsValid(WaterInfoMeshComponent) && WaterInfoMeshComponent->bIsConservativeRasterCompatible;
+	const bool bShouldHaveConservativeRastermesh = CVarWaterBodyBuildConservativeRasterizationMesh.GetValueOnGameThread() != 0;
+	if ((bHasConservativeRasterMesh != bShouldHaveConservativeRastermesh) || GetLinkerCustomVersion(FFortniteMainBranchObjectVersion::GUID) < FFortniteMainBranchObjectVersion::WaterBodyStaticMeshFixup)
 	{
 		UpdateWaterBodyRenderData();
 	}
@@ -1942,8 +1950,9 @@ void UWaterBodyComponent::UpdateWaterInfoMeshComponents()
 		return;
 	}
 
+	const bool bShouldBuildConservativeRasterMesh = CVarWaterBodyBuildConservativeRasterizationMesh.GetValueOnGameThread() != 0;
 	FWaterBodyMeshBuilder MeshBuilder;
-	MeshBuilder.BuildWaterInfoMeshes( this, GetWaterInfoMeshComponent(), GetDilatedWaterInfoMeshComponent());
+	MeshBuilder.BuildWaterInfoMeshes( this, GetWaterInfoMeshComponent(), GetDilatedWaterInfoMeshComponent(), bShouldBuildConservativeRasterMesh);
 }
 
 void UWaterBodyComponent::UpdateWaterBodyStaticMeshComponents()
diff --git a/Engine/Plugins/Experimental/Water/Source/Runtime/Private/WaterBodyInfoMeshComponent.h b/Engine/Plugins/Experimental/Water/Source/Runtime/Private/WaterBodyInfoMeshComponent.h
--- a/Engine/Plugins/Experimental/Water/Source/Runtime/Private/WaterBodyInfoMeshComponent.h
+++ b/Engine/Plugins/Experimental/Water/Source/Runtime/Private/WaterBodyInfoMeshComponent.h
@@ -29,6 +29,11 @@ public:
 	UPROPERTY()
 	bool bIsDilatedMesh = false;
 
+	// Whether the water body mesh is compatible with software conservative rasterization used for generating the GPU water quadtree.
+	// This requires three UV channels with 32bit floats storing the vertex position of the previous and next vertex in the triangle.
+	UPROPERTY()
+	bool bIsConservativeRasterCompatible = false;
+
 protected:
 	// Even though we may be still precaching PSOs, we don't want to delay render proxy creation
 	// or fallback to the default material. We're always returning false here so that we decline
diff --git a/Engine/Plugins/Experimental/Water/Source/Runtime/Private/WaterBodyMeshBuilder.cpp b/Engine/Plugins/Experimental/Water/Source/Runtime/Private/WaterBodyMeshBuilder.cpp
--- a/Engine/Plugins/Experimental/Water/Source/Runtime/Private/WaterBodyMeshBuilder.cpp
+++ b/Engine/Plugins/Experimental/Water/Source/Runtime/Private/WaterBodyMeshBuilder.cpp
@@ -32,20 +32,20 @@
 #endif // WITH_EDITOR
 
 #if WITH_EDITOR
-void FWaterBodyMeshBuilder::BuildWaterInfoMeshes(UWaterBodyComponent* WaterBodyComponent, UWaterBodyInfoMeshComponent* WaterInfoMeshComponent, UWaterBodyInfoMeshComponent* WaterInfoDilatedMeshComponent) const
+void FWaterBodyMeshBuilder::BuildWaterInfoMeshes(UWaterBodyComponent* WaterBodyComponent, UWaterBodyInfoMeshComponent* WaterInfoMeshComponent, UWaterBodyInfoMeshComponent* WaterInfoDilatedMeshComponent, bool bMakeWaterInfoMeshConservativeRasterCompatible) const
 {
 	TRACE_CPUPROFILER_EVENT_SCOPE(FWaterBodyMeshBuilder::BuildWaterInfoMesh);
 	
 	using namespace UE::Geometry;
 
-	FDynamicMesh3 WaterBodyMesh(EMeshComponents::None);
-	FDynamicMesh3 WaterBodyDilatedMesh(EMeshComponents::None);
-	GetDynamicMesh(WaterBodyComponent, WaterBodyMesh, &WaterBodyDilatedMesh);
+	FDynamicMesh3 WaterInfoMesh(EMeshComponents::None);
+	FDynamicMesh3 WaterInfoDilatedMesh(EMeshComponents::None);
+	GetDynamicMesh(WaterBodyComponent, WaterInfoMesh, &WaterInfoDilatedMesh, bMakeWaterInfoMeshConservativeRasterCompatible);
 
 	UMaterialInterface* WaterInfoMID = WaterBodyComponent->GetWaterInfoMaterialInstance();
 	UObject* Outer = WaterBodyComponent->GetOwner();
 
-	auto BuildWaterInfoMesh = [this, Outer, WaterInfoMID](const FDynamicMesh3& DynamicMesh, UWaterBodyInfoMeshComponent* MeshComponent, FName BaseName) -> UStaticMesh*
+	auto BuildWaterInfoMesh = [this, Outer, WaterInfoMID](const FDynamicMesh3& DynamicMesh, UWaterBodyInfoMeshComponent* MeshComponent, FName BaseName, bool bIsConservativeRasterMesh) -> UStaticMesh*
 	{
 		if (DynamicMesh.TriangleCount() == 0)
 		{
@@ -55,18 +55,20 @@ void FWaterBodyMeshBuilder::BuildWaterInfoMeshes(UWaterBodyComponent* WaterBodyC
 		FName WaterInfoMeshName = MakeUniqueObjectName(Outer, UStaticMesh::StaticClass(), BaseName);
 		UStaticMesh* StaticMesh = CreateUStaticMesh(Outer, WaterInfoMeshName);
 
-		UpdateStaticMesh(StaticMesh, ConvertDynamicMeshToMeshDescription(DynamicMesh));
+		UpdateStaticMesh(StaticMesh, ConvertDynamicMeshToMeshDescription(DynamicMesh), bIsConservativeRasterMesh);
 
 		StaticMesh->AddMaterial(WaterInfoMID);
 
 		MeshComponent->SetStaticMesh(StaticMesh);
+		MeshComponent->bIsConservativeRasterCompatible = bIsConservativeRasterMesh;
 
 		return StaticMesh;
 	};
 
-	TArray<UStaticMesh*> StaticMeshes = {
-		BuildWaterInfoMesh(WaterBodyMesh, WaterInfoMeshComponent, TEXT("WaterInfoMesh")),
-		BuildWaterInfoMesh(WaterBodyDilatedMesh, WaterInfoDilatedMeshComponent, TEXT("WaterInfoDilatedMesh")),
+	TArray<UStaticMesh*> StaticMeshes = 
+	{
+		BuildWaterInfoMesh(WaterInfoMesh, WaterInfoMeshComponent, TEXT("WaterInfoMesh"), bMakeWaterInfoMeshConservativeRasterCompatible),
+		BuildWaterInfoMesh(WaterInfoDilatedMesh, WaterInfoDilatedMeshComponent, TEXT("WaterInfoDilatedMesh"), false /*bIsConservativeRasterMesh*/),
 	};
 
 	UStaticMesh::BatchBuild(StaticMeshes);
@@ -400,7 +402,7 @@ TArray<TObjectPtr<UWaterBodyStaticMeshComponent>> FWaterBodyMeshBuilder::BuildWa
 
 			StaticMesh->AddMaterial(WaterStaticMeshMID);
 
-			UpdateStaticMesh(StaticMesh, MeshDescription);
+			UpdateStaticMesh(StaticMesh, MeshDescription, false /*bIsConservativeRasterMesh*/);
 
 			WaterBodyStaticMeshComponent->SetStaticMesh(StaticMesh);
 
@@ -471,7 +473,7 @@ TArray<TObjectPtr<UWaterBodyStaticMeshComponent>> FWaterBodyMeshBuilder::BuildWa
 
 		StaticMeshes.Add(QuadStaticMesh);
 
-		UpdateStaticMesh(QuadStaticMesh, QuadMeshDescription);
+		UpdateStaticMesh(QuadStaticMesh, QuadMeshDescription, false /*bIsConservativeRasterMesh*/);
 
 		for (const FVector& QuadMeshPosition : QuadSectionPositions)
 		{
@@ -492,7 +494,7 @@ TArray<TObjectPtr<UWaterBodyStaticMeshComponent>> FWaterBodyMeshBuilder::BuildWa
 	return StaticMeshComponents;
 }
 
-void FWaterBodyMeshBuilder::GetDynamicMesh(const UWaterBodyComponent* WaterBodyComponent, UE::Geometry::FDynamicMesh3& InMesh, UE::Geometry::FDynamicMesh3* InDilatedMesh) const
+void FWaterBodyMeshBuilder::GetDynamicMesh(const UWaterBodyComponent* WaterBodyComponent, UE::Geometry::FDynamicMesh3& InMesh, UE::Geometry::FDynamicMesh3* InDilatedMesh, bool bMakeWaterInfoMeshConservativeRasterCompatible) const
 {
 	using namespace UE::Geometry;
 
@@ -508,6 +510,39 @@ void FWaterBodyMeshBuilder::GetDynamicMesh(const UWaterBodyComponent* WaterBodyC
 
 	WaterBodyComponent->GenerateWaterBodyMesh(InMesh, InDilatedMesh);
 
+	// Store positions of the previous and next vertex within each triangle in three UV channels.
+	// This is needed for conservative rasterization of the mesh when building the GPU water quadtree.
+	if (bMakeWaterInfoMeshConservativeRasterCompatible)
+	{
+		InMesh.Attributes()->SetNumUVLayers(3);
+		FDynamicMeshUVOverlay* UVOverlay0 = InMesh.Attributes()->GetUVLayer(0);
+		FDynamicMeshUVOverlay* UVOverlay1 = InMesh.Attributes()->GetUVLayer(1);
+		FDynamicMeshUVOverlay* UVOverlay2 = InMesh.Attributes()->GetUVLayer(2);
+
+		for (int TriangleID : InMesh.TriangleIndicesItr())
+		{
+			FVector3d TriangleVertices[3];
+			InMesh.GetTriVertices(TriangleID, TriangleVertices[0], TriangleVertices[1], TriangleVertices[2]);
+
+			FIndex3i UVTriangle0;
+			FIndex3i UVTriangle1;
+			FIndex3i UVTriangle2;
+			for (int TriangleVertexIndex = 0; TriangleVertexIndex < 3; ++TriangleVertexIndex)
+			{
+				const FVector3f PrevVertex = FVector3f(TriangleVertices[(TriangleVertexIndex + 2) % 3]);
+				const FVector3f NextVertex = FVector3f(TriangleVertices[(TriangleVertexIndex + 1) % 3]);
+
+				UVTriangle0[TriangleVertexIndex] = UVOverlay0->AppendElement(FVector2f(PrevVertex.X, PrevVertex.Y));
+				UVTriangle1[TriangleVertexIndex] = UVOverlay1->AppendElement(FVector2f(PrevVertex.Z, NextVertex.X));
+				UVTriangle2[TriangleVertexIndex] = UVOverlay2->AppendElement(FVector2f(NextVertex.Y, NextVertex.Z));
+			}
+
+			UVOverlay0->SetTriangle(TriangleID, UVTriangle0);
+			UVOverlay1->SetTriangle(TriangleID, UVTriangle1);
+			UVOverlay2->SetTriangle(TriangleID, UVTriangle2);
+		}
+	}
+
 	for (FDynamicMesh3* Mesh : { &InMesh, InDilatedMesh })
 	{
 		if (Mesh)
@@ -526,18 +561,12 @@ FMeshDescription FWaterBodyMeshBuilder::ConvertDynamicMeshToMeshDescription(cons
 	StaticMeshAttributes.Register();
 
 	FDynamicMeshToMeshDescription DynamicMeshToMeshDescription;
-	DynamicMeshToMeshDescription.ConversionOptions.bUpdateVtxColors = true;
-	// We recompute normals and tangents after generating the mesh so no need to recompute them here
-	DynamicMeshToMeshDescription.ConversionOptions.bUpdateNormals = false;
-	DynamicMeshToMeshDescription.ConversionOptions.bUpdateTangents = false;
-
 	DynamicMeshToMeshDescription.Convert(&Mesh, Result);
-	DynamicMeshToMeshDescription.UpdateUsingConversionOptions(&Mesh, Result);
 
 	return Result;
 }
 
-void FWaterBodyMeshBuilder::UpdateStaticMesh(UStaticMesh* WaterMesh, const FMeshDescription& MeshDescription) const
+void FWaterBodyMeshBuilder::UpdateStaticMesh(UStaticMesh* WaterMesh, const FMeshDescription& MeshDescription, bool bIsConservativeRasterMesh) const
 {
 	if (!WaterMesh->IsSourceModelValid(0))
 	{
@@ -549,7 +578,8 @@ void FWaterBodyMeshBuilder::UpdateStaticMesh(UStaticMesh* WaterMesh, const FMesh
 	SrcModel.BuildSettings.bRecomputeTangents = true;
 	SrcModel.BuildSettings.bRemoveDegenerates = false;
 	SrcModel.BuildSettings.bUseHighPrecisionTangentBasis = false;
-	SrcModel.BuildSettings.bUseFullPrecisionUVs = false;
+	SrcModel.BuildSettings.bUseFullPrecisionUVs = bIsConservativeRasterMesh; // CR mesh stores vertex positions in the UVs and needs full 32bit precision
+	SrcModel.BuildSettings.bGenerateLightmapUVs = false;
 	WaterMesh->CreateMeshDescription(0, MeshDescription);
 
 	UStaticMesh::FCommitMeshDescriptionParams CommitParams;
diff --git a/Engine/Plugins/Experimental/Water/Source/Runtime/Private/WaterBodyMeshBuilder.h b/Engine/Plugins/Experimental/Water/Source/Runtime/Private/WaterBodyMeshBuilder.h
--- a/Engine/Plugins/Experimental/Water/Source/Runtime/Private/WaterBodyMeshBuilder.h
+++ b/Engine/Plugins/Experimental/Water/Source/Runtime/Private/WaterBodyMeshBuilder.h
@@ -24,9 +24,10 @@ public:
 	TArray<TObjectPtr<UWaterBodyStaticMeshComponent>> BuildWaterBodyStaticMesh(UWaterBodyComponent* WaterBodyComponent, const FWaterBodyStaticMeshSettings& Settings, const TArray<TObjectPtr<UWaterBodyStaticMeshComponent>>& ReusableComponents) const;
 
 	/** 
-	 * Builds and assigns static meshes for the WaterBody Info mesh and Dilated Mesh directly to the components passed.
+	 * Builds and assigns static meshes for the WaterBody Info mesh and Dilated Mesh directly to the components passed. Setting bMakeWaterInfoMeshConservativeRasterCompatible to true stores positions of the previous and next vertex within each triangle in three UV channels
+	 * of the water info mesh. Such a setup is needed when using conservative rasterization to create the GPU water quadtree. This effectively makes every vertex unique, increasing the number of vertices in the buffer and making the index buffer useless.
 	 */
-	void BuildWaterInfoMeshes(UWaterBodyComponent* WaterBodyComponent, UWaterBodyInfoMeshComponent* WaterBodyInfoMesh, UWaterBodyInfoMeshComponent* WaterBodyDilatedMesh) const;
+	void BuildWaterInfoMeshes(UWaterBodyComponent* WaterBodyComponent, UWaterBodyInfoMeshComponent* WaterBodyInfoMesh, UWaterBodyInfoMeshComponent* WaterBodyDilatedMesh, bool bMakeWaterInfoMeshConservativeRasterCompatible) const;
 
 	/**
 	 * Builds a mesh description out of the water body geometry.
@@ -34,11 +35,11 @@ public:
 	FMeshDescription BuildMeshDescription(const UWaterBodyComponent* WaterBodyComponent) const;
 private:
 	/** Helper function to retrieve the dynamic mesh representations of a water body component with all the vertex attributes initialized. Optionally does the same for the DilatedMesh if the parameter is non-null */
-	void GetDynamicMesh(const UWaterBodyComponent* WaterBodyComponent, UE::Geometry::FDynamicMesh3& Mesh, UE::Geometry::FDynamicMesh3* DilateMesh = nullptr) const;
+	void GetDynamicMesh(const UWaterBodyComponent* WaterBodyComponent, UE::Geometry::FDynamicMesh3& Mesh, UE::Geometry::FDynamicMesh3* DilateMesh = nullptr, bool bMakeWaterInfoMeshConservativeRasterCompatible = false) const;
 	
 	FMeshDescription ConvertDynamicMeshToMeshDescription(const UE::Geometry::FDynamicMesh3& Mesh) const;
 
-	void UpdateStaticMesh(UStaticMesh* Mesh, const FMeshDescription& MeshDescription) const;
+	void UpdateStaticMesh(UStaticMesh* Mesh, const FMeshDescription& MeshDescription, bool bIsConservativeRasterMesh) const;
 
 	UStaticMesh* CreateUStaticMesh(UObject* Outer, FName MeshName) const;
 };
diff --git a/Engine/Plugins/Experimental/Water/Source/Runtime/Private/WaterMeshSceneProxy.cpp b/Engine/Plugins/Experimental/Water/Source/Runtime/Private/WaterMeshSceneProxy.cpp
--- a/Engine/Plugins/Experimental/Water/Source/Runtime/Private/WaterMeshSceneProxy.cpp
+++ b/Engine/Plugins/Experimental/Water/Source/Runtime/Private/WaterMeshSceneProxy.cpp
@@ -58,6 +58,11 @@ static TAutoConsoleVariable<float> CVarWaterMeshGPUQuadTreeJitterSampleFootprint
 	TEXT("Pixel footprint of the jitter sample pattern. Values greater than 1.0 can cause the water mesh to raster into neighboring pixels not normally covered by the mesh. Default: 1.5, Min 0.0"),
 	ECVF_RenderThreadSafe);
 
+static TAutoConsoleVariable<int32> CVarWaterMeshGPUQuadTreeConservativeRasterization(
+	TEXT("r.Water.WaterMesh.GPUQuadTree.ConservativeRasterization"), 0,
+	TEXT("Enables software conservative rasterization for rasterizing water body meshes into the water quadtree. Disables jittered draws. Default: 0"),
+	ECVF_RenderThreadSafe);
+
 static TAutoConsoleVariable<int32> CVarWaterMeshOcclusionCullingMaxQueries(
 	TEXT("r.Water.WaterMesh.OcclusionCulling.MaxQueries"), 256,
 	TEXT("Maximum number of occlusion queries for the CPU water quadtree nodes. Using fewer queries than nodes will result in coarser culling."),
@@ -1137,6 +1142,8 @@ void FWaterMeshSceneProxy::BuildGPUQuadTree(FRDGBuilder& GraphBuilder)
 	const FMatrix44f ProjectionMatrix = FMatrix44f(BuildOrthoMatrix(Resolution.X * LeafSize, Resolution.Y * LeafSize, WaterQuadTreeDepthRange));
 	const FMatrix44f ViewProjection = ViewMatrix * ProjectionMatrix;
 
+	bool bAllDrawsAreConservativeRasterCompatible = true;
+
 	// Create array of water body draws
 	TArray<FWaterQuadTreeGPU::FDraw> Draws;
 	{
@@ -1157,6 +1164,7 @@ void FWaterMeshSceneProxy::BuildGPUQuadTree(FRDGBuilder& GraphBuilder)
 					FWaterQuadTreeGPU::FDraw& Draw = Draws.AddDefaulted_GetRef();
 					Draw.Transform = FMatrix44f(LocalToQuadTreeWorld.ToMatrixWithScale()) * ViewProjection;
 					Draw.VertexBuffer = LODResource.VertexBuffers.PositionVertexBuffer.GetRHI();
+					Draw.TexCoordBuffer = LODResource.VertexBuffers.StaticMeshVertexBuffer.TexCoordVertexBuffer.GetRHI();
 					Draw.IndexBuffer = LODResource.IndexBuffer.GetRHI();
 					Draw.FirstIndex = MeshSection.FirstIndex;
 					Draw.NumPrimitives = MeshSection.NumTriangles;
@@ -1168,6 +1176,13 @@ void FWaterMeshSceneProxy::BuildGPUQuadTree(FRDGBuilder& GraphBuilder)
 					Draw.MaxZ = float(WBRenderData.BoundsMaxZ + WBRenderData.MaxWaveHeight - WaterQuadTreeMinHeight) * RcpWaterQuadTreeDepthRange;
 					Draw.MaxWaveHeight = WBRenderData.MaxWaveHeight * RcpWaterQuadTreeDepthRange;
 					Draw.bIsRiver = Info.bIsRiver;
+
+					// Conservative raster mesh has 3 unique vertices per triangle and 3 UV channels with 32bit float UVs.
+					bAllDrawsAreConservativeRasterCompatible = bAllDrawsAreConservativeRasterCompatible
+						&& LODResource.VertexBuffers.PositionVertexBuffer.GetNumVertices() == LODResource.IndexBuffer.GetNumIndices()
+						&& (LODResource.VertexBuffers.PositionVertexBuffer.GetNumVertices() % 3) == 0
+						&& LODResource.VertexBuffers.StaticMeshVertexBuffer.GetNumTexCoords() == 3
+						&& LODResource.VertexBuffers.StaticMeshVertexBuffer.GetUseFullPrecisionUVs();
 				}
 			}
 		}
@@ -1182,6 +1197,7 @@ void FWaterMeshSceneProxy::BuildGPUQuadTree(FRDGBuilder& GraphBuilder)
 	Params.JitterSampleFootprint = FMath::Max(CVarWaterMeshGPUQuadTreeJitterSampleFootprint.GetValueOnRenderThread(), 0.0f);
 	Params.CaptureDepthRange = WaterQuadTreeDepthRange;
 	Params.bUseMSAAJitterPattern = CVarWaterMeshGPUQuadTreeJitterPattern.GetValueOnRenderThread() == 1;
+	Params.bUseConservativeRasterization = bAllDrawsAreConservativeRasterCompatible && CVarWaterMeshGPUQuadTreeConservativeRasterization.GetValueOnRenderThread() != 0;
 
 	QuadTreeGPU.Init(GraphBuilder, Params, Draws);
 }
diff --git a/Engine/Plugins/Experimental/Water/Source/Runtime/Private/WaterQuadTreeGPU.cpp b/Engine/Plugins/Experimental/Water/Source/Runtime/Private/WaterQuadTreeGPU.cpp
--- a/Engine/Plugins/Experimental/Water/Source/Runtime/Private/WaterQuadTreeGPU.cpp
+++ b/Engine/Plugins/Experimental/Water/Source/Runtime/Private/WaterQuadTreeGPU.cpp
@@ -47,12 +47,13 @@ public:
 	DECLARE_GLOBAL_SHADER(FWaterQuadTreeVS);
 	SHADER_USE_PARAMETER_STRUCT(FWaterQuadTreeVS, FGlobalShader);
 
-	class FApplyJitter : SHADER_PERMUTATION_BOOL("APPLY_JITTER");
-	using FPermutationDomain = TShaderPermutationDomain<FApplyJitter>;
+	class FRasterMode : SHADER_PERMUTATION_SPARSE_INT("RASTER_MODE", 0, 1, 2);
+	using FPermutationDomain = TShaderPermutationDomain<FRasterMode>;
 
 	BEGIN_SHADER_PARAMETER_STRUCT(FParameters, )
 		SHADER_PARAMETER(FMatrix44f, Transform)
 		SHADER_PARAMETER(FVector2f, JitterScale)
+		SHADER_PARAMETER(FVector2f, HalfPixelSize)
 	END_SHADER_PARAMETER_STRUCT()
 
 	static void ModifyCompilationEnvironment(const FGlobalShaderPermutationParameters& Parameters, FShaderCompilerEnvironment& OutEnvironment)
@@ -68,6 +69,9 @@ public:
 	DECLARE_GLOBAL_SHADER(FWaterQuadTreePS);
 	SHADER_USE_PARAMETER_STRUCT(FWaterQuadTreePS, FGlobalShader);
 
+	class FClipConservativeTriangle : SHADER_PERMUTATION_BOOL("CLIP_CONSERVATIVE_TRIANGLE");
+	using FPermutationDomain = TShaderPermutationDomain<FClipConservativeTriangle>;
+
 	BEGIN_SHADER_PARAMETER_STRUCT(FParameters, )
 		SHADER_PARAMETER(uint32, WaterBodyRenderDataIndex)
 		SHADER_PARAMETER(int32, Priority)
@@ -486,6 +490,28 @@ public:
 TGlobalResource<FVector3AndInstancedVector2VertexDeclaration> GVector3AndInstancedVector2VertexDeclaration;
 
 
+class FVector3AndThreeVector2VertexDeclaration : public FRenderResource
+{
+public:
+	FVertexDeclarationRHIRef VertexDeclarationRHI;
+	virtual void InitRHI(FRHICommandListBase& RHICmdList) override
+	{
+		FVertexDeclarationElementList Elements;
+		Elements.Add(FVertexElement(0, 0, VET_Float3, 0, sizeof(FVector3f)));
+		Elements.Add(FVertexElement(1, 0, VET_Float2, 1, sizeof(FVector2f) * 3));
+		Elements.Add(FVertexElement(1, 8, VET_Float2, 2, sizeof(FVector2f) * 3));
+		Elements.Add(FVertexElement(1, 16, VET_Float2, 3, sizeof(FVector2f) * 3));
+		VertexDeclarationRHI = PipelineStateCache::GetOrCreateVertexDeclaration(Elements);
+	}
+	virtual void ReleaseRHI() override
+	{
+		VertexDeclarationRHI.SafeRelease();
+	}
+};
+
+TGlobalResource<FVector3AndThreeVector2VertexDeclaration> GVector3AndAndThreeVector2VertexDeclaration;
+
+
 void FWaterQuadTreeGPU::Init(FRDGBuilder& GraphBuilder, const FInitParams& Params, TArray<FDraw>& Draws)
 {
 	TRACE_CPUPROFILER_EVENT_SCOPE(FWaterQuadTreeGPU::Init);
@@ -556,10 +582,16 @@ void FWaterQuadTreeGPU::Init(FRDGBuilder& GraphBuilder, const FInitParams& Param
 
 	// Raster water body meshes
 	{
+		enum class ERasterMode { Regular = 0, Jittered = 1, Conservative = 2 };
+		const ERasterMode RasterMode = Params.bUseConservativeRasterization ? ERasterMode::Conservative : NumJitterSamples > 1 ? ERasterMode::Jittered : ERasterMode::Regular;
+
 		FWaterQuadTreeVS::FPermutationDomain VSPermutationDomain;
-		VSPermutationDomain.Set<FWaterQuadTreeVS::FApplyJitter>(NumJitterSamples > 1);
+		VSPermutationDomain.Set<FWaterQuadTreeVS::FRasterMode>(static_cast<int32>(RasterMode));
 		TShaderMapRef<FWaterQuadTreeVS> VertexShader(ShaderMap, VSPermutationDomain);
-		TShaderMapRef<FWaterQuadTreePS> PixelShader(ShaderMap);
+		
+		FWaterQuadTreePS::FPermutationDomain PSPermutationDomain;
+		PSPermutationDomain.Set<FWaterQuadTreePS::FClipConservativeTriangle>(RasterMode == ERasterMode::Conservative);
+		TShaderMapRef<FWaterQuadTreePS> PixelShader(ShaderMap, PSPermutationDomain);
 
 		FWaterQuadTreeParameters* PassParameters = GraphBuilder.AllocParameters<FWaterQuadTreeParameters>();
 		PassParameters->PS.RenderTargets[0] = FRenderTargetBinding(WaterBodyRasterTexture, ERenderTargetLoadAction::EClear);
@@ -569,7 +601,7 @@ void FWaterQuadTreeGPU::Init(FRDGBuilder& GraphBuilder, const FInitParams& Param
 			RDG_EVENT_NAME("WaterQuadTreeRaster"),
 			PassParameters,
 			ERDGPassFlags::Raster,
-			[PassParameters, RasterResolution, NumJitterSamples, bMSAAPattern = Params.bUseMSAAJitterPattern, JitterFootprint = Params.JitterSampleFootprint, Draws = MoveTemp(Draws), VertexShader, PixelShader](FRHICommandList& RHICmdList)
+			[PassParameters, RasterResolution, NumJitterSamples, RasterMode, bMSAAPattern = Params.bUseMSAAJitterPattern, JitterFootprint = Params.JitterSampleFootprint, Draws = MoveTemp(Draws), VertexShader, PixelShader](FRHICommandList& RHICmdList)
 			{
 				RHICmdList.SetViewport(0.0f, 0.0f, 0.0f, RasterResolution.X, RasterResolution.Y, 1.0f);
 
@@ -581,14 +613,21 @@ void FWaterQuadTreeGPU::Init(FRDGBuilder& GraphBuilder, const FInitParams& Param
 					CW_RGBA, BO_Max, BF_One, BF_One, BO_Max, BF_One, BF_One,
 					CW_RGBA, BO_Max, BF_One, BF_One, BO_Max, BF_One, BF_One>::GetRHI(); // MAX blending
 				GraphicsPSOInit.PrimitiveType = PT_TriangleList;
-				GraphicsPSOInit.BoundShaderState.VertexDeclarationRHI = NumJitterSamples > 1 ? GVector3AndInstancedVector2VertexDeclaration.VertexDeclarationRHI : GetVertexDeclarationFVector3();
 				GraphicsPSOInit.BoundShaderState.VertexShaderRHI = VertexShader.GetVertexShader();
 				GraphicsPSOInit.BoundShaderState.PixelShaderRHI = PixelShader.GetPixelShader();
+				switch (RasterMode)
+				{
+				case ERasterMode::Regular: GraphicsPSOInit.BoundShaderState.VertexDeclarationRHI = GetVertexDeclarationFVector3(); break;
+				case ERasterMode::Jittered: GraphicsPSOInit.BoundShaderState.VertexDeclarationRHI = GVector3AndInstancedVector2VertexDeclaration.VertexDeclarationRHI; break;
+				case ERasterMode::Conservative: GraphicsPSOInit.BoundShaderState.VertexDeclarationRHI = GVector3AndAndThreeVector2VertexDeclaration.VertexDeclarationRHI; break;
+				default: checkNoEntry(); break;
+				}
 
 				SetGraphicsPipelineState(RHICmdList, GraphicsPSOInit, 0);
 
+				const FVector2f HalfPixelSize = FVector2f(1.0f / RasterResolution.X, 1.0f / RasterResolution.Y);
 				const FVector2f JitterScale = JitterFootprint * FVector2f(2.0f / RasterResolution.X, 2.0f / RasterResolution.Y);
-				if (NumJitterSamples > 1)
+				if (RasterMode == ERasterMode::Jittered)
 				{
 					const int32 JitterVertexBufferOffset = bMSAAPattern ? FJitterOffsetVertexBuffer::GetMSAADataOffset(NumJitterSamples) : 0;
 					RHICmdList.SetStreamSource(1, GJitterOffsetVertexBuffer.GetRHI(), JitterVertexBufferOffset);
@@ -598,6 +637,7 @@ void FWaterQuadTreeGPU::Init(FRDGBuilder& GraphBuilder, const FInitParams& Param
 				{
 					PassParameters->VS.Transform = Draw.Transform;
 					PassParameters->VS.JitterScale = JitterScale;
+					PassParameters->VS.HalfPixelSize = HalfPixelSize;
 					PassParameters->PS.WaterBodyRenderDataIndex = Draw.WaterBodyRenderDataIndex;
 					PassParameters->PS.Priority = Draw.Priority;
 					PassParameters->PS.WaterBodyMinZ = Draw.MinZ;
@@ -608,7 +648,25 @@ void FWaterQuadTreeGPU::Init(FRDGBuilder& GraphBuilder, const FInitParams& Param
 					SetShaderParameters(RHICmdList, VertexShader, VertexShader.GetVertexShader(), PassParameters->VS);
 					SetShaderParameters(RHICmdList, PixelShader, PixelShader.GetPixelShader(), PassParameters->PS);
 					RHICmdList.SetStreamSource(0, Draw.VertexBuffer, 0);
-					RHICmdList.DrawIndexedPrimitive(Draw.IndexBuffer, Draw.BaseVertexIndex, 0, Draw.NumVertices, Draw.FirstIndex, Draw.NumPrimitives, NumJitterSamples);
+
+					switch (RasterMode)
+					{
+					case ERasterMode::Regular: // Fallthrough
+					case ERasterMode::Jittered:
+					{
+						const uint32 NumInstances = (RasterMode == ERasterMode::Jittered) ? NumJitterSamples : 1;
+						RHICmdList.DrawIndexedPrimitive(Draw.IndexBuffer, Draw.BaseVertexIndex, 0, Draw.NumVertices, Draw.FirstIndex, Draw.NumPrimitives, NumInstances);
+						break;
+					}
+					case ERasterMode::Conservative:
+					{
+						RHICmdList.SetStreamSource(1, Draw.TexCoordBuffer, 0);
+						check((Draw.NumVertices % 3) == 0);
+						RHICmdList.DrawPrimitive(Draw.BaseVertexIndex, Draw.NumVertices / 3, 1);
+						break;
+					}
+					default: checkNoEntry(); break;
+					}
 				}
 			});
 	}
diff --git a/Engine/Plugins/Experimental/Water/Source/Runtime/Private/WaterQuadTreeGPU.h b/Engine/Plugins/Experimental/Water/Source/Runtime/Private/WaterQuadTreeGPU.h
--- a/Engine/Plugins/Experimental/Water/Source/Runtime/Private/WaterQuadTreeGPU.h
+++ b/Engine/Plugins/Experimental/Water/Source/Runtime/Private/WaterQuadTreeGPU.h
@@ -17,6 +17,7 @@ public:
 	{
 		FMatrix44f Transform = FMatrix44f::Identity;
 		FRHIBuffer* VertexBuffer = nullptr;
+		FRHIBuffer* TexCoordBuffer = nullptr; // Stores previous and next vertex in the triangle for every vertex when using conservative rasterization
 		FRHIBuffer* IndexBuffer = nullptr;
 		uint32 FirstIndex = 0;
 		uint32 NumPrimitives = 0;
@@ -54,6 +55,7 @@ public:
 		float JitterSampleFootprint = 1.0f;
 		float CaptureDepthRange = 1.0f;
 		bool bUseMSAAJitterPattern = false;
+		bool bUseConservativeRasterization = false;
 	};
 
 	struct FTraverseParams
