The GPU waterquadtree now only has a single mesh density per draw bucket. This is done by using one or multiple lower poly quad meshes to construct higher poly water quadtree tiles. Instance data is now more compressed (2-3 float4 -> 3-4 uint32) and supplied with fixed function to avoid an explicit buffer load from an offset buffer.

#rnx

diff --git a/Engine/Plugins/Experimental/Water/Shaders/Private/WaterMeshVertexFactory.ush b/Engine/Plugins/Experimental/Water/Shaders/Private/WaterMeshVertexFactory.ush
--- a/Engine/Plugins/Experimental/Water/Shaders/Private/WaterMeshVertexFactory.ush
+++ b/Engine/Plugins/Experimental/Water/Shaders/Private/WaterMeshVertexFactory.ush
@@ -6,17 +6,11 @@
 #define WATER_MESH_DRAW_INDIRECT 0
 #endif
 
-#if WATER_MESH_DRAW_INDIRECT
-
-Buffer<uint> InstanceDataOffsetsBuffer;
-Buffer<float4> InstanceData0Buffer;
-Buffer<float4> InstanceData1Buffer;
-#if WITH_WATER_SELECTION_SUPPORT_VF
-Buffer<float4> InstanceData2Buffer;
-#endif
-
-uint DrawBucketIndex;
+#define WATER_MESH_DRAW_INDIRECT_INTERNAL WATER_MESH_DRAW_INDIRECT && !RAYHITGROUPSHADER && !COMPUTESHADER
 
+#if WATER_MESH_DRAW_INDIRECT_INTERNAL
+float3 QuadTreePosition;
+uint bLODMorphingEnabled;
 #endif
 
 struct FVertexFactoryInterpolantsVSToPS
@@ -43,8 +37,13 @@ struct FVertexFactoryInput
 {
 	float4	Position	: ATTRIBUTE0;
 
-#if WATER_MESH_DRAW_INDIRECT && !RAYHITGROUPSHADER && !COMPUTESHADER
-	uint InstanceId : SV_InstanceID;
+#if WATER_MESH_DRAW_INDIRECT_INTERNAL
+	uint InstanceData0 : ATTRIBUTE8;
+	uint InstanceData1 : ATTRIBUTE9;
+	uint InstanceData2 : ATTRIBUTE10;
+#if HIT_PROXY_SHADER || WITH_WATER_SELECTION_SUPPORT_VF
+	uint InstanceData3 : ATTRIBUTE11;
+#endif
 #else
 	float4 InstanceData0 : ATTRIBUTE8;
 	float4 InstanceData1 : ATTRIBUTE9; 
@@ -52,7 +51,7 @@ struct FVertexFactoryInput
 	float4 InstanceData2 : ATTRIBUTE10; // Store an HitProxyID per instance
 #endif
 	VF_INSTANCED_STEREO_DECLARE_INPUT_BLOCK()
-#endif // WATER_MESH_DRAW_INDIRECT && !RAYHITGROUPSHADER && !COMPUTESHADER
+#endif // WATER_MESH_DRAW_INDIRECT_INTERNAL
 
 	VF_GPUSCENE_DECLARE_INPUT_BLOCK(13)
 };
@@ -64,13 +63,15 @@ struct FPositionOnlyVertexFactoryInput
 {
 	float4	Position	: ATTRIBUTE0;
 
-#if WATER_MESH_DRAW_INDIRECT && !RAYHITGROUPSHADER && !COMPUTESHADER
-	uint InstanceId : SV_InstanceID;
+#if WATER_MESH_DRAW_INDIRECT_INTERNAL
+	uint InstanceData0 : ATTRIBUTE8;
+	uint InstanceData1 : ATTRIBUTE9;
+	uint InstanceData2 : ATTRIBUTE10;
 #else
 	float4 InstanceData0 : ATTRIBUTE8;
 	float4 InstanceData1 : ATTRIBUTE9;
 	VF_INSTANCED_STEREO_DECLARE_INPUT_BLOCK()
-#endif // WATER_MESH_DRAW_INDIRECT && !RAYHITGROUPSHADER && !COMPUTESHADER
+#endif // WATER_MESH_DRAW_INDIRECT_INTERNAL
 
 	VF_GPUSCENE_DECLARE_INPUT_BLOCK(1)
 };
@@ -83,13 +84,15 @@ struct FPositionAndNormalOnlyVertexFactoryInput
 	float4	Position	: ATTRIBUTE0;
 	float4	Normal		: ATTRIBUTE2;
 
-#if WATER_MESH_DRAW_INDIRECT && !RAYHITGROUPSHADER && !COMPUTESHADER
-	uint InstanceId : SV_InstanceID;
+#if WATER_MESH_DRAW_INDIRECT_INTERNAL
+	uint InstanceData0 : ATTRIBUTE8;
+	uint InstanceData1 : ATTRIBUTE9;
+	uint InstanceData2 : ATTRIBUTE10;
 #else
 	float4 InstanceData0 : ATTRIBUTE8;
 	float4 InstanceData1 : ATTRIBUTE9;
 	VF_INSTANCED_STEREO_DECLARE_INPUT_BLOCK()
-#endif // WATER_MESH_DRAW_INDIRECT && !RAYHITGROUPSHADER && !COMPUTESHADER
+#endif // WATER_MESH_DRAW_INDIRECT_INTERNAL
 
 	VF_GPUSCENE_DECLARE_INPUT_BLOCK(1)
 };
@@ -233,52 +236,129 @@ float3 MorphTranslatedWorldPosition(float3 OriginalTranslatedWorldPos, float2 In
 	return TranslatedWorldPos;
 }
 
+struct FWaterVertexFactoryInstanceInput
+{
+	float2 Position;
+	float3 Translation;
+	uint WaterBodyIndex;
+	float LODLevel;
+	float2 Scale;
+	float HeightLODFactor;
+	uint NumQuadsPerTileSide;
+	bool bShouldMorph;
+	bool bCanMorphTwice;
+};
+
+#if WATER_MESH_DRAW_INDIRECT_INTERNAL
+
+FWaterVertexFactoryInstanceInput UnpackWaterVertexFactoryInstanceInput(float4 InPosition, uint InData0, uint InData1, uint InData2)
+{
+	const uint2 NodeCoord = uint2(InData0 & 0x7FFu, (InData0 >> 11u) & 0x7FFu);
+	const uint RawLODLevel = (InData0 >> 22u) & 0x1Fu;
+	const uint RawDensityIndex = (InData0 >> 27u) & 0x1Fu;
+	const uint LODLevel = RawLODLevel + RawDensityIndex;
+	const uint DensityIndex = min(RawDensityIndex, (uint)WaterVF.NumDensities - 1);
+	const float WaterHeight = f16tof32(InData1) * WaterVF.CaptureDepthRange;
+	const uint2 TileCoord = uint2((InData1 >> 16u) & 0xFFu, (InData1 >> 24u) & 0xFFu);
+	const float HeightLODFactor = float(InData2 & 0xFFu) / 255.0f;
+	const uint WaterBodyIndex = InData2 >> 8u;
+	const float Scale = pow(2.0f, (float)RawLODLevel) * WaterVF.LeafSize;
+
+	bool bKilledTriangle = false;
+
+	// Shift and scale the tile according to its tile coord within the node and kill superfluous triangles
+	const uint NumQuadsPerLogicalTileSide = max(2u, (uint)WaterVF.NumQuadsLOD0 >> DensityIndex);
+	{
+		if ((uint)WaterVF.NumQuadsPerTileSide > NumQuadsPerLogicalTileSide)
+		{
+			InPosition.xy *= (float)WaterVF.NumQuadsPerTileSide / (float)NumQuadsPerLogicalTileSide;
+			if (any(InPosition.xy < -0.5f) || any(InPosition.xy > 0.5f))
+			{
+				// Kill triangle by setting position to NaN
+				InPosition = asfloat(0xFFFFFFFF).xxxx;
+				bKilledTriangle = true;
+			}
+		}
+
+		const uint NumTiles = max((uint)WaterVF.NumQuadsPerTileSide, NumQuadsPerLogicalTileSide) / (uint)WaterVF.NumQuadsPerTileSide;
+		const float PositionScale = 1.0f / NumTiles;
+		const float2 PositionBias = TileCoord * PositionScale;
+		InPosition.xy += 0.5f;
+		InPosition.xy = InPosition.xy * PositionScale + PositionBias;
+		InPosition.xy -= 0.5f;
+	}
+
+	FWaterVertexFactoryInstanceInput Result = (FWaterVertexFactoryInstanceInput)0;
+	Result.Position = InPosition.xy;
+	Result.Translation = float3((NodeCoord + 0.5f) * Scale, WaterHeight) + QuadTreePosition;
+	Result.WaterBodyIndex = WaterBodyIndex;
+	Result.LODLevel = (float)LODLevel;
+	Result.Scale = Scale.xx;
+	Result.NumQuadsPerTileSide = NumQuadsPerLogicalTileSide;
+	Result.HeightLODFactor = HeightLODFactor;
+	// Only allow a tile to morph if it's not the last density level and not the last LOD level, since there is no next level to morph to
+	Result.bShouldMorph = !bKilledTriangle && ((bLODMorphingEnabled != 0) && (DensityIndex != (uint)WaterVF.NumDensities - 1)) ? 1 : 0;
+	// Tiles can morph twice to be able to morph between 3 LOD levels. Next to last density level can only morph once
+	Result.bCanMorphTwice = !bKilledTriangle && (DensityIndex < (uint)WaterVF.NumDensities - 2) ? 1 : 0;
+
+	return Result;
+}
+
+#else // WATER_MESH_DRAW_INDIRECT_INTERNAL
+
+FWaterVertexFactoryInstanceInput UnpackWaterVertexFactoryInstanceInput(float4 InPosition, float4 InData0, float4 InData1)
+{
+	const uint PackedDataChannel = asuint(InData1.x);
+
+	FWaterVertexFactoryInstanceInput Result = (FWaterVertexFactoryInstanceInput)0;
+	Result.Position = InPosition.xy;
+	Result.Translation = InData0.xyz;
+	Result.WaterBodyIndex = asuint(InData0.w);
+	Result.LODLevel = (float)(PackedDataChannel & 0xFF);
+	Result.Scale = InData1.zw;
+	Result.HeightLODFactor = InData1.y;
+	Result.NumQuadsPerTileSide = (uint)WaterVF.NumQuadsPerTileSide;
+	Result.bShouldMorph = ((PackedDataChannel >> 8u) & 0x1u) != 0;
+	Result.bCanMorphTwice = ((PackedDataChannel >> 9u) & 0x1u) != 0;
+
+	return Result;
+}
+
+#endif // WATER_MESH_DRAW_INDIRECT_INTERNAL
+
 FVertexFactoryIntermediates GetVertexFactoryIntermediates(FVertexFactoryInput Input)
 {
 	FVertexFactoryIntermediates Intermediates;
 
-	// Get the packed instance data
-#if WATER_MESH_DRAW_INDIRECT && !RAYHITGROUPSHADER && !COMPUTESHADER
-	const uint InstanceDataIndex = InstanceDataOffsetsBuffer.Load(DrawBucketIndex) + Input.InstanceId;
-	float4 Data0 = InstanceData0Buffer.Load(InstanceDataIndex);
-	float4 Data1 = InstanceData1Buffer.Load(InstanceDataIndex);
+#if WATER_MESH_DRAW_INDIRECT_INTERNAL
+	const FWaterVertexFactoryInstanceInput InstanceInput = UnpackWaterVertexFactoryInstanceInput(Input.Position, Input.InstanceData0, Input.InstanceData1, Input.InstanceData2);
 #else
-	float4 Data0 = Input.InstanceData0;
-	float4 Data1 = Input.InstanceData1;
-#endif // WATER_MESH_DRAW_INDIRECT && !RAYHITGROUPSHADER && !COMPUTESHADER
-
-	const float3 Translation = Data0.xyz;
-	Intermediates.WaveParamIndex = asuint(Data0.w);
-	const float3 Scale = float3(Data1.zw, 1.0f);
-	const uint PackedDataChannel = asuint(Data1.x);
+	const FWaterVertexFactoryInstanceInput InstanceInput = UnpackWaterVertexFactoryInstanceInput(Input.Position, Input.InstanceData0, Input.InstanceData1);
+#endif
 
-	// Lod level is in first 8 bits and ShouldMorph bit is in the 9th bit
-	const float LODLevel = (float)(PackedDataChannel & 0xFF);
-	const uint ShouldMorph = ((PackedDataChannel >> 8) & 0x1); 
+	Intermediates.WaveParamIndex = InstanceInput.WaterBodyIndex;
 
 	// Calculate the world pos
-	float3 TranslatedWorldPosition = float3(Input.Position.xy, 0.0f) * Scale + Translation;
+	float3 TranslatedWorldPosition = float3(InstanceInput.Position.xy * InstanceInput.Scale, 0.0f) + InstanceInput.Translation;
 
 	Intermediates.SceneData = VF_GPUSCENE_GET_INTERMEDIATES(Input);
 	
-	if (ShouldMorph == 1)
+	if (InstanceInput.bShouldMorph)
 	{
 		// Factor that applies across the entire lowest LOD, based on height above water
 		// Then this reaches 1.0, the lowest LOD will pop out and be replaced by the next LOD, starting over at 0.0. Rinse and repeat.
-		//const float HeightLODFactor = 0.0f;
-		const float HeightLODFactor = Data1.y; 
+		const float HeightLODFactor = InstanceInput.HeightLODFactor;
 	
-		const float2 QuadSize = Scale.xy / WaterVF.NumQuadsPerTileSide;
+		const float2 QuadSize = InstanceInput.Scale.xy / InstanceInput.NumQuadsPerTileSide;
 	
 		float OutLODFactor = 0.0f;
-		float3 TempTranslatedWorldPosition = MorphTranslatedWorldPosition(TranslatedWorldPosition, Translation.xy - Scale.xy*0.5f, LODLevel, WaterVF.LODScale, QuadSize, HeightLODFactor, OutLODFactor);
+		float3 TempTranslatedWorldPosition = MorphTranslatedWorldPosition(TranslatedWorldPosition, InstanceInput.Translation.xy - InstanceInput.Scale.xy*0.5f, InstanceInput.LODLevel, WaterVF.LODScale, QuadSize, HeightLODFactor, OutLODFactor);
 	
 		// If the vert is fully morphed, we morph it again as if it was a vert from the next LOD. This effectively means one tile can morph between 3 different LOD levels. 
 		// This is needed because the current LOD level tile will stick in to the next LOD level since we move the observer around. Make sure to not add any heightLODFactor to this, since that only applies to lowest LOD
-		const uint CanMorphTwice = ((PackedDataChannel >> 9) & 0x1); 
-		if (OutLODFactor >= 1.0f && CanMorphTwice == 1)
+		if (OutLODFactor >= 1.0f && InstanceInput.bCanMorphTwice)
 		{
-			TempTranslatedWorldPosition = MorphTranslatedWorldPosition(TempTranslatedWorldPosition, Translation.xy - Scale.xy*0.5f, LODLevel + 1, WaterVF.LODScale, QuadSize * 2.0f, 0.0f, OutLODFactor);
+			TempTranslatedWorldPosition = MorphTranslatedWorldPosition(TempTranslatedWorldPosition, InstanceInput.Translation.xy - InstanceInput.Scale.xy*0.5f, InstanceInput.LODLevel + 1, WaterVF.LODScale, QuadSize * 2.0f, 0.0f, OutLODFactor);
 		}
 	
 		Intermediates.MorphedTranslatedWorldPos = TempTranslatedWorldPosition;
@@ -289,8 +369,8 @@ FVertexFactoryIntermediates GetVertexFactoryIntermediates(FVertexFactoryInput In
 	}
 
 #if WITH_WATER_SELECTION_SUPPORT_VF
-#if WATER_MESH_DRAW_INDIRECT && !RAYHITGROUPSHADER && !COMPUTESHADER
-	float SelectedValue = InstanceData2Buffer.Load(InstanceDataIndex).w;
+#if WATER_MESH_DRAW_INDIRECT_INTERNAL
+	float SelectedValue = UnpackRGBA8(Input.InstanceData2).w;
 #else
 	float SelectedValue = Input.InstanceData2.w;
 #endif
@@ -307,9 +387,8 @@ FVertexFactoryIntermediates GetVertexFactoryIntermediates(FVertexFactoryInput In
 float4 VertexFactoryGetInstanceHitProxyId(FVertexFactoryInput Input, FVertexFactoryIntermediates Intermediates)
 {
 #if WITH_WATER_SELECTION_SUPPORT_VF
-#if WATER_MESH_DRAW_INDIRECT && !RAYHITGROUPSHADER && !COMPUTESHADER
-	const uint InstanceDataIndex = InstanceDataOffsetsBuffer.Load(DrawBucketIndex) + Input.InstanceId;
-	return float4(InstanceData2Buffer.Load(InstanceDataIndex).rgb, 0);
+#if WATER_MESH_DRAW_INDIRECT_INTERNAL
+	return float4(UnpackRGBA8(Input.InstanceData2).bgr, 0);
 #else
 	return float4(Input.InstanceData2.rgb, 0);
 #endif
diff --git a/Engine/Plugins/Experimental/Water/Shaders/Private/WaterQuadTreeDraws.usf b/Engine/Plugins/Experimental/Water/Shaders/Private/WaterQuadTreeDraws.usf
--- a/Engine/Plugins/Experimental/Water/Shaders/Private/WaterQuadTreeDraws.usf
+++ b/Engine/Plugins/Experimental/Water/Shaders/Private/WaterQuadTreeDraws.usf
@@ -61,30 +61,25 @@ RWBuffer<uint> IndirectArgs;
 #if PRECISE_OCCLUSION_QUERIES
 RWBuffer<uint> OcclusionQueryIndirectArgs;
 #endif
-uint NumDensities;
-uint NumMaterials;
+uint NumDrawBuckets;
 uint NumViews;
-uint NumQuadsLOD0;
+uint NumQuads;
 
 [numthreads(64, 1, 1)]
 void MainCS(uint3 DispatchThreadId : SV_DispatchThreadID)
 {
-	const uint NumBuckets = NumDensities * NumMaterials;
 	const uint BucketIndex = DispatchThreadId.x;
 
-	if (BucketIndex >= NumBuckets)
+	if (BucketIndex >= NumDrawBuckets)
 	{
 		return;
 	}
 
-	const uint DensityIndex = BucketIndex % NumDensities;
-	const uint MaterialIndex = BucketIndex / NumDensities;
-	const uint NumQuads = NumQuadsLOD0 >> DensityIndex;
 	const uint IndexCountPerInstance = NumQuads * NumQuads * 6;
 
 	for (uint ViewIndex = 0; ViewIndex < NumViews; ++ViewIndex)
 	{
-		const uint IndirectArgIndex = NumBuckets * ViewIndex + BucketIndex;
+		const uint IndirectArgIndex = NumDrawBuckets * ViewIndex + BucketIndex;
 		IndirectArgs[IndirectArgIndex * 5 + 0] = IndexCountPerInstance;
 		IndirectArgs[IndirectArgIndex * 5 + 1] = 0; // InstanceCount
 		IndirectArgs[IndirectArgIndex * 5 + 2] = 0; // StartIndexLocation
@@ -156,8 +151,6 @@ float3 QuadTreePosition;
 float3 ObserverPosition;
 uint QuadTreeResolutionX;
 uint QuadTreeResolutionY;
-uint NumDensities;
-uint NumMaterials;
 uint ViewIndex;
 float LeafSize;
 float LODScale;
@@ -427,9 +420,10 @@ ByteAddressBuffer PackedNodes;
 #if PRECISE_OCCLUSION_QUERIES
 Buffer<uint> OcclusionResults;
 #endif // PRECISE_OCCLUSION_QUERIES
-uint NumDensities;
-uint NumMaterials;
 uint NumDispatchedThreads;
+uint NumDensities;
+uint NumQuadsLOD0;
+uint NumQuadsPerDraw;
 
 [numthreads(64, 1, 1)]
 void MainCS(uint3 DispatchThreadId : SV_DispatchThreadID)
@@ -449,7 +443,7 @@ void MainCS(uint3 DispatchThreadId : SV_DispatchThreadID)
 
 		// Load packed node coord and density index
 		const uint4 NodeCoordAndDensity = UnpackNodeCoord(PackedNodes.Load((DrawIndex + 1) << 2));
-		const uint DensityIndex = min(NodeCoordAndDensity.w, NumDensities - 1);
+		const uint DensityIndexClamped = min(NodeCoordAndDensity.w, NumDensities - 1);
 
 		// Sample the quadtree to access the water body render data
 		const float4 QuadTreeSample = QuadTreeTexture.Load(NodeCoordAndDensity.xyz);
@@ -473,9 +467,11 @@ void MainCS(uint3 DispatchThreadId : SV_DispatchThreadID)
 		}
 
 		// Increment bucket counter
-		const uint BucketIndex = MaterialIndex * NumDensities + DensityIndex;
+		const uint BucketIndex = MaterialIndex;
+		const uint NumTilesPerEdge = max(NumQuadsPerDraw, NumQuadsLOD0 >> DensityIndexClamped) / NumQuadsPerDraw;
+		const uint NumDraws = NumTilesPerEdge * NumTilesPerEdge;
 		uint Dummy;
-		BucketCounts.InterlockedAdd(BucketIndex << 2, 1, Dummy);
+		BucketCounts.InterlockedAdd(BucketIndex << 2, NumDraws, Dummy);
 	}
 }
 
@@ -629,9 +625,10 @@ void MainCS(uint3 DispatchThreadId : SV_DispatchThreadID)
 #endif
 
 RWBuffer<uint> IndirectArgs;
-RWBuffer<float4> InstanceData0;
-RWBuffer<float4> InstanceData1;
-RWBuffer<float4> InstanceData2;
+RWBuffer<uint> InstanceData0;
+RWBuffer<uint> InstanceData1;
+RWBuffer<uint> InstanceData2;
+RWBuffer<uint> InstanceData3;
 Texture2D<float4> QuadTreeTexture;
 Texture2D<float4> WaterZBoundsTexture;
 StructuredBuffer<FWaterBodyRenderData> WaterBodyRenderData;
@@ -650,6 +647,8 @@ uint NumMaterials;
 uint NumDispatchedThreads;
 uint BucketIndexOffset;
 uint NumLODs;
+uint NumQuadsLOD0;
+uint NumQuadsPerDraw;
 float LeafSize;
 float LODScale;
 float CaptureDepthRange;
@@ -702,64 +701,78 @@ void MainCS(uint3 DispatchThreadId : SV_DispatchThreadID)
 			}
 		}
 
-		const uint BucketIndex = MaterialIndex * NumDensities + DensityIndexClamped + BucketIndexOffset;
+		const uint BucketIndex = MaterialIndex + BucketIndexOffset;
 
 		const uint BucketInstanceDataOffset = InstanceDataOffsets.Load(BucketIndex);
-		uint LocalInstanceDataOffset;
-		InterlockedAdd(IndirectArgs[BucketIndex * 5 + 1], 1, LocalInstanceDataOffset); // Increment InstanceCount
-		const uint InstanceDataIndex = BucketInstanceDataOffset + LocalInstanceDataOffset;
+		const uint NumTilesPerEdge = max(NumQuadsPerDraw, NumQuadsLOD0 >> DensityIndexClamped) / NumQuadsPerDraw;
+		const uint NumInstances = NumTilesPerEdge * NumTilesPerEdge;
 		
-		const float Scale = (1u << LODLevel) * LeafSize;
-		const float WaterSurfaceBaseHeight = WaterZBoundsTexture.Load(int3(NodeCoord, LODLevel)).z * CaptureDepthRange;
-		const float3 TranslatedWorldPosition = float3((NodeCoord + 0.5f) * Scale, WaterSurfaceBaseHeight) + QuadTreePosition;
+		// Increment InstanceCount
+		uint LocalInstanceDataOffset;
+		InterlockedAdd(IndirectArgs[BucketIndex * 5 + 1], NumInstances, LocalInstanceDataOffset); 
+		const uint InstanceDataBaseIndex = BucketInstanceDataOffset + LocalInstanceDataOffset;
 
-		// Add the data to the bucket
+		// On the first write to this bucket, set StartInstanceLocation to offset where the per-instance vertex attributes are read from
+		if (LocalInstanceDataOffset == 0)
 		{
-			float4 Data0;
-			Data0.x = TranslatedWorldPosition.x;
-			Data0.y = TranslatedWorldPosition.y;
-			Data0.z = TranslatedWorldPosition.z;
-			Data0.w = asfloat(WaterBodyIndex);
-
-			InstanceData0[InstanceDataIndex] = Data0;
+			IndirectArgs[BucketIndex * 5 + 4] = BucketInstanceDataOffset;
 		}
-		
-		// Should morph
+
+		// Write packed instance data
 		{
+			// Data0:
+			// uint NodeCoord.x : 11;
+			// uint NodeCoord.y : 11;
+			// uint LODLevel : 5;
+			// uint DensityIndex : 5;
+			//
+			// Data1:
+			// half WaterSurfaceBaseHeight;
+			// uint TileX : 8;
+			// uint TileY : 8;
+			//
+			// Data2:
+			// uint HeightLODFactorUnorm : 8;
+			// uint WaterBodyIndex : 24; // Could potentially use fewer bits for this to cram additional data in here in the future
+
 			const float3 ObserverPos = ObserverPosition - QuadTreePosition;
 			const int2 TextureLoadCoord = (int2)clamp(ObserverPos.xy / LeafSize, 0.0f, float2(QuadTreeResolutionX, QuadTreeResolutionY) - 0.5f);
 			const float WaterHeightForLOD = WaterZBoundsTexture.Load(int3(TextureLoadCoord, 0)).z * CaptureDepthRange;
 			const FWaterLODParams WaterLODParams = GetWaterLODParams(ObserverPos.z, WaterHeightForLOD, LODScale, NumLODs - 1);
-
 			const uint LogicalLODLevel = LODLevel + DensityIndex;
 
 			// Lowest LOD isn't always 0, this increases with the height distance 
 			const bool bIsLowestLOD = (LogicalLODLevel == WaterLODParams.LowestLOD);
-
-			// Only allow a tile to morph if it's not the last density level and not the last LOD level, sicne there is no next level to morph to
-			const uint bShouldMorph = (bLODMorphingEnabled && (DensityIndexClamped != NumDensities - 1)) ? 1 : 0;
-			// Tiles can morph twice to be able to morph between 3 LOD levels. Next to last density level can only morph once
-			const uint bCanMorphTwice = (DensityIndexClamped < NumDensities - 2) ? 1 : 0;
-
-			// Pack some of the data to save space. LOD level in the lower 8 bits and then bShouldMorph in the 9th bit and bCanMorphTwice in the 10th bit
-			const uint BitPackedChannel = ((uint)LogicalLODLevel & 0xFFu) | (bShouldMorph << 8u) | (bCanMorphTwice << 9u);
-
-			float4 Data1;
-			Data1.x = asfloat(BitPackedChannel);
-			Data1.y = bIsLowestLOD ? WaterLODParams.HeightLODFactor : 0.0f;
-			Data1.z = Scale;
-			Data1.w = Scale;
-
-			InstanceData1[InstanceDataIndex] = Data1;
+			const float HeightLODFactor = bIsLowestLOD ? WaterLODParams.HeightLODFactor : 0.0f;
+			const uint HeightLODFactorUnorm = uint(saturate(HeightLODFactor) * 255.0f);
+			const float WaterSurfaceBaseHeight = WaterZBoundsTexture.Load(int3(NodeCoord, LODLevel)).z;
+
+			const uint Data0 = (NodeCoord.x & 0x7FFu) | ((NodeCoord.y & 0x7FFu) << 11u) | ((LODLevel & 0x1Fu) << 22u) | ((DensityIndex & 0x1Fu) << 27u);
+			const uint WaterHeightF16 = f32tof16(WaterSurfaceBaseHeight);
+			const uint Data2 = (HeightLODFactorUnorm & 0xFFu) | (WaterBodyIndex << 8u);
+
+			// Render a single quadtree node by drawing one or multiple tile instances
+			uint TileX = 0;
+			uint TileY = 0;
+			for (uint InstanceIndex = 0; InstanceIndex < NumInstances; ++InstanceIndex)
+			{
+				const uint Data1 = WaterHeightF16 | ((TileX & 0xFFu) << 16u) | ((TileY & 0xFFu) << 24u);
+				InstanceData0[InstanceDataBaseIndex + InstanceIndex] = Data0;
+				InstanceData1[InstanceDataBaseIndex + InstanceIndex] = Data1;
+				InstanceData2[InstanceDataBaseIndex + InstanceIndex] = Data2;
+				
+				TileY = (TileX + 1) >= NumTilesPerEdge ? (TileY + 1) : TileY;
+				TileX = (TileX + 1) >= NumTilesPerEdge ? 0 : (TileX + 1);
+			}
 		}
 
 		// Instance Hit Proxy ID
 		if (bWithWaterSelectionSupport)
 		{
-			float4 Data2 = UnpackRGBA8(WBRenderData.HitProxyColorAndIsSelected).bgra;
-			Data2.w = Data2.w != 0.0f ? 1.0f : 0.0f;
-
-			InstanceData2[InstanceDataIndex] = Data2;
+			for (uint InstanceIndex = 0; InstanceIndex < NumInstances; ++InstanceIndex)
+			{
+				InstanceData3[InstanceDataBaseIndex + InstanceIndex] = WBRenderData.HitProxyColorAndIsSelected;
+			}
 		}
 	}
 }
diff --git a/Engine/Plugins/Experimental/Water/Source/Runtime/Private/WaterMeshSceneProxy.cpp b/Engine/Plugins/Experimental/Water/Source/Runtime/Private/WaterMeshSceneProxy.cpp
--- a/Engine/Plugins/Experimental/Water/Source/Runtime/Private/WaterMeshSceneProxy.cpp
+++ b/Engine/Plugins/Experimental/Water/Source/Runtime/Private/WaterMeshSceneProxy.cpp
@@ -63,6 +63,16 @@ static TAutoConsoleVariable<int32> CVarWaterMeshGPUQuadTreeConservativeRasteriza
 	TEXT("Enables software conservative rasterization for rasterizing water body meshes into the water quadtree. Disables jittered draws. Default: 0"),
 	ECVF_RenderThreadSafe);
 
+static TAutoConsoleVariable<int32> CVarWaterMeshGPUQuadTreeNumQuads(
+	TEXT("r.Water.WaterMesh.GPUQuadTree.NumQuadsPerTileSide"), 8,
+	TEXT("Number of quads per side of each tile mesh used to draw the water surface. A lower number results in more draw calls, a higher number in wasted VS invocations. Default: 8"),
+	ECVF_RenderThreadSafe);
+
+static TAutoConsoleVariable<float> CVarWaterMeshGPUQuadInstanceDataAllocMult(
+	TEXT("r.Water.WaterMesh.GPUQuadTree.InstanceDataAllocMult"), 1.0f,
+	TEXT("Multiplier to apply to the number of tiles in the quadtree at LOD0. The derived number is how many slots for water quad mesh instance data are allocated. Default: 1.0, Min: 0.0, Max: 8.0"),
+	ECVF_RenderThreadSafe);
+
 static TAutoConsoleVariable<int32> CVarWaterMeshOcclusionCullingMaxQueries(
 	TEXT("r.Water.WaterMesh.OcclusionCulling.MaxQueries"), 256,
 	TEXT("Maximum number of occlusion queries for the CPU water quadtree nodes. Using fewer queries than nodes will result in coarser culling."),
@@ -167,15 +177,43 @@ FWaterMeshSceneProxy::FWaterMeshSceneProxy(UWaterMeshComponent* Component)
 
 	int32 NumQuads = (int32)FMath::Pow(2.0f, (float)Component->GetTessellationFactor());
 	NumQuadsLOD0 = NumQuads;
+	NumQuadsPerIndirectDrawTile = FMath::Min((int32)FMath::RoundUpToPowerOfTwo(FMath::Clamp(CVarWaterMeshGPUQuadTreeNumQuads.GetValueOnGameThread(), 2, 128)), NumQuadsLOD0);
+
+	// Initialize Z bounds needed for GPU driven rendering
+	if (WaterQuadTree.IsGPUQuadTree())
+	{
+		WaterQuadTreeMinHeight = DBL_MAX;
+		WaterQuadTreeMaxHeight = -DBL_MAX;
+
+		const TArrayView<const FWaterBodyRenderData> WaterBodyRenderData = WaterQuadTree.GetWaterBodyRenderData();
+		for (const FWaterBodyRenderData& RenderData : WaterBodyRenderData)
+		{
+			WaterQuadTreeMinHeight = FMath::Min(WaterQuadTreeMinHeight, RenderData.BoundsMinZ - RenderData.MaxWaveHeight);
+			WaterQuadTreeMaxHeight = FMath::Max(WaterQuadTreeMaxHeight, RenderData.BoundsMaxZ + RenderData.MaxWaveHeight);
+		}
+
+		if (WaterQuadTreeMinHeight == DBL_MAX)
+		{
+			WaterQuadTreeMinHeight = 0.0;
+			WaterQuadTreeMaxHeight = 0.0;
+		}
+		WaterQuadTreeMinHeight -= 1.0;
+		WaterQuadTreeMaxHeight += 1.0;
+	}
+
+	DensityCount = FMath::Min(WaterQuadTree.GetTreeDepth(), (int32)FMath::FloorLog2(NumQuadsLOD0));
+
+	const FVector QuadTreeCorner = FVector(WaterQuadTree.GetTileRegion().Min, WaterQuadTreeMinHeight);
+	const float WaterQuadTreeDepthRange = WaterQuadTreeMaxHeight - WaterQuadTreeMinHeight;
+	const float LeafSize = WaterQuadTree.GetLeafSize();
+	WaterVertexFactoryIndirectDraw = new FWaterVertexFactoryIndirectDrawType(GetScene().GetFeatureLevel(), QuadTreeCorner, NumQuadsPerIndirectDrawTile, NumQuadsLOD0, DensityCount, LeafSize, LODScale, WaterQuadTreeDepthRange);
+	BeginInitResource(WaterVertexFactoryIndirectDraw);
 
 	WaterVertexFactories.Reserve(WaterQuadTree.GetTreeDepth());
-	WaterVertexFactoriesIndirectDraw.Reserve(WaterQuadTree.GetTreeDepth());
 	for (uint8 i = 0; i < WaterQuadTree.GetTreeDepth(); i++)
 	{
-		WaterVertexFactories.Add(new FWaterVertexFactoryType(GetScene().GetFeatureLevel(), NumQuads, LODScale));
+		WaterVertexFactories.Add(new FWaterVertexFactoryType(GetScene().GetFeatureLevel(), QuadTreeCorner, NumQuads, NumQuadsLOD0, DensityCount, LeafSize, LODScale, WaterQuadTreeDepthRange));
 		BeginInitResource(WaterVertexFactories.Last());
-		WaterVertexFactoriesIndirectDraw.Add(new FWaterVertexFactoryIndirectDrawType(GetScene().GetFeatureLevel(), NumQuads, LODScale));
-		BeginInitResource(WaterVertexFactoriesIndirectDraw.Last());
 
 		NumQuads /= 2;
 
@@ -187,8 +225,7 @@ FWaterMeshSceneProxy::FWaterMeshSceneProxy(UWaterMeshComponent* Component)
 	}
 
 	WaterVertexFactories.Shrink();
-	WaterVertexFactoriesIndirectDraw.Shrink();
-	DensityCount = WaterVertexFactories.Num();
+	check(DensityCount == WaterVertexFactories.Num());
 
 	const int32 TotalLeafNodes = WaterQuadTree.GetMaxLeafCount();
 	WaterInstanceDataBuffers = new FWaterInstanceDataBuffersType(TotalLeafNodes);
@@ -207,28 +244,6 @@ FWaterMeshSceneProxy::FWaterMeshSceneProxy(UWaterMeshComponent* Component)
 	RayTracingWaterData.SetNum(DensityCount);
 #endif
 
-	// Initialize Z bounds needed for GPU driven rendering
-	if (WaterQuadTree.IsGPUQuadTree())
-	{
-		WaterQuadTreeMinHeight = DBL_MAX;
-		WaterQuadTreeMaxHeight = -DBL_MAX;
-
-		const TArrayView<const FWaterBodyRenderData> WaterBodyRenderData = WaterQuadTree.GetWaterBodyRenderData();
-		for (const FWaterBodyRenderData& RenderData : WaterBodyRenderData)
-		{
-			WaterQuadTreeMinHeight = FMath::Min(WaterQuadTreeMinHeight, RenderData.BoundsMinZ - RenderData.MaxWaveHeight);
-			WaterQuadTreeMaxHeight = FMath::Max(WaterQuadTreeMaxHeight, RenderData.BoundsMaxZ + RenderData.MaxWaveHeight);
-		}
-
-		if (WaterQuadTreeMinHeight == DBL_MAX)
-		{
-			WaterQuadTreeMinHeight = 0.0;
-			WaterQuadTreeMaxHeight = 0.0;
-		}
-		WaterQuadTreeMinHeight -= 1.0;
-		WaterQuadTreeMaxHeight += 1.0;
-	}
-
 	// Always do CPU occlusion queries, even if this is a GPU quadtree. The GPU quadtree still potentially uses the far mesh which is CPU driven.
 	const int32 MaxQueries = CVarWaterMeshOcclusionCullingMaxQueries.GetValueOnGameThread();
 	const bool bIncludeFarMeshOcclusionQueries = CVarWaterMeshOcclusionCullingIncludeFarMesh.GetValueOnGameThread() != 0;
@@ -253,11 +268,8 @@ FWaterMeshSceneProxy::~FWaterMeshSceneProxy()
 		WaterFactory->ReleaseResource();
 		delete WaterFactory;
 	}
-	for (FWaterVertexFactoryIndirectDrawType* WaterFactory : WaterVertexFactoriesIndirectDraw)
-	{
-		WaterFactory->ReleaseResource();
-		delete WaterFactory;
-	}
+	WaterVertexFactoryIndirectDraw->ReleaseResource();
+	delete WaterVertexFactoryIndirectDraw;
 
 	delete WaterInstanceDataBuffers;
 
@@ -399,25 +411,26 @@ void FWaterMeshSceneProxy::GetDynamicMeshElements(const TArray<const FSceneView*
 
 		const int32 NumVisibleViews = VisibleViews.Num();
 		const int32 NumIndirectDrawCalls = NumBuckets * NumVisibleViews;
-		const int32 NumInstances = WaterQuadTree.GetMaxLeafCount() * NumVisibleViews;
+		const int32 LeafCountUpperBound = WaterQuadTree.GetMaxLeafCount() * NumVisibleViews;
+		const int32 NumInstanceDataSlots = FMath::Max(1, static_cast<int32>(LeafCountUpperBound * FMath::Clamp(CVarWaterMeshGPUQuadInstanceDataAllocMult.GetValueOnRenderThread(), 0.0f, 8.0f)));
 
 		// Allocate buffers for the indirect draws
 		struct FIndirectDrawResources
 		{
 			TRefCountPtr<FRDGPooledBuffer> IndirectArgs;
-			TRefCountPtr<FRDGPooledBuffer> InstanceDataOffsets;
 			TRefCountPtr<FRDGPooledBuffer> InstanceData0;
 			TRefCountPtr<FRDGPooledBuffer> InstanceData1;
 			TRefCountPtr<FRDGPooledBuffer> InstanceData2;
+			TRefCountPtr<FRDGPooledBuffer> InstanceData3;
 		};
 		FIndirectDrawResources& IndirectDrawResources = Collector.AllocateOneFrameResource<FIndirectDrawResources>();
 		IndirectDrawResources.IndirectArgs = AllocatePooledBuffer(FRDGBufferDesc::CreateIndirectDesc<FRHIDrawIndexedIndirectParameters>(FMath::Max(1, NumIndirectDrawCalls)), TEXT("WaterQuadTree.IndirectArgsBuffer"));
-		IndirectDrawResources.InstanceDataOffsets = AllocatePooledBuffer(FRDGBufferDesc::CreateBufferDesc(sizeof(uint32), FMath::Max(1, NumInstances)), TEXT("WaterQuadTree.InstanceDataOffsetsBuffer"));
-		IndirectDrawResources.InstanceData0 = AllocatePooledBuffer(FRDGBufferDesc::CreateBufferDesc(sizeof(FVector4f), FMath::Max(1, NumInstances)), TEXT("WaterQuadTree.InstanceDataBuffer0"));
-		IndirectDrawResources.InstanceData1 = AllocatePooledBuffer(FRDGBufferDesc::CreateBufferDesc(sizeof(FVector4f), FMath::Max(1, NumInstances)), TEXT("WaterQuadTree.InstanceDataBuffer1"));
+		IndirectDrawResources.InstanceData0 = AllocatePooledBuffer(FRDGBufferDesc::CreateBufferDesc(sizeof(uint32), FMath::Max(1, NumInstanceDataSlots)), TEXT("WaterQuadTree.InstanceDataBuffer0"));
+		IndirectDrawResources.InstanceData1 = AllocatePooledBuffer(FRDGBufferDesc::CreateBufferDesc(sizeof(uint32), FMath::Max(1, NumInstanceDataSlots)), TEXT("WaterQuadTree.InstanceDataBuffer1"));
+		IndirectDrawResources.InstanceData2 = AllocatePooledBuffer(FRDGBufferDesc::CreateBufferDesc(sizeof(uint32), FMath::Max(1, NumInstanceDataSlots)), TEXT("WaterQuadTree.InstanceDataBuffer2"));
 		if (WITH_WATER_SELECTION_SUPPORT != 0)
 		{
-			IndirectDrawResources.InstanceData2 = AllocatePooledBuffer(FRDGBufferDesc::CreateBufferDesc(sizeof(FVector4f), FMath::Max(1, NumInstances)), TEXT("WaterQuadTree.InstanceDataBuffer2"));
+			IndirectDrawResources.InstanceData3 = AllocatePooledBuffer(FRDGBufferDesc::CreateBufferDesc(sizeof(uint32), FMath::Max(1, NumInstanceDataSlots)), TEXT("WaterQuadTree.InstanceDataBuffer3"));
 		}
 		
 
@@ -426,10 +439,10 @@ void FWaterMeshSceneProxy::GetDynamicMeshElements(const TArray<const FSceneView*
 		{
 			WaterQuadTreeGPUTraverseParams = {};
 			WaterQuadTreeGPUTraverseParams.OutIndirectArgsBuffer = IndirectDrawResources.IndirectArgs;
-			WaterQuadTreeGPUTraverseParams.OutInstanceDataOffsetsBuffer = IndirectDrawResources.InstanceDataOffsets;
 			WaterQuadTreeGPUTraverseParams.OutInstanceData0Buffer = IndirectDrawResources.InstanceData0;
 			WaterQuadTreeGPUTraverseParams.OutInstanceData1Buffer = IndirectDrawResources.InstanceData1;
 			WaterQuadTreeGPUTraverseParams.OutInstanceData2Buffer = IndirectDrawResources.InstanceData2;
+			WaterQuadTreeGPUTraverseParams.OutInstanceData3Buffer = IndirectDrawResources.InstanceData3;
 			WaterQuadTreeGPUTraverseParams.Views = MoveTemp(VisibleViews);
 			WaterQuadTreeGPUTraverseParams.QuadTreePosition = FVector(WaterQuadTree.GetTileRegion().Min, WaterQuadTreeMinHeight);
 			WaterQuadTreeGPUTraverseParams.CullingBounds = WaterInfoBounds.ShiftBy(-WaterQuadTree.GetTileRegion().Min);
@@ -437,6 +450,7 @@ void FWaterMeshSceneProxy::GetDynamicMeshElements(const TArray<const FSceneView*
 			WaterQuadTreeGPUTraverseParams.NumMaterials = NumWaterMaterials;
 			WaterQuadTreeGPUTraverseParams.NumViews = NumVisibleViews;
 			WaterQuadTreeGPUTraverseParams.NumQuadsLOD0 = NumQuadsLOD0;
+			WaterQuadTreeGPUTraverseParams.NumQuadsPerTileSide = NumQuadsPerIndirectDrawTile;
 			WaterQuadTreeGPUTraverseParams.ForceCollapseDensityLevel = ForceCollapseDensityLevel;
 			WaterQuadTreeGPUTraverseParams.LeafSize = WaterQuadTree.GetLeafSize();
 			WaterQuadTreeGPUTraverseParams.LODScale = LODScale;
@@ -454,11 +468,10 @@ void FWaterMeshSceneProxy::GetDynamicMeshElements(const TArray<const FSceneView*
 		{
 			FWaterVertexFactoryUserDataWrapperType& UserDataWrapper = Collector.AllocateOneFrameResource<FWaterVertexFactoryUserDataWrapperType>();
 			UserDataWrapper.UserData.RenderGroupType = RenderGroup;
-			UserDataWrapper.UserData.IndirectInstanceDataOffsets = IndirectDrawResources.InstanceDataOffsets->GetSRV(RHICmdList, FRHIBufferSRVCreateInfo(PF_R32_UINT));
-			UserDataWrapper.UserData.IndirectInstanceData0 = IndirectDrawResources.InstanceData0->GetSRV(RHICmdList, FRHIBufferSRVCreateInfo(PF_A32B32G32R32F));
-			UserDataWrapper.UserData.IndirectInstanceData1 = IndirectDrawResources.InstanceData1->GetSRV(RHICmdList, FRHIBufferSRVCreateInfo(PF_A32B32G32R32F));
-			UserDataWrapper.UserData.IndirectInstanceData2 = (WITH_WATER_SELECTION_SUPPORT != 0) ? IndirectDrawResources.InstanceData2->GetSRV(RHICmdList, FRHIBufferSRVCreateInfo(PF_A32B32G32R32F)) : nullptr;
-			
+			UserDataWrapper.UserData.IndirectInstanceData0 = IndirectDrawResources.InstanceData0->GetRHI();
+			UserDataWrapper.UserData.IndirectInstanceData1 = IndirectDrawResources.InstanceData1->GetRHI();
+			UserDataWrapper.UserData.IndirectInstanceData2 = IndirectDrawResources.InstanceData2->GetRHI();
+			UserDataWrapper.UserData.IndirectInstanceData3 = (WITH_WATER_SELECTION_SUPPORT != 0) ? IndirectDrawResources.InstanceData3->GetRHI() : nullptr;
 			
 			UserDataWrappers[(int32)RenderGroup] = &UserDataWrapper;
 		}
@@ -477,9 +490,9 @@ void FWaterMeshSceneProxy::GetDynamicMeshElements(const TArray<const FSceneView*
 					TRACE_CPUPROFILER_EVENT_SCOPE(MaterialBucket);
 					bool bMaterialDrawn = false;
 
-					for (int32 DensityIndex = 0; DensityIndex < DensityCount; ++DensityIndex)
+					// We only render all tiles with one or multiple instances of a single density quad mesh, so no need to create one bucket per density.
 					{
-						const int32 BucketIndex = MaterialIndex * DensityCount + DensityIndex;
+						const int32 BucketIndex = MaterialIndex;
 
 						TRACE_CPUPROFILER_EVENT_SCOPE(DensityBucket);
 
@@ -501,7 +514,7 @@ void FWaterMeshSceneProxy::GetDynamicMeshElements(const TArray<const FSceneView*
 							// Set up mesh batch
 							FMeshBatch& Mesh = Collector.AllocateMesh();
 							Mesh.bWireframe = bWireframe;
-							Mesh.VertexFactory = WaterVertexFactoriesIndirectDraw[DensityIndex];
+							Mesh.VertexFactory = WaterVertexFactoryIndirectDraw;
 							Mesh.MaterialRenderProxy = MaterialRenderProxy;
 							Mesh.ReverseCulling = IsLocalToWorldDeterminantNegative();
 							Mesh.Type = PT_TriangleList;
@@ -533,7 +546,7 @@ void FWaterMeshSceneProxy::GetDynamicMeshElements(const TArray<const FSceneView*
 								BatchElement.UserData = (void*)&UserDataWrappers[(int32)RenderGroup]->UserData;
 								BatchElement.UserIndex = CompactViewIndex * NumBuckets + BucketIndex;
 
-								BatchElement.IndexBuffer = WaterVertexFactoriesIndirectDraw[DensityIndex]->IndexBuffer;
+								BatchElement.IndexBuffer = WaterVertexFactoryIndirectDraw->IndexBuffer;
 								BatchElement.PrimitiveIdMode = PrimID_ForceZero;
 
 								// We need the uniform buffer of this primitive because it stores the proper value for the bOutputVelocity flag.
diff --git a/Engine/Plugins/Experimental/Water/Source/Runtime/Private/WaterQuadTreeGPU.cpp b/Engine/Plugins/Experimental/Water/Source/Runtime/Private/WaterQuadTreeGPU.cpp
--- a/Engine/Plugins/Experimental/Water/Source/Runtime/Private/WaterQuadTreeGPU.cpp
+++ b/Engine/Plugins/Experimental/Water/Source/Runtime/Private/WaterQuadTreeGPU.cpp
@@ -152,12 +152,11 @@ public:
 	using FPermutationDomain = TShaderPermutationDomain<FPreciseOcclusionQueries>;
 
 	BEGIN_SHADER_PARAMETER_STRUCT(FParameters, )
-		SHADER_PARAMETER_RDG_BUFFER_UAV(RWStructuredBuffer, IndirectArgs)
-		SHADER_PARAMETER_RDG_BUFFER_UAV(RWStructuredBuffer, OcclusionQueryIndirectArgs)
-		SHADER_PARAMETER(uint32, NumDensities)
-		SHADER_PARAMETER(uint32, NumMaterials)
+		SHADER_PARAMETER_RDG_BUFFER_UAV(RWBuffer, IndirectArgs)
+		SHADER_PARAMETER_RDG_BUFFER_UAV(RWBuffer, OcclusionQueryIndirectArgs)
+		SHADER_PARAMETER(uint32, NumDrawBuckets)
 		SHADER_PARAMETER(uint32, NumViews)
-		SHADER_PARAMETER(uint32, NumQuadsLOD0)
+		SHADER_PARAMETER(uint32, NumQuads)
 	END_SHADER_PARAMETER_STRUCT()
 
 	static void ModifyCompilationEnvironment(const FGlobalShaderPermutationParameters& Parameters, FShaderCompilerEnvironment& OutEnvironment)
@@ -203,7 +202,7 @@ public:
 		SHADER_PARAMETER_RDG_BUFFER_UAV(RWByteAddressBuffer, PackedNodes)
 		SHADER_PARAMETER_RDG_BUFFER_UAV(RWBuffer, OcclusionQueryBoxes)
 		SHADER_PARAMETER_RDG_BUFFER_UAV(RWBuffer, OcclusionVisibility)
-		SHADER_PARAMETER_RDG_BUFFER_UAV(RWStructuredBuffer, OcclusionQueryIndirectArgs)
+		SHADER_PARAMETER_RDG_BUFFER_UAV(RWBuffer, OcclusionQueryIndirectArgs)
 		SHADER_PARAMETER_RDG_TEXTURE(Texture2D, QuadTreeTexture)
 		SHADER_PARAMETER_RDG_TEXTURE(Texture2D, WaterZBoundsTexture)
 		SHADER_PARAMETER_RDG_BUFFER_SRV(StructuredBuffer, WaterBodyRenderData)
@@ -216,8 +215,6 @@ public:
 		SHADER_PARAMETER(FVector3f, ObserverPosition)
 		SHADER_PARAMETER(uint32, QuadTreeResolutionX)
 		SHADER_PARAMETER(uint32, QuadTreeResolutionY)
-		SHADER_PARAMETER(uint32, NumDensities)
-		SHADER_PARAMETER(uint32, NumMaterials)
 		SHADER_PARAMETER(uint32, ViewIndex)
 		SHADER_PARAMETER(float, LeafSize)
 		SHADER_PARAMETER(float, LODScale)
@@ -295,9 +292,10 @@ public:
 		SHADER_PARAMETER_RDG_BUFFER_SRV(StructuredBuffer, WaterBodyRenderData)
 		SHADER_PARAMETER_RDG_BUFFER_SRV(ByteAddressBuffer, PackedNodes)
 		SHADER_PARAMETER_RDG_BUFFER_SRV(Buffer, OcclusionResults)
-		SHADER_PARAMETER(uint32, NumDensities)
-		SHADER_PARAMETER(uint32, NumMaterials)
 		SHADER_PARAMETER(uint32, NumDispatchedThreads)
+		SHADER_PARAMETER(uint32, NumDensities)
+		SHADER_PARAMETER(uint32, NumQuadsLOD0)
+		SHADER_PARAMETER(uint32, NumQuadsPerDraw)
 	END_SHADER_PARAMETER_STRUCT()
 
 	static void ModifyCompilationEnvironment(const FGlobalShaderPermutationParameters& Parameters, FShaderCompilerEnvironment& OutEnvironment)
@@ -342,10 +340,11 @@ public:
 	using FPermutationDomain = TShaderPermutationDomain<FPreciseOcclusionQueries>;
 
 	BEGIN_SHADER_PARAMETER_STRUCT(FParameters, )
-		SHADER_PARAMETER_RDG_BUFFER_UAV(RWStructuredBuffer, IndirectArgs)
+		SHADER_PARAMETER_RDG_BUFFER_UAV(RWBuffer, IndirectArgs)
 		SHADER_PARAMETER_RDG_BUFFER_UAV(RWBuffer, InstanceData0)
 		SHADER_PARAMETER_RDG_BUFFER_UAV(RWBuffer, InstanceData1)
 		SHADER_PARAMETER_RDG_BUFFER_UAV(RWBuffer, InstanceData2)
+		SHADER_PARAMETER_RDG_BUFFER_UAV(RWBuffer, InstanceData3)
 		SHADER_PARAMETER_RDG_TEXTURE(Texture2D, QuadTreeTexture)
 		SHADER_PARAMETER_RDG_TEXTURE(Texture2D, WaterZBoundsTexture)
 		SHADER_PARAMETER_RDG_BUFFER_SRV(StructuredBuffer, WaterBodyRenderData)
@@ -361,6 +360,8 @@ public:
 		SHADER_PARAMETER(uint32, NumDispatchedThreads)
 		SHADER_PARAMETER(uint32, BucketIndexOffset)
 		SHADER_PARAMETER(uint32, NumLODs)
+		SHADER_PARAMETER(uint32, NumQuadsLOD0)
+		SHADER_PARAMETER(uint32, NumQuadsPerDraw)
 		SHADER_PARAMETER(float, LeafSize)
 		SHADER_PARAMETER(float, LODScale)
 		SHADER_PARAMETER(float, CaptureDepthRange)
@@ -745,16 +746,14 @@ void FWaterQuadTreeGPU::Traverse(FRDGBuilder& GraphBuilder, const FTraverseParam
 	FRDGTexture* WaterZBoundsTextureRDG = GraphBuilder.RegisterExternalTexture(CreateRenderTarget(WaterZBoundsTexture, TEXT("WaterQuadTree.WaterSurfaceHeight")));
 	FRDGBuffer* WaterBodyRenderDataBufferRDG = GraphBuilder.RegisterExternalBuffer(WaterBodyRenderDataBuffer);
 	FRDGBuffer* IndirectArgsBuffer = GraphBuilder.RegisterExternalBuffer(Params.OutIndirectArgsBuffer);
-	FRDGBuffer* InstanceDataOffsetsBuffer = GraphBuilder.RegisterExternalBuffer(Params.OutInstanceDataOffsetsBuffer);
 	FRDGBuffer* InstanceData0Buffer = GraphBuilder.RegisterExternalBuffer(Params.OutInstanceData0Buffer);
 	FRDGBuffer* InstanceData1Buffer = GraphBuilder.RegisterExternalBuffer(Params.OutInstanceData1Buffer);
-	FRDGBuffer* InstanceData2Buffer = Params.bWithWaterSelectionSupport ? 
-		GraphBuilder.RegisterExternalBuffer(Params.OutInstanceData2Buffer) : 
-		GraphBuilder.CreateBuffer(FRDGBufferDesc::CreateBufferDesc(sizeof(float) * 4, 1), TEXT("WaterQuadTree.InstanceData2Dummy"));
+	FRDGBuffer* InstanceData2Buffer = GraphBuilder.RegisterExternalBuffer(Params.OutInstanceData2Buffer);
+	FRDGBuffer* InstanceData3Buffer = Params.bWithWaterSelectionSupport ? 
+		GraphBuilder.RegisterExternalBuffer(Params.OutInstanceData3Buffer) : 
+		GraphBuilder.CreateBuffer(FRDGBufferDesc::CreateBufferDesc(sizeof(uint32), 1), TEXT("WaterQuadTree.InstanceData3Dummy"));
 	
 	FRDGBufferUAV* IndirectArgsBufferUAV = GraphBuilder.CreateUAV(FRDGBufferUAVDesc(IndirectArgsBuffer, PF_R32_UINT));
-	FRDGBufferUAV* InstanceDataOffsetsBufferUAV = GraphBuilder.CreateUAV(FRDGBufferUAVDesc(InstanceDataOffsetsBuffer, PF_R32_UINT));
-	FRDGBufferSRV* InstanceDataOffsetsBufferSRV = GraphBuilder.CreateSRV(FRDGBufferSRVDesc(InstanceDataOffsetsBuffer, PF_R32_UINT));
 	FRDGBufferSRV* WaterBodyRenderDataBufferSRV = GraphBuilder.CreateSRV(FRDGBufferSRVDesc(WaterBodyRenderDataBufferRDG));
 
 	const EOcclusionQueryMode OcclusionQueryMode = static_cast<EOcclusionQueryMode>(GWaterQuadTreeOcclusionCulling);
@@ -773,6 +772,8 @@ void FWaterQuadTreeGPU::Traverse(FRDGBuilder& GraphBuilder, const FTraverseParam
 	const uint32 NumBucketsTotal = NumBucketsPerView * Params.NumViews;
 	const FIntPoint QuadTreeResolution = QuadTreeTextureRDG->Desc.Extent;
 
+	FRDGBuffer* PrefixSumBuffer = GraphBuilder.CreateBuffer(FRDGBufferDesc::CreateBufferDesc(sizeof(uint32), NumBucketsTotal), TEXT("WaterQuadTree.PrefixSum"));
+
 	// Initialize indirect args
 	{
 		FWaterQuadTreeInitializeIndirectArgsCS::FPermutationDomain PermutationDomain;
@@ -782,10 +783,9 @@ void FWaterQuadTreeGPU::Traverse(FRDGBuilder& GraphBuilder, const FTraverseParam
 		FWaterQuadTreeInitializeIndirectArgsCS::FParameters* PassParameters = GraphBuilder.AllocParameters<FWaterQuadTreeInitializeIndirectArgsCS::FParameters>();
 		PassParameters->IndirectArgs = IndirectArgsBufferUAV;
 		PassParameters->OcclusionQueryIndirectArgs = OcclusionQueryIndirectArgsBufferUAV;
-		PassParameters->NumDensities = Params.NumDensities;
-		PassParameters->NumMaterials = Params.NumMaterials;
+		PassParameters->NumDrawBuckets = Params.NumMaterials;
 		PassParameters->NumViews = Params.NumViews;
-		PassParameters->NumQuadsLOD0 = Params.NumQuadsLOD0;
+		PassParameters->NumQuads = Params.NumQuadsPerTileSide;
 		
 		FComputeShaderUtils::AddPass(GraphBuilder, RDG_EVENT_NAME("WaterQuadTreeInitIndirectArgs"), ComputeShader, PassParameters, FComputeShaderUtils::GetGroupCount(NumBucketsTotal, 64));
 	}
@@ -870,8 +870,6 @@ void FWaterQuadTreeGPU::Traverse(FRDGBuilder& GraphBuilder, const FTraverseParam
 			PassParameters->ObserverPosition = ObserverPositionTranslatedWorldSpace;
 			PassParameters->QuadTreeResolutionX = QuadTreeResolution.X;
 			PassParameters->QuadTreeResolutionY = QuadTreeResolution.Y;
-			PassParameters->NumDensities = Params.NumDensities;
-			PassParameters->NumMaterials = Params.NumMaterials;
 			PassParameters->ViewIndex = ViewIndex;
 			PassParameters->LeafSize = Params.LeafSize;
 			PassParameters->LODScale = Params.LODScale;
@@ -943,9 +941,10 @@ void FWaterQuadTreeGPU::Traverse(FRDGBuilder& GraphBuilder, const FTraverseParam
 			PassParameters->WaterBodyRenderData = WaterBodyRenderDataBufferSRV;
 			PassParameters->PackedNodes = PackedNodesSRV;
 			PassParameters->OcclusionResults = OcclusionQueryResultsSRV;
-			PassParameters->NumDensities = Params.NumDensities;
-			PassParameters->NumMaterials = Params.NumMaterials;
 			PassParameters->NumDispatchedThreads = FMath::Min(FMath::DivideAndRoundUp(MaxNumDraws, 64u), 65535u) * 64u; // The maximum number of dispatched groups is 64k
+			PassParameters->NumDensities = Params.NumDensities;
+			PassParameters->NumQuadsLOD0 = Params.NumQuadsLOD0;
+			PassParameters->NumQuadsPerDraw = Params.NumQuadsPerTileSide;
 
 			FComputeShaderUtils::AddPass(GraphBuilder, RDG_EVENT_NAME("WaterQuadTreeDrawBucketCounts(View: %i)", ViewIndex), ComputeShader, PassParameters, FComputeShaderUtils::GetGroupCount(PassParameters->NumDispatchedThreads, 64));
 		}
@@ -957,7 +956,7 @@ void FWaterQuadTreeGPU::Traverse(FRDGBuilder& GraphBuilder, const FTraverseParam
 			TShaderMapRef<FWaterQuadTreeBucketPrefixSumCS> ComputeShader(ShaderMap, PermutationVector);
 
 			FWaterQuadTreeBucketPrefixSumCS::FParameters* PassParameters = GraphBuilder.AllocParameters<FWaterQuadTreeBucketPrefixSumCS::FParameters>();
-			PassParameters->BucketPrefixSums = InstanceDataOffsetsBufferUAV;
+			PassParameters->BucketPrefixSums = GraphBuilder.CreateUAV(FRDGBufferUAVDesc(PrefixSumBuffer, PF_R32_UINT));
 			PassParameters->BucketCounts = GraphBuilder.CreateSRV(FRDGBufferSRVDesc(BucketCounts));
 			PassParameters->NumBuckets = NumBucketsPerView;
 			PassParameters->OutputOffset = ViewIndex * NumBucketsPerView;
@@ -973,14 +972,15 @@ void FWaterQuadTreeGPU::Traverse(FRDGBuilder& GraphBuilder, const FTraverseParam
 
 			FWaterQuadTreeGenerateInstanceDataCS::FParameters* PassParameters = GraphBuilder.AllocParameters<FWaterQuadTreeGenerateInstanceDataCS::FParameters>();
 			PassParameters->IndirectArgs = IndirectArgsBufferUAV;
-			PassParameters->InstanceData0 = GraphBuilder.CreateUAV(FRDGBufferUAVDesc(InstanceData0Buffer, PF_A32B32G32R32F));
-			PassParameters->InstanceData1 = GraphBuilder.CreateUAV(FRDGBufferUAVDesc(InstanceData1Buffer, PF_A32B32G32R32F));
-			PassParameters->InstanceData2 = GraphBuilder.CreateUAV(FRDGBufferUAVDesc(InstanceData2Buffer, PF_A32B32G32R32F));
+			PassParameters->InstanceData0 = GraphBuilder.CreateUAV(FRDGBufferUAVDesc(InstanceData0Buffer, PF_R32_UINT));
+			PassParameters->InstanceData1 = GraphBuilder.CreateUAV(FRDGBufferUAVDesc(InstanceData1Buffer, PF_R32_UINT));
+			PassParameters->InstanceData2 = GraphBuilder.CreateUAV(FRDGBufferUAVDesc(InstanceData2Buffer, PF_R32_UINT));
+			PassParameters->InstanceData3 = GraphBuilder.CreateUAV(FRDGBufferUAVDesc(InstanceData3Buffer, PF_R32_UINT));
 			PassParameters->QuadTreeTexture = QuadTreeTextureRDG;
 			PassParameters->WaterZBoundsTexture = WaterZBoundsTextureRDG;
 			PassParameters->WaterBodyRenderData = WaterBodyRenderDataBufferSRV;
 			PassParameters->PackedNodes = PackedNodesSRV;
-			PassParameters->InstanceDataOffsets = InstanceDataOffsetsBufferSRV;
+			PassParameters->InstanceDataOffsets = GraphBuilder.CreateSRV(FRDGBufferSRVDesc(PrefixSumBuffer, PF_R32_UINT));
 			PassParameters->OcclusionResults = OcclusionQueryResultsSRV;
 			PassParameters->QuadTreePosition = QuadTreePositionTranslatedWorldSpace;
 			PassParameters->ObserverPosition = ObserverPositionTranslatedWorldSpace;
@@ -991,6 +991,8 @@ void FWaterQuadTreeGPU::Traverse(FRDGBuilder& GraphBuilder, const FTraverseParam
 			PassParameters->NumDispatchedThreads = FMath::Min(FMath::DivideAndRoundUp(MaxNumDraws, 64u), 65535u) * 64u; // The maximum number of dispatched groups is 64k
 			PassParameters->BucketIndexOffset = ViewIndex * NumBucketsPerView;
 			PassParameters->NumLODs = QuadTreeTextureRDG->Desc.NumMips;
+			PassParameters->NumQuadsLOD0 = Params.NumQuadsLOD0;
+			PassParameters->NumQuadsPerDraw = Params.NumQuadsPerTileSide;
 			PassParameters->LeafSize = Params.LeafSize;
 			PassParameters->LODScale = Params.LODScale;
 			PassParameters->CaptureDepthRange = CaptureDepthRange;
diff --git a/Engine/Plugins/Experimental/Water/Source/Runtime/Private/WaterQuadTreeGPU.h b/Engine/Plugins/Experimental/Water/Source/Runtime/Private/WaterQuadTreeGPU.h
--- a/Engine/Plugins/Experimental/Water/Source/Runtime/Private/WaterQuadTreeGPU.h
+++ b/Engine/Plugins/Experimental/Water/Source/Runtime/Private/WaterQuadTreeGPU.h
@@ -61,10 +61,10 @@ public:
 	struct FTraverseParams
 	{
 		FRDGPooledBuffer* OutIndirectArgsBuffer = nullptr;
-		FRDGPooledBuffer* OutInstanceDataOffsetsBuffer = nullptr;
 		FRDGPooledBuffer* OutInstanceData0Buffer = nullptr;
 		FRDGPooledBuffer* OutInstanceData1Buffer = nullptr;
 		FRDGPooledBuffer* OutInstanceData2Buffer = nullptr;
+		FRDGPooledBuffer* OutInstanceData3Buffer = nullptr;
 		TArray<const FSceneView*> Views;
 		FVector QuadTreePosition;
 		FBox2D CullingBounds = FBox2D(ForceInit);
@@ -72,6 +72,7 @@ public:
 		uint32 NumMaterials = 0;
 		uint32 NumViews = 0;
 		uint32 NumQuadsLOD0 = 0;
+		uint32 NumQuadsPerTileSide = 0;
 		int32 ForceCollapseDensityLevel = -1;
 		float LeafSize = 1.0f;
 		float LODScale = 1.0f;
diff --git a/Engine/Plugins/Experimental/Water/Source/Runtime/Private/WaterVertexFactory.cpp b/Engine/Plugins/Experimental/Water/Source/Runtime/Private/WaterVertexFactory.cpp
--- a/Engine/Plugins/Experimental/Water/Source/Runtime/Private/WaterVertexFactory.cpp
+++ b/Engine/Plugins/Experimental/Water/Source/Runtime/Private/WaterVertexFactory.cpp
@@ -28,15 +28,8 @@ public:
 	{
 		if (bIndirectDraws)
 		{
-			InstanceDataOffsetsBufferParameter.Bind(ParameterMap, TEXT("InstanceDataOffsetsBuffer"));
-			InstanceData0BufferParameter.Bind(ParameterMap, TEXT("InstanceData0Buffer"));
-			InstanceData1BufferParameter.Bind(ParameterMap, TEXT("InstanceData1Buffer"));
-			if (bWithWaterSelectionSupport)
-			{
-				InstanceData2BufferParameter.Bind(ParameterMap, TEXT("InstanceData2Buffer"));
-			}
-			
-			DrawBucketIndexParameter.Bind(ParameterMap, TEXT("DrawBucketIndex"));
+			QuadTreePositionParameter.Bind(ParameterMap, TEXT("QuadTreePosition"));
+			LODMorphingEnabledParameter.Bind(ParameterMap, TEXT("bLODMorphingEnabled"));
 		}
 	}
 
@@ -66,13 +59,43 @@ public:
 		}
 #endif
 
-		if (VertexStreams.Num() > 0 && !bIndirectDraws)
+		if (bIndirectDraws)
+		{
+			if (VertexStreams.Num() > 0)
+			{
+				const int32 NumBuffers = bWithWaterSelectionSupport ? 4 : 3;
+				FRHIBuffer* InstanceVertexBuffers[] = 
+				{ 
+					WaterMeshUserData->IndirectInstanceData0, 
+					WaterMeshUserData->IndirectInstanceData1,
+					WaterMeshUserData->IndirectInstanceData2,
+					WaterMeshUserData->IndirectInstanceData3,
+				};
+				for (int32 i = 0; i < NumBuffers; ++i)
+				{
+					FVertexInputStream* InstanceInputStream = VertexStreams.FindByPredicate([i](const FVertexInputStream& InStream) { return InStream.StreamIndex == 1 + i; });
+					check(InstanceInputStream);
+
+					// Bind vertex buffer
+					check(InstanceVertexBuffers[i]);
+					InstanceInputStream->VertexBuffer = InstanceVertexBuffers[i];
+				}
+			}
+
+			const FVector PreViewTranslation = View->ViewMatrices.GetPreViewTranslation();
+			ShaderBindings.Add(QuadTreePositionParameter, FVector3f(PreViewTranslation + VertexFactory->GetQuadTreePositionWS()));
+
+			static const auto CVar = IConsoleManager::Get().FindTConsoleVariableDataInt(TEXT("r.Water.WaterMesh.LODMorphEnabled"));
+			const bool bLODMorphingEnabled = CVar && CVar->GetValueOnRenderThread() != 0;
+			ShaderBindings.Add(LODMorphingEnabledParameter, bLODMorphingEnabled ? 1 : 0);
+		}
+		else if (VertexStreams.Num() > 0)
 		{
 			for (int32 i = 0; i < WaterInstanceDataBuffersType::NumBuffers; ++i)
 			{
-				FVertexInputStream* InstanceInputStream = VertexStreams.FindByPredicate([i](const FVertexInputStream& InStream) { return InStream.StreamIndex == i+1; });
+				FVertexInputStream* InstanceInputStream = VertexStreams.FindByPredicate([i](const FVertexInputStream& InStream) { return InStream.StreamIndex == i + 1; });
 				check(InstanceInputStream);
-				
+
 				// Bind vertex buffer
 				check(InstanceDataBuffers->GetBuffer(i));
 				InstanceInputStream->VertexBuffer = InstanceDataBuffers->GetBuffer(i);
@@ -84,27 +107,11 @@ public:
 				VertexFactory->OffsetInstanceStreams(InstanceOffsetValue, InputStreamType, VertexStreams);
 			}
 		}
-
-		if (bIndirectDraws)
-		{
-			const uint32 DrawBucketIndex = BatchElement.UserIndex;
-			ShaderBindings.Add(InstanceDataOffsetsBufferParameter, WaterMeshUserData->IndirectInstanceDataOffsets);
-			ShaderBindings.Add(InstanceData0BufferParameter, WaterMeshUserData->IndirectInstanceData0);
-			ShaderBindings.Add(InstanceData1BufferParameter, WaterMeshUserData->IndirectInstanceData1);
-			if (bWithWaterSelectionSupport)
-			{
-				ShaderBindings.Add(InstanceData2BufferParameter, WaterMeshUserData->IndirectInstanceData2);
-			}
-			ShaderBindings.Add(DrawBucketIndexParameter, DrawBucketIndex);
-		}
 	}
 
 private:
-	LAYOUT_FIELD(FShaderResourceParameter, InstanceDataOffsetsBufferParameter);
-	LAYOUT_FIELD(FShaderResourceParameter, InstanceData0BufferParameter);
-	LAYOUT_FIELD(FShaderResourceParameter, InstanceData1BufferParameter);
-	LAYOUT_FIELD(FShaderResourceParameter, InstanceData2BufferParameter);
-	LAYOUT_FIELD(FShaderParameter, DrawBucketIndexParameter);
+	LAYOUT_FIELD(FShaderParameter, QuadTreePositionParameter);
+	LAYOUT_FIELD(FShaderParameter, LODMorphingEnabledParameter);
 };
 
 // ----------------------------------------------------------------------------------
diff --git a/Engine/Plugins/Experimental/Water/Source/Runtime/Public/WaterMeshSceneProxy.h b/Engine/Plugins/Experimental/Water/Source/Runtime/Public/WaterMeshSceneProxy.h
--- a/Engine/Plugins/Experimental/Water/Source/Runtime/Public/WaterMeshSceneProxy.h
+++ b/Engine/Plugins/Experimental/Water/Source/Runtime/Public/WaterMeshSceneProxy.h
@@ -106,7 +106,7 @@ private:
 
 	// One vertex factory per LOD
 	TArray<FWaterVertexFactoryType*> WaterVertexFactories;
-	TArray<FWaterVertexFactoryIndirectDrawType*> WaterVertexFactoriesIndirectDraw;
+	FWaterVertexFactoryIndirectDrawType* WaterVertexFactoryIndirectDraw;
 
 	/** Tiles containing water, stored in a quad tree */
 	FWaterQuadTree WaterQuadTree;
@@ -132,6 +132,8 @@ private:
 	/** Number of quads per side of a water quad tree tile at LOD0 */
 	int32 NumQuadsLOD0 = 0;
 
+	int32 NumQuadsPerIndirectDrawTile = 0;
+
 	/** Number of densities (same as number of grid index/vertex buffers) */
 	int32 DensityCount = 0;
 
diff --git a/Engine/Plugins/Experimental/Water/Source/Runtime/Public/WaterVertexFactory.h b/Engine/Plugins/Experimental/Water/Source/Runtime/Public/WaterVertexFactory.h
--- a/Engine/Plugins/Experimental/Water/Source/Runtime/Public/WaterVertexFactory.h
+++ b/Engine/Plugins/Experimental/Water/Source/Runtime/Public/WaterVertexFactory.h
@@ -14,9 +14,14 @@ struct FShaderCompilerEnvironment;
  */
 BEGIN_GLOBAL_SHADER_PARAMETER_STRUCT(FWaterVertexFactoryParameters, )
 	SHADER_PARAMETER(float, LODScale)
+	SHADER_PARAMETER(float, LeafSize)
+	SHADER_PARAMETER(float, CaptureDepthRange)
 	SHADER_PARAMETER(int32, NumQuadsPerTileSide)
+	SHADER_PARAMETER(int32, NumQuadsLOD0)
+	SHADER_PARAMETER(int32, NumDensities)
 	SHADER_PARAMETER(int32, bRenderSelected)
 	SHADER_PARAMETER(int32, bRenderUnselected)
+	SHADER_PARAMETER(int32, bLODMorphingEnabled)
 END_GLOBAL_SHADER_PARAMETER_STRUCT()
 
 typedef TUniformBufferRef<FWaterVertexFactoryParameters> FWaterVertexFactoryBufferRef;
@@ -176,7 +181,7 @@ public:
 	static constexpr int32 NumRenderGroups = bWithWaterSelectionSupport ? 3 : 1; // Must match EWaterMeshRenderGroupType
 	static constexpr int32 NumAdditionalVertexStreams = bIndirectDraws ? 0 : TWaterInstanceDataBuffers<bWithWaterSelectionSupport>::NumBuffers;
 
-	TWaterVertexFactory(ERHIFeatureLevel::Type InFeatureLevel, int32 InNumQuadsPerSide,	float InLODScale);
+	TWaterVertexFactory(ERHIFeatureLevel::Type InFeatureLevel, const FVector& InQuadTreePositionWS, int32 InNumQuadsPerSide, int32 InNumQuadsLOD0, int32 InNumDensities, float InLeafSize, float InLODScale, float InCaptureDepthRange);
 	~TWaterVertexFactory();
 
 	/**
@@ -202,6 +207,8 @@ public:
 
 	inline const FUniformBufferRHIRef GetWaterVertexFactoryUniformBuffer(EWaterMeshRenderGroupType InRenderGroupType) const { return UniformBuffers[(int32)InRenderGroupType]; }
 
+	inline FVector GetQuadTreePositionWS() const { return QuadTreePositionWS; }
+
 private:
 	void SetupUniformDataForGroup(EWaterMeshRenderGroupType InRenderGroupType);
 
@@ -212,8 +219,13 @@ public:
 private:
 	TStaticArray<FWaterVertexFactoryBufferRef, NumRenderGroups> UniformBuffers;
 
+	const FVector QuadTreePositionWS = FVector::ZeroVector;
 	const int32 NumQuadsPerSide = 0;
+	const int32 NumQuadsLOD0 = 0;
+	const int32 NumDensities = 0;
+	const float LeafSize = 0.0f;
 	const float LODScale = 0.0f;
+	const float CaptureDepthRange = 0.0f;
 };
 
 extern const FVertexFactoryType* GetWaterVertexFactoryType(bool bWithWaterSelectionSupport, bool bIndirectDraws);
@@ -240,10 +252,10 @@ struct TWaterMeshUserData
 	FUniformBufferRHIRef WaterVertexFactoryRaytracingVFUniformBuffer = nullptr;
 #endif
 
-	FRHIShaderResourceView* IndirectInstanceDataOffsets = nullptr;
-	FRHIShaderResourceView* IndirectInstanceData0 = nullptr;
-	FRHIShaderResourceView* IndirectInstanceData1 = nullptr;
-	FRHIShaderResourceView* IndirectInstanceData2 = nullptr;
+	FRHIBuffer* IndirectInstanceData0 = nullptr;
+	FRHIBuffer* IndirectInstanceData1 = nullptr;
+	FRHIBuffer* IndirectInstanceData2 = nullptr;
+	FRHIBuffer* IndirectInstanceData3 = nullptr;
 };
 
 /**
diff --git a/Engine/Plugins/Experimental/Water/Source/Runtime/Public/WaterVertexFactory.inl b/Engine/Plugins/Experimental/Water/Source/Runtime/Public/WaterVertexFactory.inl
--- a/Engine/Plugins/Experimental/Water/Source/Runtime/Public/WaterVertexFactory.inl
+++ b/Engine/Plugins/Experimental/Water/Source/Runtime/Public/WaterVertexFactory.inl
@@ -14,10 +14,15 @@
 // ----------------------------------------------------------------------------------
 
 template <bool bWithWaterSelectionSupport, bool bIndirectDraws>
-TWaterVertexFactory<bWithWaterSelectionSupport, bIndirectDraws>::TWaterVertexFactory(ERHIFeatureLevel::Type InFeatureLevel, int32 InNumQuadsPerSide, float InLODScale)
+TWaterVertexFactory<bWithWaterSelectionSupport, bIndirectDraws>::TWaterVertexFactory(ERHIFeatureLevel::Type InFeatureLevel, const FVector& InQuadTreePositionWS, int32 InNumQuadsPerSide, int32 InNumQuadsLOD0, int32 InNumDensities, float InLeafSize, float InLODScale, float InCaptureDepthRange)
 	: FVertexFactory(InFeatureLevel)
+	, QuadTreePositionWS(InQuadTreePositionWS)
 	, NumQuadsPerSide(InNumQuadsPerSide)
+	, NumQuadsLOD0(InNumQuadsLOD0)
+	, NumDensities(InNumDensities)
+	, LeafSize(InLeafSize)
 	, LODScale(InLODScale)
+	, CaptureDepthRange(InCaptureDepthRange)
 {
 	VertexBuffer = new FWaterMeshVertexBuffer(NumQuadsPerSide);
 	IndexBuffer = new FWaterMeshIndexBuffer(NumQuadsPerSide);
@@ -56,13 +61,6 @@ void TWaterVertexFactory<bWithWaterSelectionSupport, bIndirectDraws>::InitRHI(FR
 	PositionVertexStream.Offset = 0;
 	PositionVertexStream.VertexStreamUsage = EVertexStreamUsage::Default;
 
-	// Simple instancing vertex stream with nullptr vertex buffer to be set at binding time
-	FVertexStream InstanceDataVertexStream;
-	InstanceDataVertexStream.VertexBuffer = nullptr;
-	InstanceDataVertexStream.Stride = sizeof(FVector4f);
-	InstanceDataVertexStream.Offset = 0;
-	InstanceDataVertexStream.VertexStreamUsage = EVertexStreamUsage::Instancing;
-
 	FVertexElement VertexPositionElement(Streams.Add(PositionVertexStream), 0, VET_Float4, 0, PositionVertexStream.Stride, false);
 
 	// Vertex declaration
@@ -70,8 +68,30 @@ void TWaterVertexFactory<bWithWaterSelectionSupport, bIndirectDraws>::InitRHI(FR
 	Elements.Add(VertexPositionElement);
 
 	// Adds all streams
-	if constexpr (NumAdditionalVertexStreams > 0)
+	if constexpr (bIndirectDraws)
+	{
+		FVertexStream InstanceDataVertexStream;
+		InstanceDataVertexStream.VertexBuffer = nullptr;
+		InstanceDataVertexStream.Stride = sizeof(uint32);
+		InstanceDataVertexStream.Offset = 0;
+		InstanceDataVertexStream.VertexStreamUsage = EVertexStreamUsage::Instancing;
+
+		constexpr int NumBuffers = bWithWaterSelectionSupport ? 4 : 3;
+
+		for (int i = 0; i < NumBuffers; ++i)
+		{
+			Elements.Add(FVertexElement(Streams.Add(InstanceDataVertexStream), 0, VET_UInt, 8 + i, InstanceDataVertexStream.Stride, true));
+		}
+	}
+	else if constexpr (NumAdditionalVertexStreams > 0)
 	{
+		// Simple instancing vertex stream with nullptr vertex buffer to be set at binding time
+		FVertexStream InstanceDataVertexStream;
+		InstanceDataVertexStream.VertexBuffer = nullptr;
+		InstanceDataVertexStream.Stride = sizeof(FVector4f);
+		InstanceDataVertexStream.Offset = 0;
+		InstanceDataVertexStream.VertexStreamUsage = EVertexStreamUsage::Instancing;
+
 		for (int32 StreamIdx = 0; StreamIdx < NumAdditionalVertexStreams; ++StreamIdx)
 		{
 			FVertexElement InstanceElement(Streams.Add(InstanceDataVertexStream), 0, VET_Float4, 8 + StreamIdx, InstanceDataVertexStream.Stride, true);
@@ -122,7 +142,11 @@ void TWaterVertexFactory<bWithWaterSelectionSupport, bIndirectDraws>::SetupUnifo
 {
 	FWaterVertexFactoryParameters UniformParams;
 	UniformParams.NumQuadsPerTileSide = NumQuadsPerSide;
+	UniformParams.NumQuadsLOD0 = NumQuadsLOD0;
+	UniformParams.NumDensities = NumDensities;
 	UniformParams.LODScale = LODScale;
+	UniformParams.LeafSize = LeafSize;
+	UniformParams.CaptureDepthRange = CaptureDepthRange;
 	UniformParams.bRenderSelected = true;
 	UniformParams.bRenderUnselected = true;
 
@@ -171,6 +195,11 @@ void TWaterVertexFactory<bWithWaterSelectionSupport, bIndirectDraws>::ModifyComp
 	}
 
 	OutEnvironment.SetDefine(TEXT("RAY_TRACING_DYNAMIC_MESH_IN_LOCAL_SPACE"), TEXT("1"));
+
+	if (bIndirectDraws)
+	{
+		OutEnvironment.CompilerFlags.Add(CFLAG_IndirectDraw);
+	}
 }
 
 template <bool bWithWaterSelectionSupport, bool bIndirectDraws>
@@ -193,7 +222,17 @@ void TWaterVertexFactory<bWithWaterSelectionSupport, bIndirectDraws>::GetPSOPrec
 	Elements.Add(FVertexElement(0, 0, VET_Float4, 0, sizeof(FVector4f), false));
 
 	// Add all the additional streams
-	if constexpr (NumAdditionalVertexStreams > 0)
+	if constexpr (bIndirectDraws)
+	{
+		Elements.Add(FVertexElement(1, 0, VET_UInt, 8, sizeof(uint32), true));
+		Elements.Add(FVertexElement(2, 0, VET_UInt, 9, sizeof(uint32), true));
+		Elements.Add(FVertexElement(3, 0, VET_UInt, 10, sizeof(uint32), true));
+		if (bWithWaterSelectionSupport)
+		{
+			Elements.Add(FVertexElement(4, 0, VET_UInt, 11, sizeof(uint32), true));
+		}
+	}
+	else if constexpr (NumAdditionalVertexStreams > 0)
 	{
 		for (int32 StreamIdx = 0; StreamIdx < NumAdditionalVertexStreams; ++StreamIdx)
 		{
