Fixed OpenGL RHI warning (creating an SRV of only a subset of mip levels is unsupported) when downsampling the GPU water quadtree texture. Fixed by binding the entire mip chain and only loading a specific mip level in the shader on platforms where we can't create views of just a single mip. This approach has been used in other places in the engine and seems to be legal OpenGL.

#rnx

diff --git a/Engine/Plugins/Experimental/Water/Shaders/Private/WaterQuadTreeBuild.usf b/Engine/Plugins/Experimental/Water/Shaders/Private/WaterQuadTreeBuild.usf
--- a/Engine/Plugins/Experimental/Water/Shaders/Private/WaterQuadTreeBuild.usf
+++ b/Engine/Plugins/Experimental/Water/Shaders/Private/WaterQuadTreeBuild.usf
@@ -6,6 +6,7 @@
 
 Texture2D<float4> QuadTreeTexture;
 Texture2D<float4> WaterZBoundsTexture;
+int InputMipLevelIndex;
 
 void Main(in float4 Position : SV_Position, out float4 OutColor : SV_Target0, out float4 OutWaterZBounds : SV_Target1)
 {
@@ -25,7 +26,7 @@ void Main(in float4 Position : SV_Position, out float4 OutColor : SV_Target0, ou
 	{
 		for (int x = 0; x < 2; ++x)
 		{
-			const float4 QuadTreeSample = QuadTreeTexture.Load(int3(BaseCoord + int2(x, y), 0));
+			const float4 QuadTreeSample = QuadTreeTexture.Load(int3(BaseCoord + int2(x, y), InputMipLevelIndex));
 			const FWaterQuadTreeNode ChildNode = WaterQuadTreeUnpackNodeRGBA8(QuadTreeSample);
 			
 			bHasCompleteSubtree = bHasCompleteSubtree && ChildNode.bHasCompleteSubtree;
@@ -34,7 +35,7 @@ void Main(in float4 Position : SV_Position, out float4 OutColor : SV_Target0, ou
 			MinTransitionWaterBodyRenderDataIndex = min(MinTransitionWaterBodyRenderDataIndex, ChildNode.TransitionWaterBodyRenderDataIndex);
 			MaxTransitionWaterBodyRenderDataIndex = max(MaxTransitionWaterBodyRenderDataIndex, ChildNode.TransitionWaterBodyRenderDataIndex);
 
-			const float3 WaterZBoundsSample = WaterZBoundsTexture.Load(int3(BaseCoord + int2(x, y), 0)).xyz;
+			const float3 WaterZBoundsSample = WaterZBoundsTexture.Load(int3(BaseCoord + int2(x, y), InputMipLevelIndex)).xyz;
 			MinZ = min(MinZ, WaterZBoundsSample.x);
 			MaxZ = max(MaxZ, WaterZBoundsSample.y);
 			SurfaceBaseHeightSum += WaterZBoundsSample.z;
diff --git a/Engine/Plugins/Experimental/Water/Source/Runtime/Private/WaterQuadTreeGPU.cpp b/Engine/Plugins/Experimental/Water/Source/Runtime/Private/WaterQuadTreeGPU.cpp
--- a/Engine/Plugins/Experimental/Water/Source/Runtime/Private/WaterQuadTreeGPU.cpp
+++ b/Engine/Plugins/Experimental/Water/Source/Runtime/Private/WaterQuadTreeGPU.cpp
@@ -132,6 +132,7 @@ public:
 	BEGIN_SHADER_PARAMETER_STRUCT(FParameters, )
 		SHADER_PARAMETER_RDG_TEXTURE_SRV(Texture2D, QuadTreeTexture)
 		SHADER_PARAMETER_RDG_TEXTURE_SRV(Texture2D, WaterZBoundsTexture)
+		SHADER_PARAMETER(int32, InputMipLevelIndex)
 		RENDER_TARGET_BINDING_SLOTS()
 	END_SHADER_PARAMETER_STRUCT()
 
@@ -711,13 +712,23 @@ void FWaterQuadTreeGPU::Init(FRDGBuilder& GraphBuilder, const FInitParams& Param
 	{
 		TShaderMapRef<FWaterQuadTreeBuildPS> PixelShader(ShaderMap);
 
+		const bool bSupportsTextureViews = GRHISupportsTextureViews;
 		FIntRect BuildViewport(0, 0, QuadTreeResolution.X, QuadTreeResolution.Y);
 
 		for (int32 MipLevel = 1; MipLevel < NumMipLevels; ++MipLevel)
 		{
+			FRDGTextureSRVDesc QuadTreeSourceSRVDesc = FRDGTextureSRVDesc::CreateForMipLevel(QuadTreeTextureRDG, MipLevel - 1);
+			FRDGTextureSRVDesc WaterZBoundsSourceSRVDesc = FRDGTextureSRVDesc::CreateForMipLevel(WaterZBoundsTextureRDG, MipLevel - 1);
+			if (!bSupportsTextureViews)
+			{
+				QuadTreeSourceSRVDesc = FRDGTextureSRVDesc::Create(QuadTreeTextureRDG);
+				WaterZBoundsSourceSRVDesc = FRDGTextureSRVDesc::Create(WaterZBoundsTextureRDG);
+			}
+
 			FWaterQuadTreeBuildPS::FParameters* PassParameters = GraphBuilder.AllocParameters<FWaterQuadTreeBuildPS::FParameters>();
-			PassParameters->QuadTreeTexture = GraphBuilder.CreateSRV(FRDGTextureSRVDesc::CreateForMipLevel(QuadTreeTextureRDG, MipLevel - 1));
-			PassParameters->WaterZBoundsTexture = GraphBuilder.CreateSRV(FRDGTextureSRVDesc::CreateForMipLevel(WaterZBoundsTextureRDG, MipLevel - 1));
+			PassParameters->QuadTreeTexture = GraphBuilder.CreateSRV(QuadTreeSourceSRVDesc);
+			PassParameters->WaterZBoundsTexture = GraphBuilder.CreateSRV(WaterZBoundsSourceSRVDesc);
+			PassParameters->InputMipLevelIndex = bSupportsTextureViews ? 0 : (MipLevel - 1);
 			PassParameters->RenderTargets[0] = FRenderTargetBinding(QuadTreeTextureRDG, ERenderTargetLoadAction::ENoAction, MipLevel);
 			PassParameters->RenderTargets[1] = FRenderTargetBinding(WaterZBoundsTextureRDG, ERenderTargetLoadAction::ENoAction, MipLevel);
 
