Fix for missing resource transitions on the water info texture and indirect draw buffers of the new GPU water quadtree.

#rnx

diff --git a/Engine/Plugins/Experimental/Water/Source/Runtime/Private/WaterInfoRendering.cpp b/Engine/Plugins/Experimental/Water/Source/Runtime/Private/WaterInfoRendering.cpp
--- a/Engine/Plugins/Experimental/Water/Source/Runtime/Private/WaterInfoRendering.cpp
+++ b/Engine/Plugins/Experimental/Water/Source/Runtime/Private/WaterInfoRendering.cpp
@@ -980,6 +980,7 @@ public:
 
 		FRDGTextureRef WaterInfoTexture = RegisterExternalTexture(GraphBuilder, WaterInfoRenderTarget->GetRenderTargetTexture(), TEXT("WaterInfoTexture"));
 		AddCopyTexturePass(GraphBuilder, FinalizedTexture, WaterInfoTexture);
+		GraphBuilder.UseExternalAccessMode(WaterInfoTexture, ERHIAccess::SRVMask);
 	}
 
 	FRenderTarget* WaterInfoRenderTarget = nullptr;
diff --git a/Engine/Plugins/Experimental/Water/Source/Runtime/Private/WaterQuadTreeGPU.cpp b/Engine/Plugins/Experimental/Water/Source/Runtime/Private/WaterQuadTreeGPU.cpp
--- a/Engine/Plugins/Experimental/Water/Source/Runtime/Private/WaterQuadTreeGPU.cpp
+++ b/Engine/Plugins/Experimental/Water/Source/Runtime/Private/WaterQuadTreeGPU.cpp
@@ -1049,4 +1049,13 @@ void FWaterQuadTreeGPU::Traverse(FRDGBuilder& GraphBuilder, const FTraverseParam
 			}
 		}
 	}
+
+	GraphBuilder.UseExternalAccessMode(IndirectArgsBuffer, ERHIAccess::IndirectArgs);
+	GraphBuilder.UseExternalAccessMode(InstanceData0Buffer, ERHIAccess::VertexOrIndexBuffer);
+	GraphBuilder.UseExternalAccessMode(InstanceData1Buffer, ERHIAccess::VertexOrIndexBuffer);
+	GraphBuilder.UseExternalAccessMode(InstanceData2Buffer, ERHIAccess::VertexOrIndexBuffer);
+	if (Params.bWithWaterSelectionSupport)
+	{
+		GraphBuilder.UseExternalAccessMode(InstanceData3Buffer, ERHIAccess::VertexOrIndexBuffer);
+	}
 }
