Fix for GLSL cross compile issue in the indirect draw arg init shader for the GPU water quadtree. Worked around the issue by renaming a buffer parameter.

#rnx

diff --git a/Engine/Plugins/Experimental/Water/Shaders/Private/WaterQuadTreeDraws.usf b/Engine/Plugins/Experimental/Water/Shaders/Private/WaterQuadTreeDraws.usf
--- a/Engine/Plugins/Experimental/Water/Shaders/Private/WaterQuadTreeDraws.usf
+++ b/Engine/Plugins/Experimental/Water/Shaders/Private/WaterQuadTreeDraws.usf
@@ -59,7 +59,7 @@ float4 GetNodeAABB2D(uint3 InNodeCoord, float3 InQuadTreePosition, float InLeafS
 
 RWBuffer<uint> IndirectArgs;
 #if PRECISE_OCCLUSION_QUERIES
-RWBuffer<uint> OcclusionQueryIndirectArgs;
+RWBuffer<uint> OcclusionQueryArgs;
 #endif
 uint NumDrawBuckets;
 uint NumViews;
@@ -92,11 +92,11 @@ void MainCS(uint3 DispatchThreadId : SV_DispatchThreadID)
 	{
 		for (uint ViewIndex = 0; ViewIndex < NumViews; ++ViewIndex)
 		{
-			OcclusionQueryIndirectArgs[ViewIndex * 5 + 0] = 36;// IndexCountPerInstance; 12 triangles per cube
-			OcclusionQueryIndirectArgs[ViewIndex * 5 + 1] = 0; // InstanceCount
-			OcclusionQueryIndirectArgs[ViewIndex * 5 + 2] = 0; // StartIndexLocation
-			OcclusionQueryIndirectArgs[ViewIndex * 5 + 3] = 0; // BaseVertexLocation
-			OcclusionQueryIndirectArgs[ViewIndex * 5 + 4] = 0; // StartInstanceLocation
+			OcclusionQueryArgs[ViewIndex * 5 + 0] = 36;// IndexCountPerInstance; 12 triangles per cube
+			OcclusionQueryArgs[ViewIndex * 5 + 1] = 0; // InstanceCount
+			OcclusionQueryArgs[ViewIndex * 5 + 2] = 0; // StartIndexLocation
+			OcclusionQueryArgs[ViewIndex * 5 + 3] = 0; // BaseVertexLocation
+			OcclusionQueryArgs[ViewIndex * 5 + 4] = 0; // StartInstanceLocation
 		}
 	}
 #endif // PRECISE_OCCLUSION_QUERIES
@@ -140,7 +140,7 @@ RWByteAddressBuffer PackedNodes;
 #if PRECISE_OCCLUSION_QUERIES
 RWBuffer<float4> OcclusionQueryBoxes;
 RWBuffer<uint> OcclusionVisibility;
-RWBuffer<uint> OcclusionQueryIndirectArgs;
+RWBuffer<uint> OcclusionQueryArgs;
 #endif // PRECISE_OCCLUSION_QUERIES
 
 Texture2D<float4> QuadTreeTexture;
@@ -338,7 +338,7 @@ void MainCS(uint3 DispatchThreadId : SV_DispatchThreadID)
 #if PRECISE_OCCLUSION_QUERIES
 			// Increase the InstanceCount by 1.
 			uint Dummy = 0;
-			InterlockedAdd(OcclusionQueryIndirectArgs[ViewIndex * 5 + 1], 1, Dummy);
+			InterlockedAdd(OcclusionQueryArgs[ViewIndex * 5 + 1], 1, Dummy);
 
 			// Write out bounding box and store a flag indicating if the box crosses the near plane.
 			const float3 NodeAABBMin = float3(NodeAABB2D.xy, WaterZBounds.x);
diff --git a/Engine/Plugins/Experimental/Water/Source/Runtime/Private/WaterQuadTreeGPU.cpp b/Engine/Plugins/Experimental/Water/Source/Runtime/Private/WaterQuadTreeGPU.cpp
--- a/Engine/Plugins/Experimental/Water/Source/Runtime/Private/WaterQuadTreeGPU.cpp
+++ b/Engine/Plugins/Experimental/Water/Source/Runtime/Private/WaterQuadTreeGPU.cpp
@@ -154,7 +154,7 @@ public:
 
 	BEGIN_SHADER_PARAMETER_STRUCT(FParameters, )
 		SHADER_PARAMETER_RDG_BUFFER_UAV(RWBuffer, IndirectArgs)
-		SHADER_PARAMETER_RDG_BUFFER_UAV(RWBuffer, OcclusionQueryIndirectArgs)
+		SHADER_PARAMETER_RDG_BUFFER_UAV(RWBuffer, OcclusionQueryArgs)
 		SHADER_PARAMETER(uint32, NumDrawBuckets)
 		SHADER_PARAMETER(uint32, NumViews)
 		SHADER_PARAMETER(uint32, NumQuads)
@@ -203,7 +203,7 @@ public:
 		SHADER_PARAMETER_RDG_BUFFER_UAV(RWByteAddressBuffer, PackedNodes)
 		SHADER_PARAMETER_RDG_BUFFER_UAV(RWBuffer, OcclusionQueryBoxes)
 		SHADER_PARAMETER_RDG_BUFFER_UAV(RWBuffer, OcclusionVisibility)
-		SHADER_PARAMETER_RDG_BUFFER_UAV(RWBuffer, OcclusionQueryIndirectArgs)
+		SHADER_PARAMETER_RDG_BUFFER_UAV(RWBuffer, OcclusionQueryArgs)
 		SHADER_PARAMETER_RDG_TEXTURE(Texture2D, QuadTreeTexture)
 		SHADER_PARAMETER_RDG_TEXTURE(Texture2D, WaterZBoundsTexture)
 		SHADER_PARAMETER_RDG_BUFFER_SRV(StructuredBuffer, WaterBodyRenderData)
@@ -793,7 +793,7 @@ void FWaterQuadTreeGPU::Traverse(FRDGBuilder& GraphBuilder, const FTraverseParam
 
 		FWaterQuadTreeInitializeIndirectArgsCS::FParameters* PassParameters = GraphBuilder.AllocParameters<FWaterQuadTreeInitializeIndirectArgsCS::FParameters>();
 		PassParameters->IndirectArgs = IndirectArgsBufferUAV;
-		PassParameters->OcclusionQueryIndirectArgs = OcclusionQueryIndirectArgsBufferUAV;
+		PassParameters->OcclusionQueryArgs = OcclusionQueryIndirectArgsBufferUAV;
 		PassParameters->NumDrawBuckets = Params.NumMaterials;
 		PassParameters->NumViews = Params.NumViews;
 		PassParameters->NumQuads = Params.NumQuadsPerTileSide;
@@ -868,7 +868,7 @@ void FWaterQuadTreeGPU::Traverse(FRDGBuilder& GraphBuilder, const FTraverseParam
 			PassParameters->PackedNodes = PackedNodesUAV;
 			PassParameters->OcclusionQueryBoxes = OcclusionQueryBoxesUAV;
 			PassParameters->OcclusionVisibility = OcclusionQueryResultsUAV;
-			PassParameters->OcclusionQueryIndirectArgs = OcclusionQueryIndirectArgsBufferUAV;
+			PassParameters->OcclusionQueryArgs = OcclusionQueryIndirectArgsBufferUAV;
 			PassParameters->QuadTreeTexture = QuadTreeTextureRDG;
 			PassParameters->WaterZBoundsTexture = WaterZBoundsTextureRDG;
 			PassParameters->WaterBodyRenderData = WaterBodyRenderDataBufferSRV;
