Changes to custom render pass API :
* Made all class members protected/public in order to help with eventual deprecation
* Added utility to have a custom render pass perform a render capture operation (full capture, begin capture, end capture) when it runs
* Added BeginPass / EndPass functions, that run right before/after PreRender/PostRender, which are only really useful for properly scoping the render capture operations described above
* BeginPass/PreRender/PostRender/EndPass are now non-virtual. Child classes now have to implement the virtual On[...] version of them
* Added proper class for custom render pass user data (same lifetime as the render pass) and allowed multiple user data to be attached to a custom render pass, identified by a FName, which allows custom render data to be retrieved in a safe way by any module. Currently, the system was only used by landscape so it was working, but really, any render pass could hold any kind of user data so this would have crashed
* Added bare-bones RTTI system to both render passes and render pass user data (GetTypeName). Use IMPLEMENT_CUSTOM_RENDER_PASS (and IMPLEMENT_CUSTOM_RENDER_PASS_USER_DATA, respectively) to implement it

Misc:
* Fixed water info render captures (r.Water.WaterInfo.RenderCaptureNextWaterInfoDraws) when running with custom render passes (r.Water.WaterInfo.RenderMethod 2)
#tests cooked build, editor

diff --git a/Engine/Plugins/Experimental/Water/Source/Runtime/Private/WaterBodyInfoMeshComponent.cpp b/Engine/Plugins/Experimental/Water/Source/Runtime/Private/WaterBodyInfoMeshComponent.cpp
--- a/Engine/Plugins/Experimental/Water/Source/Runtime/Private/WaterBodyInfoMeshComponent.cpp
+++ b/Engine/Plugins/Experimental/Water/Source/Runtime/Private/WaterBodyInfoMeshComponent.cpp
@@ -1,8 +1,9 @@
 // Copyright Epic Games, Inc. All Rights Reserved.
 
 #include "WaterBodyInfoMeshComponent.h"
+#include "WaterInfoRendering.h"
 #include "StaticMeshSceneProxy.h"
-#include  "UObject/UObjectIterator.h"
+#include "UObject/UObjectIterator.h"
 #include "Rendering/CustomRenderPass.h"
 
 #include UE_INLINE_GENERATED_CPP_BY_NAME(WaterBodyInfoMeshComponent)
@@ -94,13 +95,22 @@ void FWaterBodyInfoMeshSceneProxy::SetEnabled(bool bInEnabled)
 
 FPrimitiveViewRelevance FWaterBodyInfoMeshSceneProxy::GetViewRelevance(const FSceneView* View) const
 {
+	using namespace UE::WaterInfo;
+
 	FPrimitiveViewRelevance Result = FStaticMeshSceneProxy::GetViewRelevance(View);
 
 	// When water info mesh is rendered with custom render passes, enable the mesh for drawing
 	if (GetWaterInfoRenderingMethod() == 2)
 	{
-		const FString PassName = View->CustomRenderPass ? const_cast<FSceneView*>(View)->CustomRenderPass->Name : TEXT("");
-		Result.bDrawRelevance = (PassName == TEXT("WaterInfoDepthPass") || PassName == TEXT("WaterInfoColorPass") || PassName == TEXT("WaterInfoDilationPass"));
+		Result.bDrawRelevance = false;
+		if (View->CustomRenderPass != nullptr)
+		{
+			const FName& PassName = const_cast<FSceneView*>(View)->CustomRenderPass->GetTypeName();
+			Result.bDrawRelevance = (PassName == GetWaterInfoDepthPassName())
+				|| (PassName == GetWaterInfoColorPassName())
+				|| (PassName == GetWaterInfoDilationPassName());
+
+		}
 		Result.bShadowRelevance = false;
 	}
 
diff --git a/Engine/Plugins/Experimental/Water/Source/Runtime/Private/WaterInfoRendering.cpp b/Engine/Plugins/Experimental/Water/Source/Runtime/Private/WaterInfoRendering.cpp
--- a/Engine/Plugins/Experimental/Water/Source/Runtime/Private/WaterInfoRendering.cpp
+++ b/Engine/Plugins/Experimental/Water/Source/Runtime/Private/WaterInfoRendering.cpp
@@ -895,36 +895,40 @@ void UpdateWaterInfoRendering2(FSceneView& InView, const TMap<AWaterZone*, UE::W
 	}
 }
 
+
+// ----------------------------------------------------------------------------------
+
 /** Base class containing common functionalities for all water info passes. */
-class FWaterInfoCustomRenderPassBase : public FCustomRenderPass
+class FWaterInfoCustomRenderPassBase : public FCustomRenderPassBase
 {
 public:
-	FWaterInfoCustomRenderPassBase()
-		: FCustomRenderPass()
-	{}
-
-	void SetLandscapeLODOverrides(const TOptional<TMap<uint32, int32>>& InLandscapeLODOverrides)
+	FWaterInfoCustomRenderPassBase(const FString& InDebugName, FCustomRenderPassBase::ERenderMode InRenderMode, FCustomRenderPassBase::ERenderOutput InRenderOutput, const FIntPoint& InRenderTargetSize)
+		: FCustomRenderPassBase(InDebugName, InRenderMode, InRenderOutput, InRenderTargetSize)
 	{
-		// See LandscapeRender.cpp FLandscapeComponentSceneProxy::GetViewLodOverride() where UserData is casted back to LandscapeLODOverrides.
-		LandscapeLODOverrides = InLandscapeLODOverrides;
-		UserData = &LandscapeLODOverrides;
 	}
 
-	TOptional<TMap<uint32, int32>> LandscapeLODOverrides;
+	// Abstract class :
+	virtual const FName& GetTypeName() const PURE_VIRTUAL(FWaterInfoCustomRenderPassBase::GetTypeName, static FName Name; return Name;);
 };
 
-class FWaterInfoRenderingDepthPass : public FWaterInfoCustomRenderPassBase
+
+// ----------------------------------------------------------------------------------
+
+class FWaterInfoRenderingDepthPass final : public FWaterInfoCustomRenderPassBase
 {
 public:
-	FWaterInfoRenderingDepthPass()
-		: FWaterInfoCustomRenderPassBase()
+	IMPLEMENT_CUSTOM_RENDER_PASS(FWaterInfoRenderingDepthPass);
+
+	FWaterInfoRenderingDepthPass(const FIntPoint& InRenderTargetSize, const TMap<uint32, int32>& InLandscapeLODOverrides)
+		: FWaterInfoCustomRenderPassBase(TEXT("WaterInfoDepthPass"), FCustomRenderPassBase::ERenderMode::DepthPass, FCustomRenderPassBase::ERenderOutput::DeviceDepth, InRenderTargetSize)
 	{
-		Name = TEXT("WaterInfoDepthPass");
-		RenderMode = FCustomRenderPass::ERenderMode_DepthPass;
-		RenderOutput = FCustomRenderPass::ERenderOutput_DeviceDepth;
+		if (!InLandscapeLODOverrides.IsEmpty())
+		{
+			SetUserData(MakeUnique<FLandscapeLODOverridesCustomRenderPassUserData>(InLandscapeLODOverrides));
+		}
 	}
 
-	virtual void PreRender(FRDGBuilder& GraphBuilder) override
+	virtual void OnPreRender(FRDGBuilder& GraphBuilder) override
 	{
 		const FRDGTextureDesc TextureDesc = FRDGTextureDesc::Create2D(RenderTargetSize, PF_FloatRGBA, FClearValueBinding::Black, TexCreate_RenderTargetable | TexCreate_ShaderResource);
 		RenderTargetTexture = GraphBuilder.CreateTexture(TextureDesc, TEXT("WaterDepthTexture"));
@@ -932,18 +936,21 @@ public:
 	}
 };
 
-class FWaterInfoRenderingColorPass : public FWaterInfoCustomRenderPassBase
+const FName& GetWaterInfoDepthPassName() { return FWaterInfoRenderingDepthPass::GetTypeNameStatic(); }
+
+
+// ----------------------------------------------------------------------------------
+
+class FWaterInfoRenderingColorPass final : public FWaterInfoCustomRenderPassBase
 {
 public:
-	FWaterInfoRenderingColorPass()
-		: FWaterInfoCustomRenderPassBase()
-	{
-		Name = TEXT("WaterInfoColorPass");
-		RenderMode = FCustomRenderPass::ERenderMode_DepthAndBasePass;
-		RenderOutput = FCustomRenderPass::ERenderOutput_SceneColorAndDepth;
-	}
+	IMPLEMENT_CUSTOM_RENDER_PASS(FWaterInfoRenderingColorPass);
+
+	FWaterInfoRenderingColorPass(const FIntPoint& InRenderTargetSize)
+		: FWaterInfoCustomRenderPassBase(TEXT("WaterInfoColorPass"), FCustomRenderPassBase::ERenderMode::DepthAndBasePass, FCustomRenderPassBase::ERenderOutput::SceneColorAndDepth, InRenderTargetSize)
+	{}
 
-	virtual void PreRender(FRDGBuilder& GraphBuilder) override
+	virtual void OnPreRender(FRDGBuilder& GraphBuilder) override
 	{
 		const FRDGTextureDesc TextureDesc = FRDGTextureDesc::Create2D(RenderTargetSize, PF_A32B32G32R32F, FClearValueBinding::Black, TexCreate_RenderTargetable | TexCreate_ShaderResource);
 		RenderTargetTexture = GraphBuilder.CreateTexture(TextureDesc, TEXT("WaterColorTexture"));
@@ -951,29 +958,32 @@ public:
 	}
 };
 
-class FWaterInfoRenderingDilationPass : public FWaterInfoCustomRenderPassBase
+const FName& GetWaterInfoColorPassName() { return FWaterInfoRenderingColorPass::GetTypeNameStatic(); }
+
+
+// ----------------------------------------------------------------------------------
+
+class FWaterInfoRenderingDilationPass final : public FWaterInfoCustomRenderPassBase
 {
 public:
-	FWaterInfoRenderingDilationPass()
-		: FWaterInfoCustomRenderPassBase()
-	{
-		Name = TEXT("WaterInfoDilationPass");
-		RenderMode = FCustomRenderPass::ERenderMode_DepthPass;
-		RenderOutput = FCustomRenderPass::ERenderOutput_DeviceDepth;
-	}
+	IMPLEMENT_CUSTOM_RENDER_PASS(FWaterInfoRenderingDilationPass);
+
+	FWaterInfoRenderingDilationPass(const FIntPoint& InRenderTargetSize)
+		: FWaterInfoCustomRenderPassBase(TEXT("WaterInfoDilationPass"), FCustomRenderPassBase::ERenderMode::DepthPass, FCustomRenderPassBase::ERenderOutput::DeviceDepth, InRenderTargetSize)
+	{}
 
-	virtual void PreRender(FRDGBuilder& GraphBuilder) override
+	virtual void OnPreRender(FRDGBuilder& GraphBuilder) override
 	{
 		const FRDGTextureDesc TextureDesc = FRDGTextureDesc::Create2D(RenderTargetSize, PF_FloatRGBA, FClearValueBinding::Black, TexCreate_RenderTargetable | TexCreate_ShaderResource);
 		RenderTargetTexture = GraphBuilder.CreateTexture(TextureDesc, TEXT("WaterDilationTexture"));
 		AddClearRenderTargetPass(GraphBuilder, RenderTargetTexture, FLinearColor::Black, Views[0]->UnscaledViewRect);
 	}
 	
-	virtual void PostRender(FRDGBuilder& GraphBuilder) override
+	virtual void OnPostRender(FRDGBuilder& GraphBuilder) override
 	{
 		FRDGTextureDesc TextureDesc(RenderTargetTexture->Desc);		
 		FRDGTextureRef MergeTargetTexture = GraphBuilder.CreateTexture(TextureDesc, TEXT("WaterInfoMerged"));
-		MergeWaterInfoAndDepth(GraphBuilder, *Views[0]->Family, *Views[0], MergeTargetTexture, DepthPass->RenderTargetTexture, ColorPass->RenderTargetTexture, RenderTargetTexture, Params);
+		MergeWaterInfoAndDepth(GraphBuilder, *Views[0]->Family, *Views[0], MergeTargetTexture, DepthPass->GetRenderTargetTexture(), ColorPass->GetRenderTargetTexture(), RenderTargetTexture, Params);
 	
 		FRDGTextureRef FinalizedTexture = GraphBuilder.CreateTexture(RenderTargetTexture->Desc, TEXT("WaterInfoFinalized"));
 		FinalizeWaterInfo(GraphBuilder, *Views[0]->Family, *Views[0], MergeTargetTexture, FinalizedTexture, Params);
@@ -989,8 +999,15 @@ public:
 	FUpdateWaterInfoParams Params;
 };
 
-static void GatherLandscapeLODOverrides(const FIntPoint& RenderTargetSize, const FVector& WaterZoneExtents, TOptional<TMap<uint32, int32>>& OutLandscapeLODOverrides)
+const FName& GetWaterInfoDilationPassName() { return FWaterInfoRenderingDilationPass::GetTypeNameStatic(); }
+
+
+// ----------------------------------------------------------------------------------
+
+static TMap<uint32, int32> GatherLandscapeLODOverrides(const FIntPoint& RenderTargetSize, const FVector& WaterZoneExtents)
 {
+	TMap<uint32, int32> LandscapeLODOverrides;
+
 	TMap<uint32, FLandscapeRenderSystem*> const* LandscapeRenderSystems = nullptr;
 	if (ILandscapeModule* LandscapeModule = FModuleManager::GetModulePtr<ILandscapeModule>("Landscape"))
 	{
@@ -1002,11 +1019,9 @@ static void GatherLandscapeLODOverrides(const FIntPoint& RenderTargetSize, const
 
 	if (!LandscapeRenderSystems)
 	{
-		return;
+		return LandscapeLODOverrides;
 	}
 
-	OutLandscapeLODOverrides.Emplace();
-
 	// In order to prevent overdrawing the landscape components, we compute the lowest-detailed LOD level which satisfies the pixel coverage of the Water Info texture
 	// and force it on all landscape components. This override is set different per Landscape actor in case there are multiple under the same water zone.
 	//
@@ -1046,8 +1061,10 @@ static void GatherLandscapeLODOverrides(const FIntPoint& RenderTargetSize, const
 			}
 		}
 
-		OutLandscapeLODOverrides->Add(LandscapeRenderSystemKey, OptimalLODLevel);
+		LandscapeLODOverrides.Add(LandscapeRenderSystemKey, OptimalLODLevel);
 	}
+
+	return LandscapeLODOverrides;
 }
 
 void UpdateWaterInfoRendering_CustomRenderPass(
@@ -1056,18 +1073,19 @@ void UpdateWaterInfoRendering_CustomRenderPass(
 {
 	TRACE_CPUPROFILER_EVENT_SCOPE(WaterInfo::UpdateWaterInfoRendering_CustomRenderPass);
 
+	if (!IsValid(Context.TextureRenderTarget) || Scene == nullptr)
+	{
+		return;
+	}
+
 	static auto* CVarRenderCaptureNextWaterInfoDraws = IConsoleManager::Get().FindConsoleVariable(TEXT("r.Water.WaterInfo.RenderCaptureNextWaterInfoDraws"));
 	int32 RenderCaptureNextWaterInfoDraws = CVarRenderCaptureNextWaterInfoDraws ? CVarRenderCaptureNextWaterInfoDraws->AsVariableInt()->GetValueOnGameThread() : 0;
-	RenderCaptureInterface::FScopedCapture RenderCapture((RenderCaptureNextWaterInfoDraws != 0), TEXT("RenderWaterInfo"));
+	bool bPerformRenderCapture = false;
 	if (RenderCaptureNextWaterInfoDraws != 0)
 	{
 		RenderCaptureNextWaterInfoDraws = FMath::Max(0, RenderCaptureNextWaterInfoDraws - 1);
 		CVarRenderCaptureNextWaterInfoDraws->Set(RenderCaptureNextWaterInfoDraws);
-	}
-
-	if (!IsValid(Context.TextureRenderTarget) || Scene == nullptr)
-	{
-		return;
+		bPerformRenderCapture = true;
 	}
 
 	const FVector ZoneExtent = Context.ZoneToRender->GetDynamicWaterInfoExtent();
@@ -1084,9 +1102,6 @@ void UpdateWaterInfoRendering_CustomRenderPass(
 
 	FTextureRenderTargetResource* TextureRenderTargetResource = Context.TextureRenderTarget->GameThread_GetRenderTargetResource();
 
-	TOptional<TMap<uint32, int32>> LandscapeLODOverrides;
-	GatherLandscapeLODOverrides(RenderTargetSize, ZoneExtent, LandscapeLODOverrides);
-
 	FSceneInterface::FCustomRenderPassRendererInput PassInput;
 	PassInput.ViewLocation = ViewLocation;
 	PassInput.ViewRotationMatrix = FLookAtMatrix(ViewLocation, LookAt, FVector(0.f, -1.f, 0.f));
@@ -1109,9 +1124,12 @@ void UpdateWaterInfoRendering_CustomRenderPass(
 	}
 	PassInput.ShowOnlyPrimitives = MoveTemp(ComponentsToRenderInDepthPass);
 
-	FWaterInfoRenderingDepthPass* DepthPass = new FWaterInfoRenderingDepthPass();
-	DepthPass->RenderTargetSize = RenderTargetSize;
-	DepthPass->SetLandscapeLODOverrides(LandscapeLODOverrides);
+	FWaterInfoRenderingDepthPass* DepthPass = new FWaterInfoRenderingDepthPass(RenderTargetSize, GatherLandscapeLODOverrides(RenderTargetSize, ZoneExtent));
+	if (bPerformRenderCapture)
+	{
+		// Initiate a render capture when this pass runs :
+		DepthPass->PerformRenderCapture(FCustomRenderPassBase::ERenderCaptureType::BeginCapture);
+	}
 	PassInput.CustomRenderPass = DepthPass;
 	Scene->CustomRenderPassRendererInputs.Add(PassInput);
 
@@ -1138,9 +1156,7 @@ void UpdateWaterInfoRendering_CustomRenderPass(
 		}
 	}
 	PassInput.ShowOnlyPrimitives = MoveTemp(ComponentsToRenderInColorPass);
-	FWaterInfoRenderingColorPass* ColorPass = new FWaterInfoRenderingColorPass();
-	ColorPass->RenderTargetSize = RenderTargetSize;
-	ColorPass->SetLandscapeLODOverrides(LandscapeLODOverrides);
+	FWaterInfoRenderingColorPass* ColorPass = new FWaterInfoRenderingColorPass(RenderTargetSize);
 	PassInput.CustomRenderPass = ColorPass;
 	Scene->CustomRenderPassRendererInputs.Add(PassInput);
 
@@ -1152,13 +1168,16 @@ void UpdateWaterInfoRendering_CustomRenderPass(
 	Params.WaterZoneExtents = ZoneExtent;
 
 	PassInput.ShowOnlyPrimitives = MoveTemp(ComponentsToRenderInDilationPass);
-	FWaterInfoRenderingDilationPass* DilationPass = new FWaterInfoRenderingDilationPass();
+	FWaterInfoRenderingDilationPass* DilationPass = new FWaterInfoRenderingDilationPass(RenderTargetSize);
 	DilationPass->DepthPass = DepthPass;
 	DilationPass->ColorPass = ColorPass;
 	DilationPass->WaterInfoRenderTarget = Context.TextureRenderTarget->GameThread_GetRenderTargetResource();
-	DilationPass->RenderTargetSize = RenderTargetSize;
-	DilationPass->SetLandscapeLODOverrides(LandscapeLODOverrides);
 	DilationPass->Params = Params;
+	if (bPerformRenderCapture)
+	{
+		// End a render capture when this pass runs :
+		DilationPass->PerformRenderCapture(FCustomRenderPassBase::ERenderCaptureType::EndCapture);
+	}
 	PassInput.CustomRenderPass = DilationPass;
 	Scene->CustomRenderPassRendererInputs.Add(PassInput);
 }
diff --git a/Engine/Plugins/Experimental/Water/Source/Runtime/Private/WaterInfoRendering.h b/Engine/Plugins/Experimental/Water/Source/Runtime/Private/WaterInfoRendering.h
--- a/Engine/Plugins/Experimental/Water/Source/Runtime/Private/WaterInfoRendering.h
+++ b/Engine/Plugins/Experimental/Water/Source/Runtime/Private/WaterInfoRendering.h
@@ -13,14 +13,15 @@ class FSceneView;
 
 namespace UE::WaterInfo
 {
-	struct FRenderingContext
-	{
-		AWaterZone* ZoneToRender = nullptr;
-		UTextureRenderTarget2D* TextureRenderTarget;
-		TArray<UWaterBodyComponent*> WaterBodies;
-		TArray<TWeakObjectPtr<UPrimitiveComponent>> GroundPrimitiveComponents;
-		float CaptureZ;
-	};
+
+struct FRenderingContext
+{
+	AWaterZone* ZoneToRender = nullptr;
+	UTextureRenderTarget2D* TextureRenderTarget;
+	TArray<UWaterBodyComponent*> WaterBodies;
+	TArray<TWeakObjectPtr<UPrimitiveComponent>> GroundPrimitiveComponents;
+	float CaptureZ;
+};
 	
 void UpdateWaterInfoRendering(
 	FSceneInterface* Scene,
@@ -33,4 +34,9 @@ void UpdateWaterInfoRendering2(
 void UpdateWaterInfoRendering_CustomRenderPass(
 	FSceneInterface* Scene,
 	const FRenderingContext& Context);
-}
+
+const FName& GetWaterInfoDepthPassName();
+const FName& GetWaterInfoColorPassName();
+const FName& GetWaterInfoDilationPassName();
+
+} // namespace UE::WaterInfo
diff --git a/Engine/Source/Runtime/Engine/Private/Rendering/CustomRenderPass.cpp b/Engine/Source/Runtime/Engine/Private/Rendering/CustomRenderPass.cpp
new file mode 100644
--- /dev/null
+++ b/Engine/Source/Runtime/Engine/Private/Rendering/CustomRenderPass.cpp
@@ -0,0 +1,100 @@
+// Copyright Epic Games, Inc. All Rights Reserved.
+
+#include "Rendering/CustomRenderPass.h"
+#include "IRenderCaptureProvider.h"
+#include "RenderGraphBuilder.h"
+
+FCustomRenderPassBase::FCustomRenderPassBase(const FString& InDebugName, ERenderMode InRenderMode, ERenderOutput InRenderOutput, const FIntPoint& InRenderTargetSize)
+	: DebugName(InDebugName)
+	, RenderMode(InRenderMode)
+	, RenderOutput(InRenderOutput)
+	, RenderTargetSize(InRenderTargetSize)
+{}
+
+void FCustomRenderPassBase::BeginPass(FRDGBuilder& GraphBuilder)
+{
+	// Optionally, perform a render capture when this pass runs :
+	if ((RenderCaptureType == ERenderCaptureType::Capture) || (RenderCaptureType == ERenderCaptureType::BeginCapture))
+	{
+		GraphBuilder.AddPass(
+			RDG_EVENT_NAME("BeginCapture"),
+			ERDGPassFlags::None,
+			[FileName = RenderCaptureFileName](FRHICommandListImmediate& RHICommandListLocal)
+		{
+			IRenderCaptureProvider::Get().BeginCapture(&RHICommandListLocal, IRenderCaptureProvider::ECaptureFlags_Launch, FileName);
+		});
+
+		if (!DebugName.IsEmpty())
+		{
+			GraphBuilder.BeginEventScope(RDG_EVENT_NAME("%s", *DebugName));
+		}
+	}
+
+	OnBeginPass(GraphBuilder);
+}
+
+void FCustomRenderPassBase::PreRender(FRDGBuilder& GraphBuilder)
+{
+	OnPreRender(GraphBuilder);
+}
+
+void FCustomRenderPassBase::PostRender(FRDGBuilder& GraphBuilder)
+{
+	OnPostRender(GraphBuilder);
+}
+
+void FCustomRenderPassBase::EndPass(FRDGBuilder & GraphBuilder)
+{
+	OnEndPass(GraphBuilder);
+
+	// End the optional render capture after this pass has run :
+	if ((RenderCaptureType == ERenderCaptureType::Capture) || (RenderCaptureType == ERenderCaptureType::EndCapture))
+	{
+		if (!DebugName.IsEmpty())
+		{
+			GraphBuilder.EndEventScope();
+		}
+
+		GraphBuilder.AddPass(
+			RDG_EVENT_NAME("EndCapture"),
+			ERDGPassFlags::None,
+			[](FRHICommandListImmediate& RHICommandListLocal)
+		{
+			IRenderCaptureProvider::Get().EndCapture(&RHICommandListLocal);
+		});
+	}
+}
+
+ESceneCaptureSource FCustomRenderPassBase::GetSceneCaptureSource() const
+{
+	if (RenderOutput == ERenderOutput::SceneDepth)
+		return SCS_SceneDepth;
+	else if (RenderOutput == ERenderOutput::DeviceDepth)
+		return SCS_DeviceDepth;
+	else if (RenderOutput == ERenderOutput::SceneColorAndDepth)
+		return SCS_SceneColorSceneDepth;
+	else
+		return SCS_MAX;
+}
+
+void FCustomRenderPassBase::PerformRenderCapture(ERenderCaptureType InRenderCaptureType, const FString& InFileName)
+{
+	if (!IRenderCaptureProvider::IsAvailable() || (InRenderCaptureType == ERenderCaptureType::NoCapture))
+	{
+		return;
+	}
+
+	RenderCaptureType = InRenderCaptureType;
+	RenderCaptureFileName = InFileName;
+}
+
+void FCustomRenderPassBase::SetUserData(TUniquePtr<ICustomRenderPassUserData>&& InUserData)
+{
+	UserDatas.Add(InUserData->GetTypeName(), MoveTemp(InUserData));
+}
+
+ICustomRenderPassUserData* FCustomRenderPassBase::GetUserData(const FName& InTypeName) const
+{
+	const TUniquePtr<ICustomRenderPassUserData>* Found = UserDatas.Find(InTypeName);
+	return (Found != nullptr) ? Found->Get() : nullptr;
+}
diff --git a/Engine/Source/Runtime/Engine/Public/Rendering/CustomRenderPass.h b/Engine/Source/Runtime/Engine/Public/Rendering/CustomRenderPass.h
--- a/Engine/Source/Runtime/Engine/Public/Rendering/CustomRenderPass.h
+++ b/Engine/Source/Runtime/Engine/Public/Rendering/CustomRenderPass.h
@@ -1,7 +1,7 @@
 // Copyright Epic Games, Inc. All Rights Reserved.
 
 /*=============================================================================
-	FCustomRenderPass.h: Custom render pass implementation.
+	CustomRenderPass.h: Custom render pass implementation.
 =============================================================================*/
 
 #pragma once
@@ -10,63 +10,141 @@
 #include "RenderGraphFwd.h"
 #include "Engine/EngineTypes.h"
 
-/** Base class of the custom render pass. Create a derived class to provide custom behaviour. 
-	Custom render pass is rendered as part of the main render loop. */
-class FCustomRenderPass
+// ----------------------------------------------------------------------------------
+/** Base interface to implement for FCustomRenderPassBase */
+class ICustomRenderPass
 {
 public:
-	FCustomRenderPass()
-	{}
+	virtual ~ICustomRenderPass() {}
 
-	virtual ~FCustomRenderPass() {}
+	/** Bare-bones RTTI method to allow users to differentiate / downcast the custom render pass they're interested into */
+	virtual const FName& GetTypeName() const PURE_VIRTUAL(ICustomRenderPass::GetTypeName, static FName Name; return Name;);
+};
 
-	/** Called before the pass is rendered on render thread, use it for prep tasks such as allocating render target. */
-	virtual void PreRender(FRDGBuilder& GraphBuilder) {}
-	/** Called after the pass is rendered on render thread, use it for post processing tasks such as applying additional shader effects to the render target. */
-	virtual void PostRender(FRDGBuilder& GraphBuilder) {}
+/** Use this macro to implement the RTTI method for custom render passes. Should only be used on final classes (e.g. if C derives from B which derives from A, we won't be able 
+  *  to query a C* for whether it's a "B", in other words, A and B should be abstract and C, final) */
+#define IMPLEMENT_CUSTOM_RENDER_PASS(TypeName) const FName& GetTypeName() const override { return GetTypeNameStatic(); } \
+	static const FName& GetTypeNameStatic() { static FName Name(TEXT(#TypeName)); return Name; }
+
+
+// ----------------------------------------------------------------------------------
+/** Base interface to implement for attaching user data to a FCustomRenderPassBase */
+class ICustomRenderPassUserData
+{
+public:
+	virtual ~ICustomRenderPassUserData() {}
+	/** Bare-bones RTTI method to allow users to differentiate / downcast the custom render pass user data they're interested into */
+	virtual const FName& GetTypeName() const PURE_VIRTUAL(ICustomRenderPassUserData::GetTypeName, static FName Name; return Name;);
+};
+
+/** Use this macro to implement the RTTI method for custom render pass user data. Should only be used on final classes (e.g. if C derives from B which derives from A, we won't be able
+  *  to query a C* for whether it's a "B", in other words, A and B should be abstract and C, final) */
+#define IMPLEMENT_CUSTOM_RENDER_PASS_USER_DATA(TypeName) const FName& GetTypeName() const override { return GetTypeNameStatic(); } \
+	static const FName& GetTypeNameStatic() { static FName Name(TEXT(#TypeName)); return Name; }
+
+
+// ----------------------------------------------------------------------------------
+/** Base class of the custom render pass. Create a derived class to provide custom behavior. 
+	Custom render pass is rendered as part of the main render loop. */
+class FCustomRenderPassBase : public ICustomRenderPass
+{
+public:
+	friend class FSceneRenderer;
 
 	/** Which render passes are needed for the custom render pass. */
-	enum ERenderMode
+	enum class ERenderMode
 	{
 		/** Render depth pre-pass only. */
-		ERenderMode_DepthPass,
+		DepthPass,
 		/** Render depth pre-pass and base pass. */
-		ERenderMode_DepthAndBasePass
-	} RenderMode;
+		DepthAndBasePass
+	};
 
 	/** The output type into the render target. What type is valid depends on ERenderMode. */
-	enum ERenderOutput
+	enum class ERenderOutput
 	{
-		/** Used with ERenderMode_DepthPass. */
-		ERenderOutput_SceneDepth,
-		ERenderOutput_DeviceDepth,
-		/** Used with ERenderMode_DepthAndBasePass. */
-		ERenderOutput_SceneColorAndDepth
-	} RenderOutput;
-
-	/** Convert output type into corresponding scene capture source type. This is because we use CopySceneCaptureComponentToTarget helper function 
+		/** Used with ERenderMode::DepthPass. */
+		SceneDepth,
+		DeviceDepth,
+		/** Used with ERenderMode::DepthAndBasePass. */
+		SceneColorAndDepth
+	};
+
+	enum class ERenderCaptureType
+	{
+		/** Don't trigger a render capture */
+		NoCapture, 
+		/** Start a render capture on PreRender, stop it on PostRender */
+		Capture,
+		/** Start a render capture on PreRender (this must be followed by another pass with EndCapture). Use it when the render capture encapsulates several custom passes */
+		BeginCapture,
+		/** Stop a render capture on PostRender (this must be preceded by another pass with BeginCapture). Use it when the render capture encapsulates several custom passes */
+		EndCapture,
+	};
+
+	FCustomRenderPassBase() = delete;
+	ENGINE_API FCustomRenderPassBase(const FString& InDebugName, ERenderMode InRenderMode, ERenderOutput InRenderOutput, const FIntPoint& InRenderTargetSize);
+
+	virtual ~FCustomRenderPassBase() {}
+	
+	/** Called before PreRender on render thread (only really useful for properly scoping render capture names) */
+	ENGINE_API void BeginPass(FRDGBuilder& GraphBuilder);
+	/** Called before the pass is rendered on render thread, use it for prep tasks such as allocating render target. */
+	ENGINE_API void PreRender(FRDGBuilder& GraphBuilder);
+	/** Called after the pass is rendered on render thread, use it for post processing tasks such as applying additional shader effects to the render target. */
+	ENGINE_API void PostRender(FRDGBuilder& GraphBuilder);
+	/** Called after PostRender on render thread (only really useful for properly scoping render capture names) */
+	ENGINE_API void EndPass(FRDGBuilder& GraphBuilder);
+
+	/** Convert output type into corresponding scene capture source type. This is because we use CopySceneCaptureComponentToTarget helper function
 		to copy the custom render pass render result from scene texture into the pass's render target. */
-	ESceneCaptureSource GetSceneCaptureSource() const
+	ENGINE_API ESceneCaptureSource GetSceneCaptureSource() const;
+
+	/** Perform a render capture when this pass runs if a render capture interface is registered */
+	ENGINE_API void PerformRenderCapture(ERenderCaptureType InRenderCaptureType, const FString& InFileName = FString());
+
+	const FString& GetDebugName() const { return DebugName; }
+	ERenderMode GetRenderMode() const { return RenderMode; }
+	ERenderOutput GetRenderOutput() const { return RenderOutput; }
+	FRDGTextureRef GetRenderTargetTexture() const { return RenderTargetTexture; }
+	const FIntPoint& GetRenderTargetSize() const { return RenderTargetSize; }
+
+	ENGINE_API void SetUserData(TUniquePtr<ICustomRenderPassUserData>&& InUserData);
+	ENGINE_API ICustomRenderPassUserData* GetUserData(const FName& InTypeName) const;
+
+	template <typename UserDataType>
+	typename TEnableIf<TPointerIsConvertibleFromTo<UserDataType, ICustomRenderPassUserData>::Value, UserDataType>::Type* GetUserDataTyped() const
 	{
-		if (RenderOutput == ERenderOutput_SceneDepth)
-			return SCS_SceneDepth;
-		else if (RenderOutput == ERenderOutput_DeviceDepth)
-			return SCS_DeviceDepth;
-		else if (RenderOutput == ERenderOutput_SceneColorAndDepth)
-			return SCS_SceneColorSceneDepth;
-		else
-			return SCS_MAX;
+		if (ICustomRenderPassUserData* UserData = GetUserData(UserDataType::GetTypeNameStatic()))
+		{
+			return reinterpret_cast<UserDataType*>(UserData);
+		}
+		return nullptr;
 	}
 
-	FString Name;
+protected:
+	virtual void OnBeginPass(FRDGBuilder& GraphBuilder) {}
+	virtual void OnPreRender(FRDGBuilder& GraphBuilder) {}
+	virtual void OnPostRender(FRDGBuilder& GraphBuilder) {}
+	virtual void OnEndPass(FRDGBuilder& GraphBuilder) {}
+
+protected:
+	FString DebugName;
+	ERenderMode RenderMode = ERenderMode::DepthPass;
+	ERenderOutput RenderOutput = ERenderOutput::SceneDepth;
+
 	FRDGTextureRef RenderTargetTexture = nullptr;
-	FIntPoint RenderTargetSize;
+	FIntPoint RenderTargetSize = FIntPoint(ForceInit);
 
 	/** The views created for the custom render pass. */
 	TArray<class FViewInfo*> Views;
 
 	/** Optional user data providing a hook into the engine to override settings for this render pass.
-		Create a variable with the settings to override in the derived FCustomRenderPass, assign its pointer to UserData.
-		Then inside the engine where the setting needs hi-jacking, cast UserData to the target variable type and override the values. */
-	void* UserData = nullptr;
+		The user data should derive from ICustomRenderPassUserData and implement GetTypeName() to be retrieved.
+		Use SetUserData/GetUserData */
+	TMap<FName, TUniquePtr<ICustomRenderPassUserData>> UserDatas;
+
+private:
+	ERenderCaptureType RenderCaptureType = ERenderCaptureType::NoCapture;
+	FString RenderCaptureFileName;
 };
diff --git a/Engine/Source/Runtime/Engine/Public/SceneInterface.h b/Engine/Source/Runtime/Engine/Public/SceneInterface.h
--- a/Engine/Source/Runtime/Engine/Public/SceneInterface.h
+++ b/Engine/Source/Runtime/Engine/Public/SceneInterface.h
@@ -702,7 +702,7 @@ public:
 		TOptional<TSet<FPrimitiveComponentId>> ShowOnlyPrimitives;
 		const AActor* ViewActor = nullptr;
 
-		class FCustomRenderPass* CustomRenderPass = nullptr;
+		class FCustomRenderPassBase* CustomRenderPass = nullptr;
 	};
 	TArray<FCustomRenderPassRendererInput> CustomRenderPassRendererInputs;
 
diff --git a/Engine/Source/Runtime/Engine/Public/SceneView.h b/Engine/Source/Runtime/Engine/Public/SceneView.h
--- a/Engine/Source/Runtime/Engine/Public/SceneView.h
+++ b/Engine/Source/Runtime/Engine/Public/SceneView.h
@@ -1414,7 +1414,7 @@ public:
 	/** For sanity checking casts that are assumed to be safe. */
 	bool bIsViewInfo;
 
-	class FCustomRenderPass* CustomRenderPass = nullptr;
+	class FCustomRenderPassBase* CustomRenderPass = nullptr;
 
 	/** Whether this view is being used to render a scene capture. */
 	bool bIsSceneCapture;
diff --git a/Engine/Source/Runtime/Landscape/Private/LandscapeRender.cpp b/Engine/Source/Runtime/Landscape/Private/LandscapeRender.cpp
--- a/Engine/Source/Runtime/Landscape/Private/LandscapeRender.cpp
+++ b/Engine/Source/Runtime/Landscape/Private/LandscapeRender.cpp
@@ -394,12 +394,11 @@ static int32 GetViewLodOverride(FSceneView const& View, uint32 LandscapeKey)
 	// Use lod 0 if lodding is disabled
 	LodOverride = View.Family->EngineShowFlags.LOD == 0 ? 0 : LodOverride;
 	
-	if (View.CustomRenderPass && View.CustomRenderPass->UserData)
+	if (View.CustomRenderPass)
 	{
-		TOptional<TMap<uint32, int32>>& LandscapeLODOverrides = *(TOptional<TMap<uint32, int32>>*)(View.CustomRenderPass->UserData);
-		if (LandscapeLODOverrides)
+		if (FLandscapeLODOverridesCustomRenderPassUserData* LandscapeLODOverridesUserData = View.CustomRenderPass->GetUserDataTyped<FLandscapeLODOverridesCustomRenderPassUserData>())
 		{
-			int32* LandscapeLODOverride = LandscapeLODOverrides->Find(LandscapeKey);
+			const int32* LandscapeLODOverride = LandscapeLODOverridesUserData->GetLandscapeLODOverrides().Find(LandscapeKey);
 			LodOverride = LandscapeLODOverride ? *LandscapeLODOverride : LodOverride;
 		}
 	}
diff --git a/Engine/Source/Runtime/Landscape/Public/LandscapeRender.h b/Engine/Source/Runtime/Landscape/Public/LandscapeRender.h
--- a/Engine/Source/Runtime/Landscape/Public/LandscapeRender.h
+++ b/Engine/Source/Runtime/Landscape/Public/LandscapeRender.h
@@ -29,6 +29,7 @@ LandscapeRender.h: New terrain rendering
 #include "StaticMeshResources.h"
 #include "StaticMeshSceneProxy.h"
 #include "SceneViewExtension.h"
+#include "Rendering/CustomRenderPass.h"
 #include "Tasks/Task.h"
 
 // This defines the number of border blocks to surround terrain by when generating lightmaps
@@ -431,6 +432,9 @@ public:
 	bool bRegistered;
 };
 
+//
+// FLandscapeRenderSystem
+//
 struct FLandscapeRenderSystem
 {
 	typedef uint32 FViewKey;
@@ -579,6 +583,24 @@ private:
 	void DestroyResources_Internal(FLandscapeSectionInfo* InSectionInfo);
 };
 
+//
+// FLandscapeLODOverridesCustomRenderPassUserData
+//
+class FLandscapeLODOverridesCustomRenderPassUserData : public ICustomRenderPassUserData
+{
+public:
+	IMPLEMENT_CUSTOM_RENDER_PASS_USER_DATA(FLandscapeLODOverridesCustomRenderPassUserData);
+
+	FLandscapeLODOverridesCustomRenderPassUserData(const TMap<uint32, int32>& InLandscapeLODOverrides)
+		: LandscapeLODOverrides(InLandscapeLODOverrides)
+	{}
+
+
+	const TMap<uint32, int32>& GetLandscapeLODOverrides() const { return LandscapeLODOverrides; }
+
+private:
+	TMap<uint32, int32> LandscapeLODOverrides;
+};
 
 //
 // FLandscapeSceneViewExtension
diff --git a/Engine/Source/Runtime/Renderer/Private/DeferredShadingRenderer.cpp b/Engine/Source/Runtime/Renderer/Private/DeferredShadingRenderer.cpp
--- a/Engine/Source/Runtime/Renderer/Private/DeferredShadingRenderer.cpp
+++ b/Engine/Source/Runtime/Renderer/Private/DeferredShadingRenderer.cpp
@@ -2166,43 +2166,49 @@ void FDeferredShadingSceneRenderer::Render(FRDGBuilder& GraphBuilder)
 
 		for (int32 i = 0; i < CustomRenderPassInfos.Num(); ++i)
 		{
-			FCustomRenderPass* CustomRenderPass = CustomRenderPassInfos[i].CustomRenderPass;
+			FCustomRenderPassBase* CustomRenderPass = CustomRenderPassInfos[i].CustomRenderPass;
 			TArray<FViewInfo>& CustomRenderPassViews = CustomRenderPassInfos[i].Views;
 			FNaniteShadingCommands& NaniteBasePassShadingCommands = CustomRenderPassInfos[i].NaniteBasePassShadingCommands;
 			check(CustomRenderPass);
 
-			QUICK_SCOPE_CYCLE_COUNTER(STAT_CustomRenderPass);
-			RDG_EVENT_SCOPE(GraphBuilder, "CustomRenderPass[%d] %s", i, *CustomRenderPass->Name);
+			CustomRenderPass->BeginPass(GraphBuilder);
 
-			CustomRenderPass->PreRender(GraphBuilder);
-
-			TArray<Nanite::FRasterResults, TInlineAllocator<2>> NaniteRasterResults;
-			TArray<Nanite::FPackedView, SceneRenderingAllocator> PrimaryNaniteViews;
-			FNaniteBasePassVisibility DummyNaniteBasePassVisibility;
-			RenderPrepassAndVelocity(CustomRenderPassViews, DummyNaniteBasePassVisibility, NaniteRasterResults, PrimaryNaniteViews);
-
-			if (CustomRenderPass->RenderMode == FCustomRenderPass::ERenderMode_DepthAndBasePass)
 			{
-				SceneTextures.SetupMode |= ESceneTextureSetupMode::SceneColor;
-				SceneTextures.UniformBuffer = CreateSceneTextureUniformBuffer(GraphBuilder, &SceneTextures, FeatureLevel, SceneTextures.SetupMode);
+				QUICK_SCOPE_CYCLE_COUNTER(STAT_CustomRenderPass);
+				RDG_EVENT_SCOPE(GraphBuilder, "CustomRenderPass[%d] %s", i, *CustomRenderPass->GetDebugName());
+
+				CustomRenderPass->PreRender(GraphBuilder);
 
-				// Setup dummy uniform buffer parameters for fog volume.
-				SetDummyLocalFogVolumeForViews(GraphBuilder, CustomRenderPassViews);
+				TArray<Nanite::FRasterResults, TInlineAllocator<2>> NaniteRasterResults;
+				TArray<Nanite::FPackedView, SceneRenderingAllocator> PrimaryNaniteViews;
+				FNaniteBasePassVisibility DummyNaniteBasePassVisibility;
+				RenderPrepassAndVelocity(CustomRenderPassViews, DummyNaniteBasePassVisibility, NaniteRasterResults, PrimaryNaniteViews);
 
-				if (bNaniteEnabled && UseNaniteComputeMaterials())
+				if (CustomRenderPass->GetRenderMode() == FCustomRenderPassBase::ERenderMode::DepthAndBasePass)
 				{
-					Nanite::BuildShadingCommands(GraphBuilder, *Scene, ENaniteMeshPass::BasePass, NaniteBasePassShadingCommands, true);
-				}
+					SceneTextures.SetupMode |= ESceneTextureSetupMode::SceneColor;
+					SceneTextures.UniformBuffer = CreateSceneTextureUniformBuffer(GraphBuilder, &SceneTextures, FeatureLevel, SceneTextures.SetupMode);
 
-				RenderBasePass(GraphBuilder, CustomRenderPassViews, SceneTextures, DBufferTextures, BasePassDepthStencilAccess, /*ForwardScreenSpaceShadowMaskTexture=*/nullptr, InstanceCullingManager, bNaniteEnabled, NaniteBasePassShadingCommands, NaniteRasterResults);
-			}
+					// Setup dummy uniform buffer parameters for fog volume.
+					SetDummyLocalFogVolumeForViews(GraphBuilder, CustomRenderPassViews);
 
-			CopySceneCaptureComponentToTarget(GraphBuilder, SceneTextures, CustomRenderPass->RenderTargetTexture, ViewFamily, CustomRenderPassViews);
-			CustomRenderPass->PostRender(GraphBuilder);
+					if (bNaniteEnabled && UseNaniteComputeMaterials())
+					{
+						Nanite::BuildShadingCommands(GraphBuilder, *Scene, ENaniteMeshPass::BasePass, NaniteBasePassShadingCommands, true);
+					}
+
+					RenderBasePass(GraphBuilder, CustomRenderPassViews, SceneTextures, DBufferTextures, BasePassDepthStencilAccess, /*ForwardScreenSpaceShadowMaskTexture=*/nullptr, InstanceCullingManager, bNaniteEnabled, NaniteBasePassShadingCommands, NaniteRasterResults);
+				}
+
+				CopySceneCaptureComponentToTarget(GraphBuilder, SceneTextures, CustomRenderPass->GetRenderTargetTexture(), ViewFamily, CustomRenderPassViews);
+				CustomRenderPass->PostRender(GraphBuilder);
 
 #if WITH_MGPU
-			DoCrossGPUTransfers(GraphBuilder, CustomRenderPass->RenderTargetTexture, CustomRenderPassViews, false, FRHIGPUMask::All());
+				DoCrossGPUTransfers(GraphBuilder, CustomRenderPass->GetRenderTargetTexture(), CustomRenderPassViews, false, FRHIGPUMask::All());
 #endif
+			}
+
+			CustomRenderPass->EndPass(GraphBuilder);
 		}
 	}
 
diff --git a/Engine/Source/Runtime/Renderer/Private/MobileShadingRenderer.cpp b/Engine/Source/Runtime/Renderer/Private/MobileShadingRenderer.cpp
--- a/Engine/Source/Runtime/Renderer/Private/MobileShadingRenderer.cpp
+++ b/Engine/Source/Runtime/Renderer/Private/MobileShadingRenderer.cpp
@@ -1139,39 +1139,45 @@ void FMobileSceneRenderer::Render(FRDGBuilder& GraphBuilder)
 
 		for (int32 i = 0; i < CustomRenderPassInfos.Num(); ++i)
 		{
-			FCustomRenderPass* CustomRenderPass = CustomRenderPassInfos[i].CustomRenderPass;
+			FCustomRenderPassBase* CustomRenderPass = CustomRenderPassInfos[i].CustomRenderPass;
 			TArray<FViewInfo>& CustomRenderPassViews = CustomRenderPassInfos[i].Views;
 			check(CustomRenderPass);
 
-			QUICK_SCOPE_CYCLE_COUNTER(STAT_CustomRenderPass);
-			RDG_EVENT_SCOPE(GraphBuilder, "CustomRenderPass[%d] %s", i, *CustomRenderPass->Name);
+			CustomRenderPass->BeginPass(GraphBuilder);
 
-			CustomRenderPass->PreRender(GraphBuilder);
+			{
+				QUICK_SCOPE_CYCLE_COUNTER(STAT_CustomRenderPass);
+				RDG_EVENT_SCOPE(GraphBuilder, "CustomRenderPass[%d] %s", i, *CustomRenderPass->GetDebugName());
 
-			// Setup dummy uniform buffer parameters for fog volume.
-			SetDummyLocalFogVolumeForViews(GraphBuilder, CustomRenderPassViews);
+				CustomRenderPass->PreRender(GraphBuilder);
 
-			if (bIsFullDepthPrepassEnabled)
-			{
-				RenderFullDepthPrepass(GraphBuilder, CustomRenderPassViews, SceneTextures, true);
-				if (!bRequiresSceneDepthAux)
+				// Setup dummy uniform buffer parameters for fog volume.
+				SetDummyLocalFogVolumeForViews(GraphBuilder, CustomRenderPassViews);
+
+				if (bIsFullDepthPrepassEnabled)
 				{
-					AddResolveSceneDepthPass(GraphBuilder, CustomRenderPassViews, SceneTextures.Depth);
+					RenderFullDepthPrepass(GraphBuilder, CustomRenderPassViews, SceneTextures, true);
+					if (!bRequiresSceneDepthAux)
+					{
+						AddResolveSceneDepthPass(GraphBuilder, CustomRenderPassViews, SceneTextures.Depth);
+					}
 				}
-			}
 
-			// Render base pass if the custom pass requires it. Otherwise if full depth prepass is not enabled, then depth is generated in the base pass.
-			if (CustomRenderPass->RenderMode == FCustomRenderPass::ERenderMode_DepthAndBasePass || (CustomRenderPass->RenderMode == FCustomRenderPass::ERenderMode_DepthPass && !bIsFullDepthPrepassEnabled))
-			{
-				RenderCustomRenderPassBasePass(GraphBuilder, CustomRenderPassViews, ViewFamilyTexture, SceneTextures);
-			}
+				// Render base pass if the custom pass requires it. Otherwise if full depth prepass is not enabled, then depth is generated in the base pass.
+				if (CustomRenderPass->GetRenderMode() == FCustomRenderPassBase::ERenderMode::DepthAndBasePass || (CustomRenderPass->GetRenderMode() == FCustomRenderPassBase::ERenderMode::DepthPass && !bIsFullDepthPrepassEnabled))
+				{
+					RenderCustomRenderPassBasePass(GraphBuilder, CustomRenderPassViews, ViewFamilyTexture, SceneTextures);
+				}
 
-			SceneTextures.MobileSetupMode = EMobileSceneTextureSetupMode::SceneColor | EMobileSceneTextureSetupMode::SceneDepth | EMobileSceneTextureSetupMode::SceneDepthAux;
-			SceneTextures.MobileUniformBuffer = CreateMobileSceneTextureUniformBuffer(GraphBuilder, &SceneTextures, SceneTextures.MobileSetupMode);
+				SceneTextures.MobileSetupMode = EMobileSceneTextureSetupMode::SceneColor | EMobileSceneTextureSetupMode::SceneDepth | EMobileSceneTextureSetupMode::SceneDepthAux;
+				SceneTextures.MobileUniformBuffer = CreateMobileSceneTextureUniformBuffer(GraphBuilder, &SceneTextures, SceneTextures.MobileSetupMode);
+
+				CopySceneCaptureComponentToTarget(GraphBuilder, SceneTextures, CustomRenderPass->GetRenderTargetTexture(), ViewFamily, CustomRenderPassViews);
 
-			CopySceneCaptureComponentToTarget(GraphBuilder, SceneTextures, CustomRenderPass->RenderTargetTexture, ViewFamily, CustomRenderPassViews);
+				CustomRenderPass->PostRender(GraphBuilder);
+			}
 
-			CustomRenderPass->PostRender(GraphBuilder);
+			CustomRenderPass->EndPass(GraphBuilder);
 		}
 	}
 
diff --git a/Engine/Source/Runtime/Renderer/Private/SceneCaptureRendering.cpp b/Engine/Source/Runtime/Renderer/Private/SceneCaptureRendering.cpp
--- a/Engine/Source/Runtime/Renderer/Private/SceneCaptureRendering.cpp
+++ b/Engine/Source/Runtime/Renderer/Private/SceneCaptureRendering.cpp
@@ -833,18 +833,21 @@ static FSceneRenderer* CreateSceneRendererForSceneCapture(
 	return FSceneRenderer::CreateSceneRenderer(&ViewFamily, nullptr);
 }
 
-class FSceneCapturePass : public FCustomRenderPass
+class FSceneCapturePass final : public FCustomRenderPassBase
 {
 public:
-	FSceneCapturePass()
-		: FCustomRenderPass()
+	IMPLEMENT_CUSTOM_RENDER_PASS(FSceneCapturePass);
+
+	FSceneCapturePass(const FString& InDebugName, ERenderMode InRenderMode, ERenderOutput InRenderOutput, UTextureRenderTarget2D* InRenderTarget)
+		: FCustomRenderPassBase(InDebugName, InRenderMode, InRenderOutput, FIntPoint(InRenderTarget->GetSurfaceWidth(), InRenderTarget->GetSurfaceHeight()))
+		, SceneCaptureRenderTarget(InRenderTarget->GameThread_GetRenderTargetResource())
 	{}
 
-	virtual void PreRender(FRDGBuilder& GraphBuilder) override
+	virtual void OnPreRender(FRDGBuilder& GraphBuilder) override
 	{
 		RenderTargetTexture = SceneCaptureRenderTarget->GetRenderTargetTexture(GraphBuilder);
 	}
-
+	
 	FRenderTarget* SceneCaptureRenderTarget = nullptr;
 };
 
@@ -920,12 +923,9 @@ void FScene::UpdateSceneCaptureContents(USceneCaptureComponent2D* CaptureCompone
 			PassInput.ProjectionMatrix = ProjectionMatrix;
 			PassInput.ViewActor = CaptureComponent->GetViewOwner();
 
-			FSceneCapturePass* CustomPass = new FSceneCapturePass();
-			CustomPass->RenderMode = FCustomRenderPass::ERenderMode_DepthPass;
-			CustomPass->RenderOutput = CaptureComponent->CaptureSource == ESceneCaptureSource::SCS_SceneDepth ? FCustomRenderPass::ERenderOutput_SceneDepth : FCustomRenderPass::ERenderOutput_DeviceDepth;
-			CustomPass->Name = CaptureComponent->CaptureSource == ESceneCaptureSource::SCS_SceneDepth ? TEXT("SceneCapturePass_SceneDepth") : TEXT("SceneCapturePass_DeviceDepth");
-			CustomPass->SceneCaptureRenderTarget = TextureRenderTarget->GameThread_GetRenderTargetResource();
-			CustomPass->RenderTargetSize = FIntPoint(TextureRenderTarget->GetSurfaceWidth(), TextureRenderTarget->GetSurfaceHeight());
+			FString DebugName = CaptureComponent->CaptureSource == ESceneCaptureSource::SCS_SceneDepth ? TEXT("SceneCapturePass_SceneDepth") : TEXT("SceneCapturePass_DeviceDepth");
+			FCustomRenderPassBase::ERenderOutput RenderOutput = CaptureComponent->CaptureSource == ESceneCaptureSource::SCS_SceneDepth ? FCustomRenderPassBase::ERenderOutput::SceneDepth : FCustomRenderPassBase::ERenderOutput::DeviceDepth;
+			FSceneCapturePass* CustomPass = new FSceneCapturePass(DebugName, FCustomRenderPassBase::ERenderMode::DepthPass, RenderOutput, TextureRenderTarget);
 			PassInput.CustomRenderPass = CustomPass;
 
 			GetShowOnlyAndHiddenComponents(CaptureComponent, PassInput.HiddenPrimitives, PassInput.ShowOnlyPrimitives);
diff --git a/Engine/Source/Runtime/Renderer/Private/SceneRendering.cpp b/Engine/Source/Runtime/Renderer/Private/SceneRendering.cpp
--- a/Engine/Source/Runtime/Renderer/Private/SceneRendering.cpp
+++ b/Engine/Source/Runtime/Renderer/Private/SceneRendering.cpp
@@ -2765,12 +2765,12 @@ FSceneRenderer::FSceneRenderer(const FSceneViewFamily* InViewFamily, FHitProxyCo
 	for (int32 i = 0; i < Scene->CustomRenderPassRendererInputs.Num(); i++)
 	{
 		const FScene::FCustomRenderPassRendererInput& PassInput = Scene->CustomRenderPassRendererInputs[i];
-		FCustomRenderPass* CustomRenderPass = PassInput.CustomRenderPass;
+		FCustomRenderPassBase* CustomRenderPass = PassInput.CustomRenderPass;
 		check(CustomRenderPass);
 		CustomRenderPassInfos[i].CustomRenderPass = CustomRenderPass;
 
 		FSceneViewInitOptions ViewInitOptions;
-		ViewInitOptions.SetViewRectangle(FIntRect(0, 0, CustomRenderPass->RenderTargetSize.X, CustomRenderPass->RenderTargetSize.Y));
+		ViewInitOptions.SetViewRectangle(FIntRect(0, 0, CustomRenderPass->GetRenderTargetSize().X, CustomRenderPass->GetRenderTargetSize().Y));
 		ViewInitOptions.ViewOrigin = PassInput.ViewLocation;
 		ViewInitOptions.ViewRotationMatrix = PassInput.ViewRotationMatrix;
 		ViewInitOptions.ProjectionMatrix = PassInput.ProjectionMatrix;
diff --git a/Engine/Source/Runtime/Renderer/Private/SceneRendering.h b/Engine/Source/Runtime/Renderer/Private/SceneRendering.h
--- a/Engine/Source/Runtime/Renderer/Private/SceneRendering.h
+++ b/Engine/Source/Runtime/Renderer/Private/SceneRendering.h
@@ -2031,7 +2031,7 @@ public:
 	struct FCustomRenderPassInfo
 	{
 		/** Custom render pass that render as part of the main renderer. */
-		FCustomRenderPass* CustomRenderPass;
+		FCustomRenderPassBase* CustomRenderPass;
 		/** Views used to render the custom render pass. */
 		TArray<FViewInfo> Views;
 		FNaniteShadingCommands NaniteBasePassShadingCommands;
