Added additional parameter to ISceneViewExtension::PreRenderBasePass_RenderThread to let the callback know whether the depth buffer has been populated to some degree. This information is used to decide whether to run pixel precise occlusion queries against the depth buffer when generating indirect draws for water. This fixes an issue where the depth buffer was used in read-only mode without any prior path having initialized it with proper data.

#rnx

diff --git a/Engine/Plugins/Experimental/Water/Source/Runtime/Private/WaterMeshSceneProxy.cpp b/Engine/Plugins/Experimental/Water/Source/Runtime/Private/WaterMeshSceneProxy.cpp
--- a/Engine/Plugins/Experimental/Water/Source/Runtime/Private/WaterMeshSceneProxy.cpp
+++ b/Engine/Plugins/Experimental/Water/Source/Runtime/Private/WaterMeshSceneProxy.cpp
@@ -295,7 +295,7 @@ void FWaterMeshSceneProxy::CreateRenderThreadResources(FRHICommandListBase& RHIC
 	{
 		FWaterMeshGPUWork::FCallback Callback;
 		Callback.Proxy = this;
-		Callback.Function = [this](FRDGBuilder& GraphBuilder)
+		Callback.Function = [this](FRDGBuilder& GraphBuilder, bool bDepthBufferIsPopulated)
 		{
 			if (bNeedToTraverseGPUQuadTree)
 			{
@@ -303,8 +303,11 @@ void FWaterMeshSceneProxy::CreateRenderThreadResources(FRHICommandListBase& RHIC
 				{
 					BuildGPUQuadTree(GraphBuilder);
 				}
-				QuadTreeGPU.Traverse(GraphBuilder, WaterQuadTreeGPUTraverseParams);
+				FWaterQuadTreeGPU::FTraverseParams TraverseParams = MoveTemp(WaterQuadTreeGPUTraverseParams);
+				TraverseParams.bDepthBufferIsPopulated = bDepthBufferIsPopulated;
+				QuadTreeGPU.Traverse(GraphBuilder, TraverseParams);
 				bNeedToTraverseGPUQuadTree = false;
+				WaterQuadTreeGPUTraverseParams = {};
 			}
 		};
 		GWaterMeshGPUWork.Callbacks.Add(MoveTemp(Callback));
diff --git a/Engine/Plugins/Experimental/Water/Source/Runtime/Private/WaterQuadTreeGPU.cpp b/Engine/Plugins/Experimental/Water/Source/Runtime/Private/WaterQuadTreeGPU.cpp
--- a/Engine/Plugins/Experimental/Water/Source/Runtime/Private/WaterQuadTreeGPU.cpp
+++ b/Engine/Plugins/Experimental/Water/Source/Runtime/Private/WaterQuadTreeGPU.cpp
@@ -819,7 +819,7 @@ void FWaterQuadTreeGPU::Traverse(FRDGBuilder& GraphBuilder, const FTraverseParam
 		FRDGBufferSRV* PackedNodesSRV = GraphBuilder.CreateSRV(FRDGBufferSRVDesc(PackedNodes, PF_R32_UINT));
 		FRDGBufferSRV* BucketCountsSRV = GraphBuilder.CreateSRV(FRDGBufferSRVDesc(BucketCounts, PF_R32_UINT));
 
-		const bool bPixelPreciseOQForThisView = bPixelPreciseOcclusionQueries && View->bIsViewInfo;
+		const bool bPixelPreciseOQForThisView = bPixelPreciseOcclusionQueries && ViewInfo && Params.bDepthBufferIsPopulated;
 
 		FRDGBuffer* OcclusionQueryBoxes = nullptr;
 		FRDGBuffer* OcclusionQueryResults = nullptr;
@@ -899,7 +899,8 @@ void FWaterQuadTreeGPU::Traverse(FRDGBuilder& GraphBuilder, const FTraverseParam
 			TShaderMapRef<FWaterQuadTreeOcclusionQueryVS> VertexShader(ShaderMap);
 			TShaderMapRef<FWaterQuadTreeOcclusionQueryPS> PixelShader(ShaderMap);
 
-			FRDGTextureRef DepthTexture = ViewInfo->GetSceneTextures().Depth.Target;
+			FRDGTextureRef DepthTexture = GetIfProduced(ViewInfo->GetSceneTextures().Depth.Target);
+			check(DepthTexture);
 
 			FWaterQuadTreeOcclusionQueryParameters* PassParameters = GraphBuilder.AllocParameters<FWaterQuadTreeOcclusionQueryParameters>();
 			PassParameters->IndirectDrawArgsBuffer = OcclusionQueryIndirectArgsBuffer;
diff --git a/Engine/Plugins/Experimental/Water/Source/Runtime/Private/WaterQuadTreeGPU.h b/Engine/Plugins/Experimental/Water/Source/Runtime/Private/WaterQuadTreeGPU.h
--- a/Engine/Plugins/Experimental/Water/Source/Runtime/Private/WaterQuadTreeGPU.h
+++ b/Engine/Plugins/Experimental/Water/Source/Runtime/Private/WaterQuadTreeGPU.h
@@ -79,6 +79,7 @@ public:
 		int32 DebugShowTile = 0;
 		bool bWithWaterSelectionSupport = false;
 		bool bLODMorphingEnabled = false;
+		bool bDepthBufferIsPopulated = false;
 	};
 
 	void Init(FRDGBuilder& GraphBuilder, const FInitParams& Params, TArray<FDraw>& Draws);
diff --git a/Engine/Plugins/Experimental/Water/Source/Runtime/Private/WaterViewExtension.cpp b/Engine/Plugins/Experimental/Water/Source/Runtime/Private/WaterViewExtension.cpp
--- a/Engine/Plugins/Experimental/Water/Source/Runtime/Private/WaterViewExtension.cpp
+++ b/Engine/Plugins/Experimental/Water/Source/Runtime/Private/WaterViewExtension.cpp
@@ -415,11 +415,11 @@ void FWaterViewExtension::PreRenderView_RenderThread(FRDGBuilder& GraphBuilder,
 	}
 }
 
-void FWaterViewExtension::PreRenderBasePass_RenderThread(FRDGBuilder& GraphBuilder)
+void FWaterViewExtension::PreRenderBasePass_RenderThread(FRDGBuilder& GraphBuilder, bool bDepthBufferIsPopulated)
 {
 	for (FWaterMeshGPUWork::FCallback& Callback : GWaterMeshGPUWork.Callbacks)
 	{
-		Callback.Function(GraphBuilder);
+		Callback.Function(GraphBuilder, bDepthBufferIsPopulated);
 	}
 }
 
diff --git a/Engine/Plugins/Experimental/Water/Source/Runtime/Public/WaterViewExtension.h b/Engine/Plugins/Experimental/Water/Source/Runtime/Public/WaterViewExtension.h
--- a/Engine/Plugins/Experimental/Water/Source/Runtime/Public/WaterViewExtension.h
+++ b/Engine/Plugins/Experimental/Water/Source/Runtime/Public/WaterViewExtension.h
@@ -24,7 +24,7 @@ public:
 	virtual void SetupViewFamily(FSceneViewFamily& InViewFamily) override;
 	virtual void PreRenderViewFamily_RenderThread(FRDGBuilder& GraphBuilder, FSceneViewFamily& InViewFamily) override;
 	virtual void PreRenderView_RenderThread(FRDGBuilder& GraphBuilder, FSceneView& InView) override;
-	virtual void PreRenderBasePass_RenderThread(FRDGBuilder& GraphBuilder) override;
+	virtual void PreRenderBasePass_RenderThread(FRDGBuilder& GraphBuilder, bool bDepthBufferIsPopulated) override;
 	// End FSceneViewExtensionBase implementation
 
 	void MarkWaterInfoTextureForRebuild(const UE::WaterInfo::FRenderingContext& RenderContext);
@@ -61,7 +61,7 @@ struct FWaterMeshGPUWork
 	struct FCallback
 	{
 		class FWaterMeshSceneProxy* Proxy = nullptr;
-		TFunction<void(FRDGBuilder&)> Function;
+		TFunction<void(FRDGBuilder&, bool)> Function;
 	};
 	TArray<FCallback> Callbacks;
 };
diff --git a/Engine/Source/Runtime/Engine/Public/SceneViewExtension.h b/Engine/Source/Runtime/Engine/Public/SceneViewExtension.h
--- a/Engine/Source/Runtime/Engine/Public/SceneViewExtension.h
+++ b/Engine/Source/Runtime/Engine/Public/SceneViewExtension.h
@@ -156,9 +156,9 @@ public:
 	virtual void PreInitViews_RenderThread(FRDGBuilder& GraphBuilder) {}
 
 	/**
-	 * Called on render thread right before Base Pass rendering.
+	 * Called on render thread right before Base Pass rendering. bDepthBufferIsPopulated is true if anything has been rendered to the depth buffer. This does not need to be a full depth prepass.
 	 */
-	virtual void PreRenderBasePass_RenderThread(FRDGBuilder& GraphBuilder) {}
+	virtual void PreRenderBasePass_RenderThread(FRDGBuilder& GraphBuilder, bool bDepthBufferIsPopulated) {}
 
 	/**
 	 * Called right after Base Pass rendering finished when using the deferred renderer.
diff --git a/Engine/Source/Runtime/Renderer/Private/DeferredShadingRenderer.cpp b/Engine/Source/Runtime/Renderer/Private/DeferredShadingRenderer.cpp
--- a/Engine/Source/Runtime/Renderer/Private/DeferredShadingRenderer.cpp
+++ b/Engine/Source/Runtime/Renderer/Private/DeferredShadingRenderer.cpp
@@ -2328,7 +2328,7 @@ void FDeferredShadingSceneRenderer::Render(FRDGBuilder& GraphBuilder)
 
 		for (FSceneViewExtensionRef& ViewExtension : ViewFamily.ViewExtensions)
 		{
-			ViewExtension->PreRenderBasePass_RenderThread(GraphBuilder);
+			ViewExtension->PreRenderBasePass_RenderThread(GraphBuilder, ShouldRenderPrePass() /*bDepthBufferIsPopulated*/);
 		}
 
 		BeginAsyncDistanceFieldShadowProjections(GraphBuilder, SceneTextures, InitViewTaskDatas.DynamicShadows);
diff --git a/Engine/Source/Runtime/Renderer/Private/MobileShadingRenderer.cpp b/Engine/Source/Runtime/Renderer/Private/MobileShadingRenderer.cpp
--- a/Engine/Source/Runtime/Renderer/Private/MobileShadingRenderer.cpp
+++ b/Engine/Source/Runtime/Renderer/Private/MobileShadingRenderer.cpp
@@ -1242,7 +1242,7 @@ void FMobileSceneRenderer::Render(FRDGBuilder& GraphBuilder)
 
 	for (FSceneViewExtensionRef& ViewExtension : ViewFamily.ViewExtensions)
 	{
-		ViewExtension->PreRenderBasePass_RenderThread(GraphBuilder);
+		ViewExtension->PreRenderBasePass_RenderThread(GraphBuilder, bIsFullDepthPrepassEnabled /*bDepthBufferIsPopulated*/);
 	}
 
 	if (bRendererOutputFinalSceneColor)
