- Display water body information on Debug HUD
- Remove warning that would spam per instance about water body

diff --git a/Engine/Plugins/Experimental/Water/Source/Runtime/Private/NiagaraDataInterfaceWater.cpp b/Engine/Plugins/Experimental/Water/Source/Runtime/Private/NiagaraDataInterfaceWater.cpp
--- a/Engine/Plugins/Experimental/Water/Source/Runtime/Private/NiagaraDataInterfaceWater.cpp
+++ b/Engine/Plugins/Experimental/Water/Source/Runtime/Private/NiagaraDataInterfaceWater.cpp
@@ -188,6 +188,22 @@ bool UNiagaraDataInterfaceWater::UpgradeFunctionCall(FNiagaraFunctionSignature&
 }
 #endif
 
+#if WITH_NIAGARA_DEBUGGER
+void UNiagaraDataInterfaceWater::DrawDebugHud(FNDIDrawDebugHudContext& DebugHudContext) const
+{
+	Super::DrawDebugHud(DebugHudContext);
+
+	const FNDIWater_InstanceData* InstanceData_GT = DebugHudContext.GetSystemInstance()->FindTypedDataInterfaceInstanceData<FNDIWater_InstanceData>(this);
+	if (InstanceData_GT == nullptr)
+	{
+		return;
+	}
+
+	UWaterBodyComponent* WaterBodyComponent = InstanceData_GT->WaterBodyComponent.Get();
+	DebugHudContext.GetOutputString().Appendf(TEXT("WaterBodyComponent(%s)"), *GetNameSafe(WaterBodyComponent));
+}
+#endif
+
 void UNiagaraDataInterfaceWater::GetWaterDataAtPoint(FVectorVMExternalFunctionContext& Context)
 {
 	QUICK_SCOPE_CYCLE_COUNTER(NiagaraDataInterfaceWater_GetWaterDataAtPoint);
@@ -206,15 +222,10 @@ void UNiagaraDataInterfaceWater::GetWaterDataAtPoint(FVectorVMExternalFunctionCo
 	FNDIOutputParam<FVector3f> OutSurfaceNormal(Context);
 
 	UWaterBodyComponent* Component = InstData->WaterBodyComponent.Get();
-	if (Component == nullptr)
-	{
-		UE_LOG(LogWater, Warning, TEXT("NiagaraDataInterfaceWater: GetWaterData called with no water body component set"));
-	}
-
 	for (int32 i = 0; i < Context.GetNumInstances(); ++i)
 	{
 		FWaterBodyQueryResult QueryResult;
-		
+
 		bool bIsValid = false;
 		if (Component != nullptr)
 		{
@@ -230,11 +241,11 @@ void UNiagaraDataInterfaceWater::GetWaterDataAtPoint(FVectorVMExternalFunctionCo
 
 		OutHeight.SetAndAdvance(bIsValid ? QueryResult.GetWaveInfo().Height : 0.0f);
 		OutDepth.SetAndAdvance(bIsValid ? QueryResult.GetWaterSurfaceDepth() : 0.0f);
-		OutVelocity.SetAndAdvance( bIsValid ? FVector3f(QueryResult.GetVelocity()) : FVector3f::ZeroVector);		// LWC_TODO: Precision loss
+		OutVelocity.SetAndAdvance(bIsValid ? FVector3f(QueryResult.GetVelocity()) : FVector3f::ZeroVector);		// LWC_TODO: Precision loss
 
 		// Note we assume X and Y are in water by the time this is queried
 		const FVector& AdjustedSurfaceLoc = bIsValid ? QueryResult.GetWaterSurfaceLocation() : FVector::ZeroVector;
-		OutSurfacePos.SetAndAdvance( InstData->LWCConverter.ConvertWorldToSimulationPosition(AdjustedSurfaceLoc));
+		OutSurfacePos.SetAndAdvance(InstData->LWCConverter.ConvertWorldToSimulationPosition(AdjustedSurfaceLoc));
 
 		OutSurfaceNormal.SetAndAdvance(bIsValid ? FVector3f(QueryResult.GetWaterSurfaceNormal()) : FVector3f::UpVector);
 
diff --git a/Engine/Plugins/Experimental/Water/Source/Runtime/Public/NiagaraDataInterfaceWater.h b/Engine/Plugins/Experimental/Water/Source/Runtime/Public/NiagaraDataInterfaceWater.h
--- a/Engine/Plugins/Experimental/Water/Source/Runtime/Public/NiagaraDataInterfaceWater.h
+++ b/Engine/Plugins/Experimental/Water/Source/Runtime/Public/NiagaraDataInterfaceWater.h
@@ -30,6 +30,10 @@ public:
 	virtual bool UpgradeFunctionCall(FNiagaraFunctionSignature& FunctionSignature) override;
 #endif
 
+#if WITH_NIAGARA_DEBUGGER
+	virtual void DrawDebugHud(FNDIDrawDebugHudContext& DebugHudContext) const override;
+#endif
+
 	void GetWaterDataAtPoint(FVectorVMExternalFunctionContext& Context);
 
 	void GetWaveParamLookupTableOffset(FVectorVMExternalFunctionContext& Context);
