Water: Added a new parameter to AWaterZone::MarkForRebuild to track the source of the rebuild at runtime. This is useful to track down actors which may be streaming in/out causing expensive water info/mesh rebuilds.

diff --git a/Engine/Plugins/Experimental/Water/Source/Editor/Private/WaterLandscapeBrush.cpp b/Engine/Plugins/Experimental/Water/Source/Editor/Private/WaterLandscapeBrush.cpp
--- a/Engine/Plugins/Experimental/Water/Source/Editor/Private/WaterLandscapeBrush.cpp
+++ b/Engine/Plugins/Experimental/Water/Source/Editor/Private/WaterLandscapeBrush.cpp
@@ -210,7 +210,7 @@ void AWaterLandscapeBrush::OnActorsAffectingLandscapeChanged()
 	RequestLandscapeUpdate();
 	if (UWaterSubsystem* WaterSubsystem = UWaterSubsystem::GetWaterSubsystem(GetWorld()))
 	{
-		WaterSubsystem->MarkAllWaterZonesForRebuild();
+		WaterSubsystem->MarkAllWaterZonesForRebuild(EWaterZoneRebuildFlags::All, /* DebugRequestingObject = */ this);
 	}
 }
 
@@ -489,7 +489,7 @@ void AWaterLandscapeBrush::OnFullHeightmapRenderDone(UTextureRenderTarget2D* InH
 	// #todo_water [roey]: This needs to be changed when the WaterZone can maintain it's own list of "ground actors" so that we don't needlessly update all water zones.
 	if (UWaterSubsystem* WaterSubsystem = UWaterSubsystem::GetWaterSubsystem(GetWorld()))
 	{
-		WaterSubsystem->MarkAllWaterZonesForRebuild(EWaterZoneRebuildFlags::UpdateWaterInfoTexture);
+		WaterSubsystem->MarkAllWaterZonesForRebuild(EWaterZoneRebuildFlags::UpdateWaterInfoTexture, /* DebugRequestingObject = */ this);
 	}
 }
 
diff --git a/Engine/Plugins/Experimental/Water/Source/Runtime/Private/WaterBodyComponent.cpp b/Engine/Plugins/Experimental/Water/Source/Runtime/Private/WaterBodyComponent.cpp
--- a/Engine/Plugins/Experimental/Water/Source/Runtime/Private/WaterBodyComponent.cpp
+++ b/Engine/Plugins/Experimental/Water/Source/Runtime/Private/WaterBodyComponent.cpp
@@ -435,11 +435,11 @@ void UWaterBodyComponent::MarkOwningWaterZoneForRebuild(EWaterZoneRebuildFlags I
 		if (bInOnlyWithinWaterBodyBounds)
 		{
 			const FBox WaterBodyBounds = Bounds.GetBox();
-			WaterZone->MarkForRebuild(InRebuildFlags, FBox2D(FVector2D(WaterBodyBounds.Min), FVector2D(WaterBodyBounds.Max)));
+			WaterZone->MarkForRebuild(InRebuildFlags, FBox2D(FVector2D(WaterBodyBounds.Min), FVector2D(WaterBodyBounds.Max)), /* DebugRequestingObject = */ GetOwner());
 		}
 		else
 		{
-			WaterZone->MarkForRebuild(InRebuildFlags);
+			WaterZone->MarkForRebuild(InRebuildFlags, /* DebugRequestingObject = */ GetOwner());
 		}
 	}
 }
diff --git a/Engine/Plugins/Experimental/Water/Source/Runtime/Private/WaterSubsystem.cpp b/Engine/Plugins/Experimental/Water/Source/Runtime/Private/WaterSubsystem.cpp
--- a/Engine/Plugins/Experimental/Water/Source/Runtime/Private/WaterSubsystem.cpp
+++ b/Engine/Plugins/Experimental/Water/Source/Runtime/Private/WaterSubsystem.cpp
@@ -372,7 +372,7 @@ void UWaterSubsystem::ApplyRuntimeSettings(const UWaterRuntimeSettings* Settings
 void UWaterSubsystem::OnHeightmapStreamed(const FOnHeightmapStreamedContext& InContext)
 {
 	UE_LOG(LogWater, Verbose, TEXT("UWaterSubsystem::OnHeightmapStreamed() -- Rebuilding Water Info Texture..."));
-	MarkWaterZonesInRegionForRebuild(InContext.GetUpdateRegion(), EWaterZoneRebuildFlags::UpdateWaterInfoTexture);
+	MarkWaterZonesInRegionForRebuild(InContext.GetUpdateRegion(), EWaterZoneRebuildFlags::UpdateWaterInfoTexture, /* DebugRequestingObject = */ this);
 }
 #endif // WITH_EDITOR
 
@@ -483,7 +483,7 @@ void UWaterSubsystem::SetOceanFloodHeight(float InFloodHeight)
 		if (FloodHeight != ClampedFloodHeight)
 		{
 			FloodHeight = ClampedFloodHeight;
-			MarkAllWaterZonesForRebuild();
+			MarkAllWaterZonesForRebuild(EWaterZoneRebuildFlags::All, /* DebugRequestingObject = */ this);
 
 			// the ocean body is dynamic and needs to be readjusted when the flood height changes : 
 			if (OceanBodyComponent.IsValid())
@@ -510,18 +510,18 @@ float UWaterSubsystem::GetOceanBaseHeight() const
 	return TNumericLimits<float>::Lowest();
 }
 
-void UWaterSubsystem::MarkAllWaterZonesForRebuild(EWaterZoneRebuildFlags RebuildFlags)
+void UWaterSubsystem::MarkAllWaterZonesForRebuild(EWaterZoneRebuildFlags RebuildFlags, const UObject* DebugRequestingObject)
 {
 	if (UWorld* World = GetWorld())
 	{
 		for (AWaterZone* WaterZone : TActorRange<AWaterZone>(World))
 		{
-			WaterZone->MarkForRebuild(RebuildFlags);
+			WaterZone->MarkForRebuild(RebuildFlags, DebugRequestingObject);
 		}
 	}
 }
 
-void UWaterSubsystem::MarkWaterZonesInRegionForRebuild(const FBox2D& InUpdateRegion, EWaterZoneRebuildFlags InRebuildFlags)
+void UWaterSubsystem::MarkWaterZonesInRegionForRebuild(const FBox2D& InUpdateRegion, EWaterZoneRebuildFlags InRebuildFlags, const UObject* DebugRequestingObject)
 {
 	if (UWorld* World = GetWorld())
 	{
@@ -531,7 +531,7 @@ void UWaterSubsystem::MarkWaterZonesInRegionForRebuild(const FBox2D& InUpdateReg
 
 			if (WaterZoneBounds.Intersect(InUpdateRegion))
 			{
-				WaterZone->MarkForRebuild(InRebuildFlags, InUpdateRegion);
+				WaterZone->MarkForRebuild(InRebuildFlags, InUpdateRegion, DebugRequestingObject);
 			}
 		}
 	}
diff --git a/Engine/Plugins/Experimental/Water/Source/Runtime/Private/WaterZoneActor.cpp b/Engine/Plugins/Experimental/Water/Source/Runtime/Private/WaterZoneActor.cpp
--- a/Engine/Plugins/Experimental/Water/Source/Runtime/Private/WaterZoneActor.cpp
+++ b/Engine/Plugins/Experimental/Water/Source/Runtime/Private/WaterZoneActor.cpp
@@ -111,14 +111,14 @@ FBox AWaterZone::GetZoneBounds() const
 void AWaterZone::SetRenderTargetResolution(FIntPoint NewResolution)
 {
 	RenderTargetResolution = NewResolution;
-	MarkForRebuild(EWaterZoneRebuildFlags::UpdateWaterInfoTexture);
+	MarkForRebuild(EWaterZoneRebuildFlags::UpdateWaterInfoTexture, this);
 }
 
 void AWaterZone::BeginPlay()
 {
 	Super::BeginPlay();
 
-	MarkForRebuild(EWaterZoneRebuildFlags::All);
+	MarkForRebuild(EWaterZoneRebuildFlags::All, this);
 
 	FWorldDelegates::LevelAddedToWorld.AddUObject(this, &AWaterZone::OnLevelAddedToWorld);
 	FWorldDelegates::LevelRemovedFromWorld.AddUObject(this, &AWaterZone::OnLevelRemovedFromWorld);
@@ -189,7 +189,7 @@ void AWaterZone::DeclareConstructClasses(TArray<FTopLevelAssetPath>& OutConstruc
 }
 #endif
 
-void AWaterZone::MarkForRebuild(EWaterZoneRebuildFlags Flags, const FBox2D& UpdateRegion)
+void AWaterZone::MarkForRebuild(EWaterZoneRebuildFlags Flags, const FBox2D& UpdateRegion, const UObject* DebugRequestingObject)
 {
 	if (EnumHasAnyFlags(Flags, EWaterZoneRebuildFlags::UpdateWaterMesh))
 	{
@@ -197,7 +197,7 @@ void AWaterZone::MarkForRebuild(EWaterZoneRebuildFlags Flags, const FBox2D& Upda
 		// Suppress water mesh updates which occur outside the bounds of the water quad tree.
 		if ((!UpdateRegion.bIsValid) || UpdateRegion.Intersect(WaterQuadTreeBounds))
 		{
-			UE_LOG(LogWater, Verbose, TEXT("AWaterZone::MarkForRebuild (UpdateWaterMesh)"));
+			UE_LOG(LogWater, Verbose, TEXT("AWaterZone::MarkForRebuild (UpdateWaterMesh) in region {%s} (triggered by %s)"), *UpdateRegion.ToString(), *GetNameSafe(DebugRequestingObject));
 			WaterMesh->MarkWaterMeshGridDirty();
 			WaterMesh->MarkRenderStateDirty();
 		}
@@ -208,12 +208,17 @@ void AWaterZone::MarkForRebuild(EWaterZoneRebuildFlags Flags, const FBox2D& Upda
 		const FBox2D WaterInfoBounds2D(FVector2D(WaterInfoBounds.Min), FVector2D(WaterInfoBounds.Max));
 		if ((!UpdateRegion.bIsValid) || UpdateRegion.Intersect(WaterInfoBounds2D))
 		{
-			UE_LOG(LogWater, Verbose, TEXT("AWaterZone::MarkForRebuild (UpdateWaterInfoTexture)"));
+			UE_LOG(LogWater, Verbose, TEXT("AWaterZone::MarkForRebuild (UpdateWaterInfoTexture) in region {%s} (triggered by %s)"), *UpdateRegion.ToString(), *GetNameSafe(DebugRequestingObject));
 			bNeedsWaterInfoRebuild = true;
 		}
 	}
 }
 
+void AWaterZone::MarkForRebuild(EWaterZoneRebuildFlags Flags, const UObject* DebugRequestingObject)
+{
+	MarkForRebuild(Flags, FBox2D(ForceInitToZero), DebugRequestingObject);
+}
+
 void AWaterZone::ForEachWaterBodyComponent(TFunctionRef<bool(UWaterBodyComponent*)> Predicate) const
 {
 	for (const TWeakObjectPtr<UWaterBodyComponent>& WaterBodyComponent : OwnedWaterBodies)
@@ -242,7 +247,7 @@ void AWaterZone::AddWaterBodyComponent(UWaterBodyComponent* WaterBodyComponent)
 		}
 
 		const FBox WaterBodyBounds = WaterBodyComponent->Bounds.GetBox();
-		MarkForRebuild(RebuildFlags, FBox2D(FVector2D(WaterBodyBounds.Min), FVector2D(WaterBodyBounds.Max)));
+		MarkForRebuild(RebuildFlags, FBox2D(FVector2D(WaterBodyBounds.Min), FVector2D(WaterBodyBounds.Max)), /* DebugRequestingObject = */ WaterBodyComponent->GetOwner());
 	}
 }
 
@@ -264,7 +269,7 @@ void AWaterZone::RemoveWaterBodyComponent(UWaterBodyComponent* WaterBodyComponen
 		}
 		
 		const FBox WaterBodyBounds = WaterBodyComponent->Bounds.GetBox();
-		MarkForRebuild(RebuildFlags, FBox2D(FVector2D(WaterBodyBounds.Min), FVector2D(WaterBodyBounds.Max)));
+		MarkForRebuild(RebuildFlags, FBox2D(FVector2D(WaterBodyBounds.Min), FVector2D(WaterBodyBounds.Max)), /* DebugRequestingObject = */ WaterBodyComponent->GetOwner());
 	}
 }
 
@@ -305,7 +310,7 @@ void AWaterZone::PostEditMove(bool bFinished)
 
 	UpdateOverlappingWaterBodies();
 
-	MarkForRebuild(RebuildFlags);
+	MarkForRebuild(RebuildFlags, /* DebugRequestingObject = */ this);
 }
 
 void AWaterZone::PostEditUndo()
@@ -313,7 +318,7 @@ void AWaterZone::PostEditUndo()
 	Super::PostEditUndo();
 
 	UpdateOverlappingWaterBodies();
-	MarkForRebuild(EWaterZoneRebuildFlags::All);
+	MarkForRebuild(EWaterZoneRebuildFlags::All, /* DebugRequestingObject = */ this);
 }
 
 void AWaterZone::PostEditImport()
@@ -321,7 +326,7 @@ void AWaterZone::PostEditImport()
 	Super::PostEditImport();
 
 	UpdateOverlappingWaterBodies();
-	MarkForRebuild(EWaterZoneRebuildFlags::All);
+	MarkForRebuild(EWaterZoneRebuildFlags::All, /* DebugRequestingObject = */ this);
 }
 
 void AWaterZone::PostEditChangeProperty(FPropertyChangedEvent& PropertyChangedEvent)
@@ -339,27 +344,27 @@ void AWaterZone::PostEditChangeProperty(FPropertyChangedEvent& PropertyChangedEv
 	}
 	else if (PropertyName == GET_MEMBER_NAME_CHECKED(AWaterZone, RenderTargetResolution))
 	{
-		MarkForRebuild(EWaterZoneRebuildFlags::UpdateWaterInfoTexture);
+		MarkForRebuild(EWaterZoneRebuildFlags::UpdateWaterInfoTexture, /* DebugRequestingObject = */ this);
 	}
 	else if (PropertyName == GET_MEMBER_NAME_CHECKED(AWaterZone, bHalfPrecisionTexture))
 	{
-		MarkForRebuild(EWaterZoneRebuildFlags::UpdateWaterInfoTexture);
+		MarkForRebuild(EWaterZoneRebuildFlags::UpdateWaterInfoTexture, /* DebugRequestingObject = */ this);
 	}
 	else if (PropertyName == GET_MEMBER_NAME_CHECKED(AWaterZone, VelocityBlurRadius))
 	{
-		MarkForRebuild(EWaterZoneRebuildFlags::UpdateWaterInfoTexture);
+		MarkForRebuild(EWaterZoneRebuildFlags::UpdateWaterInfoTexture, /* DebugRequestingObject = */ this);
 	}
 	else if (PropertyName == GET_MEMBER_NAME_CHECKED(AWaterZone, LocalTessellationExtent))
 	{
-		MarkForRebuild(EWaterZoneRebuildFlags::All);
+		MarkForRebuild(EWaterZoneRebuildFlags::All, /* DebugRequestingObject = */ this);
 	}
 	else if (PropertyName == GET_MEMBER_NAME_CHECKED(AWaterZone, bEnableLocalOnlyTessellation))
 	{
-		MarkForRebuild(EWaterZoneRebuildFlags::All);
+		MarkForRebuild(EWaterZoneRebuildFlags::All, /* DebugRequestingObject = */ this);
 	}
 	else if (PropertyName == USceneComponent::GetRelativeScale3DPropertyName())
 	{
-		MarkForRebuild(EWaterZoneRebuildFlags::All);
+		MarkForRebuild(EWaterZoneRebuildFlags::All, /* DebugRequestingObject = */ this);
 	}
 }
 
@@ -396,7 +401,7 @@ void AWaterZone::SetFarMeshMaterial(UMaterialInterface* InFarDistanceMaterial)
 	if (WaterMesh && InFarDistanceMaterial != WaterMesh->FarDistanceMaterial)
 	{
 		WaterMesh->FarDistanceMaterial = InFarDistanceMaterial;
-		MarkForRebuild(EWaterZoneRebuildFlags::UpdateWaterMesh);
+		MarkForRebuild(EWaterZoneRebuildFlags::UpdateWaterMesh, /* DebugRequestingObject = */ this);
 	}
 }
 
@@ -409,7 +414,7 @@ void AWaterZone::OnExtentChanged()
 
 	UpdateOverlappingWaterBodies();
 
-	MarkForRebuild(EWaterZoneRebuildFlags::All);
+	MarkForRebuild(EWaterZoneRebuildFlags::All, /* DebugRequestingObject = */ this);
 }
 
 bool AWaterZone::UpdateOverlappingWaterBodies()
@@ -671,7 +676,7 @@ void AWaterZone::OnLevelChanged(ULevel* InLevel, UWorld* InWorld)
 
 	if (bContainsActorsAffectingWaterZone)
 	{
-		MarkForRebuild(EWaterZoneRebuildFlags::UpdateWaterInfoTexture);
+		MarkForRebuild(EWaterZoneRebuildFlags::UpdateWaterInfoTexture, /* DebugRequestingObject = */ InLevel);
 	}
 }
 
diff --git a/Engine/Plugins/Experimental/Water/Source/Runtime/Public/WaterSubsystem.h b/Engine/Plugins/Experimental/Water/Source/Runtime/Public/WaterSubsystem.h
--- a/Engine/Plugins/Experimental/Water/Source/Runtime/Public/WaterSubsystem.h
+++ b/Engine/Plugins/Experimental/Water/Source/Runtime/Public/WaterSubsystem.h
@@ -164,8 +164,8 @@ public:
 
 	UMaterialParameterCollection* GetMaterialParameterCollection() const {	return MaterialParameterCollection; }
 	
-	void MarkAllWaterZonesForRebuild(EWaterZoneRebuildFlags RebuildFlags = EWaterZoneRebuildFlags::All);
-	void MarkWaterZonesInRegionForRebuild(const FBox2D& InUpdateRegion, EWaterZoneRebuildFlags InRebuildFlags);
+	void MarkAllWaterZonesForRebuild(EWaterZoneRebuildFlags RebuildFlags = EWaterZoneRebuildFlags::All, const UObject* DebugRequestingObject = nullptr);
+	void MarkWaterZonesInRegionForRebuild(const FBox2D& InUpdateRegion, EWaterZoneRebuildFlags InRebuildFlags, const UObject* DebugRequestingObject = nullptr);
 
 	/** Returns the water with the highest priority within the bounds provided. */
 	static TSoftObjectPtr<AWaterZone> FindWaterZone(const UWorld* World, const FBox2D& Bounds, const TSoftObjectPtr<const ULevel> PreferredLevel = {});
diff --git a/Engine/Plugins/Experimental/Water/Source/Runtime/Public/WaterZoneActor.h b/Engine/Plugins/Experimental/Water/Source/Runtime/Public/WaterZoneActor.h
--- a/Engine/Plugins/Experimental/Water/Source/Runtime/Public/WaterZoneActor.h
+++ b/Engine/Plugins/Experimental/Water/Source/Runtime/Public/WaterZoneActor.h
@@ -35,7 +35,11 @@ public:
 	UWaterMeshComponent* GetWaterMeshComponent() { return WaterMesh; }
 	const UWaterMeshComponent* GetWaterMeshComponent() const { return WaterMesh; }
 
-	void MarkForRebuild(EWaterZoneRebuildFlags Flags, const FBox2D& RebuildRegion = FBox2D(EForceInit::ForceInitToZero));
+	/** Mark aspects of the water zone for rebuild based on the Flags parameter within a given region. Optionally the caller can pass in a UObject to identify who requested the update. */
+	void MarkForRebuild(EWaterZoneRebuildFlags Flags, const FBox2D& RebuildRegion, const UObject* DebugRequestingObject = nullptr);
+	/** Mark aspects of the water zone for rebuild based on the Flags parameter. Optionally the caller can pass in a UObject to identify who requested the update. */
+	void MarkForRebuild(EWaterZoneRebuildFlags Flags, const UObject* DebugRequestingObject = nullptr);
+
 	void Update();
 		
 	/** Execute a predicate function on each valid water body within the water zone. Predicate should return false for early exit. */
