Implemented ISR specific path for the new GPU driven water rendering. The non-ISR path uses instance rate vertex streams with non-zero StartInstanceLocation to offset into the underlying buffer to the correct offset for a given indirect draw bucket. The new ISR path instead directly binds an offset buffer and instance data buffers as SRVs to the vertex factory. It then uses the draw bucket index and the SV_InstanceID to find the correct instance data index and manually fetches the data from the SRVs. This is likely slightly less efficient than the non-ISR path because there is now a dependent read in every VS thread before the instance data can be fetched.
Also fixed a bug when preparing indirect draw calls for multiple views (such as for stereo rendering without instancing). The prefix sums/offsets of instance data was only relative to each view but since the instance data buffers are shared between views, they need to be global/take the number of instances in prior views into consideration to get the correct offset.

Most of the ISR relevant changes are in WaterMeshVertexFactory.ush, WaterQuadTreeDraws.usf and WaterQuadTreeGPU.cpp.

#rnx

diff --git a/Engine/Plugins/Experimental/Water/Shaders/Private/WaterMeshVertexFactory.ush b/Engine/Plugins/Experimental/Water/Shaders/Private/WaterMeshVertexFactory.ush
--- a/Engine/Plugins/Experimental/Water/Shaders/Private/WaterMeshVertexFactory.ush
+++ b/Engine/Plugins/Experimental/Water/Shaders/Private/WaterMeshVertexFactory.ush
@@ -11,8 +11,22 @@
 #if WATER_MESH_DRAW_INDIRECT_INTERNAL
 float3 QuadTreePosition;
 uint bLODMorphingEnabled;
+
+#if INSTANCED_STEREO
+Buffer<uint> InstanceDataOffsetsBuffer;
+Buffer<uint> InstanceData0Buffer;
+Buffer<uint> InstanceData1Buffer;
+Buffer<uint> InstanceData2Buffer;
+#if WITH_WATER_SELECTION_SUPPORT_VF
+Buffer<uint> InstanceData3Buffer;
 #endif
 
+uint DrawBucketIndex;
+uint StereoPassInstanceFactor;
+#endif // INSTANCED_STEREO
+
+#endif // WATER_MESH_DRAW_INDIRECT_INTERNAL
+
 struct FVertexFactoryInterpolantsVSToPS
 {
 #if NUM_TEX_COORD_INTERPOLATORS
@@ -38,13 +52,16 @@ struct FVertexFactoryInput
 	float4	Position	: ATTRIBUTE0;
 
 #if WATER_MESH_DRAW_INDIRECT_INTERNAL
+#if !INSTANCED_STEREO // ISR uses manual vertex fetching instead
 	uint InstanceData0 : ATTRIBUTE8;
 	uint InstanceData1 : ATTRIBUTE9;
 	uint InstanceData2 : ATTRIBUTE10;
 #if HIT_PROXY_SHADER || WITH_WATER_SELECTION_SUPPORT_VF
 	uint InstanceData3 : ATTRIBUTE11;
 #endif
-#else
+#endif // !INSTANCED_STEREO
+
+#else // WATER_MESH_DRAW_INDIRECT_INTERNAL
 	float4 InstanceData0 : ATTRIBUTE8;
 	float4 InstanceData1 : ATTRIBUTE9; 
 #if HIT_PROXY_SHADER || WITH_WATER_SELECTION_SUPPORT_VF
@@ -64,10 +81,13 @@ struct FPositionOnlyVertexFactoryInput
 	float4	Position	: ATTRIBUTE0;
 
 #if WATER_MESH_DRAW_INDIRECT_INTERNAL
+#if !INSTANCED_STEREO // ISR uses manual vertex fetching instead
 	uint InstanceData0 : ATTRIBUTE8;
 	uint InstanceData1 : ATTRIBUTE9;
 	uint InstanceData2 : ATTRIBUTE10;
-#else
+#endif // !INSTANCED_STEREO
+
+#else // WATER_MESH_DRAW_INDIRECT_INTERNAL
 	float4 InstanceData0 : ATTRIBUTE8;
 	float4 InstanceData1 : ATTRIBUTE9;
 #endif // WATER_MESH_DRAW_INDIRECT_INTERNAL
@@ -85,9 +105,12 @@ struct FPositionAndNormalOnlyVertexFactoryInput
 	float4	Normal		: ATTRIBUTE2;
 
 #if WATER_MESH_DRAW_INDIRECT_INTERNAL
+#if !INSTANCED_STEREO // ISR uses manual vertex fetching instead
 	uint InstanceData0 : ATTRIBUTE8;
 	uint InstanceData1 : ATTRIBUTE9;
 	uint InstanceData2 : ATTRIBUTE10;
+#endif // !INSTANCED_STEREO
+
 #else
 	float4 InstanceData0 : ATTRIBUTE8;
 	float4 InstanceData1 : ATTRIBUTE9;
@@ -334,7 +357,19 @@ FVertexFactoryIntermediates GetVertexFactoryIntermediates(FVertexFactoryInput In
 	FVertexFactoryIntermediates Intermediates;
 
 #if WATER_MESH_DRAW_INDIRECT_INTERNAL
-	const FWaterVertexFactoryInstanceInput InstanceInput = UnpackWaterVertexFactoryInstanceInput(Input.Position, Input.InstanceData0, Input.InstanceData1, Input.InstanceData2);
+#if INSTANCED_STEREO
+	const uint InstanceIndex = Input.InstanceId / StereoPassInstanceFactor;
+	const uint InstanceDataIndex = InstanceDataOffsetsBuffer.Load(DrawBucketIndex) + InstanceIndex;
+	const uint InstanceData0 = InstanceData0Buffer.Load(InstanceDataIndex);
+	const uint InstanceData1 = InstanceData1Buffer.Load(InstanceDataIndex);
+	const uint InstanceData2 = InstanceData2Buffer.Load(InstanceDataIndex);
+#else
+	const uint InstanceData0 = Input.InstanceData0;
+	const uint InstanceData1 = Input.InstanceData1;
+	const uint InstanceData2 = Input.InstanceData2;
+#endif
+
+	const FWaterVertexFactoryInstanceInput InstanceInput = UnpackWaterVertexFactoryInstanceInput(Input.Position, InstanceData0, InstanceData1, InstanceData2);
 #else
 	const FWaterVertexFactoryInstanceInput InstanceInput = UnpackWaterVertexFactoryInstanceInput(Input.Position, Input.InstanceData0, Input.InstanceData1);
 #endif
@@ -373,8 +408,12 @@ FVertexFactoryIntermediates GetVertexFactoryIntermediates(FVertexFactoryInput In
 
 #if WITH_WATER_SELECTION_SUPPORT_VF
 #if WATER_MESH_DRAW_INDIRECT_INTERNAL
+	#if INSTANCED_STEREO
+	float SelectedValue = UnpackRGBA8(InstanceData3Buffer.Load(InstanceDataIndex)).w;
+	#else
 	float SelectedValue = UnpackRGBA8(Input.InstanceData2).w;
-#else
+	#endif // INSTANCED_STEREO
+#else // WATER_MESH_DRAW_INDIRECT_INTERNAL
 	float SelectedValue = Input.InstanceData2.w;
 #endif
 	float IsVisible = WaterVF.bRenderSelected * SelectedValue + WaterVF.bRenderUnselected * (1-SelectedValue);
@@ -389,10 +428,16 @@ float4 VertexFactoryGetInstanceHitProxyId(FVertexFactoryInput Input, FVertexFact
 {
 #if WITH_WATER_SELECTION_SUPPORT_VF
 #if WATER_MESH_DRAW_INDIRECT_INTERNAL
+	#if INSTANCED_STEREO
+	const uint InstanceIndex = Input.InstanceId / StereoPassInstanceFactor;
+	const uint InstanceDataIndex = InstanceDataOffsetsBuffer.Load(DrawBucketIndex) + InstanceIndex;
+	return float4(UnpackRGBA8(InstanceData3Buffer.Load(InstanceDataIndex)).bgr, 0);
+	#else
 	return float4(UnpackRGBA8(Input.InstanceData2).bgr, 0);
-#else
+	#endif // INSTANCED_STEREO
+#else // WATER_MESH_DRAW_INDIRECT_INTERNAL
 	return float4(Input.InstanceData2.rgb, 0);
-#endif
+#endif // WATER_MESH_DRAW_INDIRECT_INTERNAL
 #else // WITH_WATER_SELECTION_SUPPORT_VF
 	return 0;
 #endif // !WITH_WATER_SELECTION_SUPPORT_VF
diff --git a/Engine/Plugins/Experimental/Water/Shaders/Private/WaterQuadTreeDraws.usf b/Engine/Plugins/Experimental/Water/Shaders/Private/WaterQuadTreeDraws.usf
--- a/Engine/Plugins/Experimental/Water/Shaders/Private/WaterQuadTreeDraws.usf
+++ b/Engine/Plugins/Experimental/Water/Shaders/Private/WaterQuadTreeDraws.usf
@@ -487,6 +487,7 @@ RWBuffer<uint> BucketPrefixSums;
 ByteAddressBuffer BucketCounts;
 uint NumBuckets;
 uint OutputOffset;
+uint bWriteTotalSumAtBufferEnd;
 
 #if PARALLEL_PREFIX_SUM
 
@@ -496,6 +497,7 @@ uint OutputOffset;
 #define ARRAY_SIZE (GROUP_SIZE * 2)
 groupshared uint SharedData[ARRAY_SIZE];
 groupshared uint BlockPrefixSum;
+groupshared uint TotalGlobalSum;
 
 uint LoadFromBuffer(uint InIndex)
 {
@@ -533,7 +535,17 @@ void MainCS(uint3 DispatchThreadId : SV_DispatchThreadID)
 
 	if (ThreadID == 0)
 	{
-		BlockPrefixSum = 0;
+		// Apply global offset if this is not the first view. All views share the same prefix sum and instance data buffers.
+		if (OutputOffset > 0)
+		{
+			// Previous view must have written the total sum after the end of its range. This corresponds to the first element in this range.
+			BlockPrefixSum = BucketPrefixSums[OutputOffset];
+		}
+		else
+		{
+			BlockPrefixSum = 0;
+		}
+		TotalGlobalSum = BlockPrefixSum;
 	}
 
 	// Process array in blocks of fixed size, keeping track of the sum of all elements in previous blocks
@@ -551,6 +563,8 @@ void MainCS(uint3 DispatchThreadId : SV_DispatchThreadID)
 		SharedData[2 * ThreadID + 0] = Value0;
 		SharedData[2 * ThreadID + 1] = Value1;
 
+		WaveInterlockedAdd(TotalGlobalSum, LocalSum);
+
 		// Up-sweep
 		for (uint d = ARRAY_SIZE >> 1; d > 0; d >>= 1)
 		{
@@ -603,16 +617,41 @@ void MainCS(uint3 DispatchThreadId : SV_DispatchThreadID)
 		}
 	}
 
+	// Write total of all values (including those of prior views) at element 0 of the next view.
+	// This way we can propagate the prefix sum offsets across views.
+	if (bWriteTotalSumAtBufferEnd != 0)
+	{
+		if (ThreadID == 0)
+		{
+			BucketPrefixSums[OutputOffset + NumBuckets] = TotalGlobalSum;
+		}
+	}
+
 #else
 
 	uint PrefixSum = 0;
+
+	// Apply global offset if this is not the first view. All views share the same prefix sum and instance data buffers.
+	if (OutputOffset > 0)
+	{
+		// Previous view must have written the total sum after the end of its range. This corresponds to the first element in this range.
+		PrefixSum = BucketPrefixSums[OutputOffset];
+	}
+
 	for (uint BucketIndex = 0; BucketIndex < NumBuckets; ++BucketIndex)
 	{
-		BucketPrefixSums[BucketIndex + OutputOffset] = PrefixSum;
+		BucketPrefixSums[OutputOffset + BucketIndex] = PrefixSum;
 		const uint Count = BucketCounts.Load(BucketIndex << 2);
 		PrefixSum += Count;
 	}
 
+	// Write total of all values (including those of prior views) at element 0 of the next view.
+	// This way we can propagate the prefix sum offsets across views.
+	if (bWriteTotalSumAtBufferEnd != 0)
+	{
+		BucketPrefixSums[OutputOffset + NumBuckets] = PrefixSum;
+	}
+
 #endif // PARALLEL_PREFIX_SUM
 }
 
@@ -652,8 +691,10 @@ uint NumQuadsPerDraw;
 float LeafSize;
 float LODScale;
 float CaptureDepthRange;
+uint StereoPassInstanceFactor;
 uint bWithWaterSelectionSupport;
 uint bLODMorphingEnabled;
+uint bInstancedStereoRendering;
 
 [numthreads(64, 1, 1)]
 void MainCS(uint3 DispatchThreadId : SV_DispatchThreadID)
@@ -706,14 +747,18 @@ void MainCS(uint3 DispatchThreadId : SV_DispatchThreadID)
 		const uint BucketInstanceDataOffset = InstanceDataOffsets.Load(BucketIndex);
 		const uint NumTilesPerEdge = max(NumQuadsPerDraw, NumQuadsLOD0 >> DensityIndexClamped) / NumQuadsPerDraw;
 		const uint NumInstances = NumTilesPerEdge * NumTilesPerEdge;
+		const uint NumInstancesStereo = NumInstances * StereoPassInstanceFactor;
 		
 		// Increment InstanceCount
-		uint LocalInstanceDataOffset;
-		InterlockedAdd(IndirectArgs[BucketIndex * 5 + 1], NumInstances, LocalInstanceDataOffset); 
+		uint LocalInstanceDataOffsetStereo;
+		InterlockedAdd(IndirectArgs[BucketIndex * 5 + 1], NumInstancesStereo, LocalInstanceDataOffsetStereo);
+		const uint LocalInstanceDataOffset = LocalInstanceDataOffsetStereo / StereoPassInstanceFactor;
 		const uint InstanceDataBaseIndex = BucketInstanceDataOffset + LocalInstanceDataOffset;
 
-		// On the first write to this bucket, set StartInstanceLocation to offset where the per-instance vertex attributes are read from
-		if (LocalInstanceDataOffset == 0)
+		// On the first write to this bucket, set StartInstanceLocation to offset where the per-instance vertex attributes are read from.
+		// With ISR, we use the InstanceId to manually fetch instance data buffers in the vertex factory.
+		// StartInstanceLocation affects InstanceId differently on different platforms, so when using InstanceId, StartInstanceLocation must be 0.
+		if (LocalInstanceDataOffset == 0 && bInstancedStereoRendering == 0)
 		{
 			IndirectArgs[BucketIndex * 5 + 4] = BucketInstanceDataOffset;
 		}
diff --git a/Engine/Plugins/Experimental/Water/Source/Runtime/Private/WaterMeshComponent.cpp b/Engine/Plugins/Experimental/Water/Source/Runtime/Private/WaterMeshComponent.cpp
--- a/Engine/Plugins/Experimental/Water/Source/Runtime/Private/WaterMeshComponent.cpp
+++ b/Engine/Plugins/Experimental/Water/Source/Runtime/Private/WaterMeshComponent.cpp
@@ -90,13 +90,17 @@ void UWaterMeshComponent::PostLoad()
 
 void UWaterMeshComponent::CollectPSOPrecacheData(const FPSOPrecacheParams& BasePrecachePSOParams, FMaterialInterfacePSOPrecacheParamsList& OutParams)
 {
-	const FVertexFactoryType* WaterVertexFactoryType = GetWaterVertexFactoryType(/*bWithWaterSelectionSupport = */ false, /*bIndirectDraws = */ false);
+	const FVertexFactoryType* WaterVertexFactoryNonIndirectType = GetWaterVertexFactoryType(/*bWithWaterSelectionSupport = */ false, EWaterVertexFactoryDrawMode::NonIndirect);
+	const FVertexFactoryType* WaterVertexFactoryIndirectType = GetWaterVertexFactoryType(/*bWithWaterSelectionSupport = */ false, EWaterVertexFactoryDrawMode::Indirect);
+	const FVertexFactoryType* WaterVertexFactoryIndirectISRType = GetWaterVertexFactoryType(/*bWithWaterSelectionSupport = */ false, EWaterVertexFactoryDrawMode::IndirectInstancedStereo);
 	if (FarDistanceMaterial)
 	{
 		FMaterialInterfacePSOPrecacheParams& ComponentParams = OutParams[OutParams.AddDefaulted()];
 		ComponentParams.Priority = EPSOPrecachePriority::High;
 		ComponentParams.MaterialInterface = FarDistanceMaterial;
-		ComponentParams.VertexFactoryDataList.Add(FPSOPrecacheVertexFactoryData(WaterVertexFactoryType));
+		ComponentParams.VertexFactoryDataList.Add(FPSOPrecacheVertexFactoryData(WaterVertexFactoryNonIndirectType));
+		ComponentParams.VertexFactoryDataList.Add(FPSOPrecacheVertexFactoryData(WaterVertexFactoryIndirectType));
+		ComponentParams.VertexFactoryDataList.Add(FPSOPrecacheVertexFactoryData(WaterVertexFactoryIndirectISRType));
 		ComponentParams.PSOPrecacheParams = BasePrecachePSOParams;
 	}
 	for (UMaterialInterface* MaterialInterface : UsedMaterials)
@@ -106,7 +110,9 @@ void UWaterMeshComponent::CollectPSOPrecacheData(const FPSOPrecacheParams& BaseP
 			FMaterialInterfacePSOPrecacheParams& ComponentParams = OutParams[OutParams.AddDefaulted()];
 			ComponentParams.Priority = EPSOPrecachePriority::High;
 			ComponentParams.MaterialInterface = MaterialInterface;
-			ComponentParams.VertexFactoryDataList.Add(FPSOPrecacheVertexFactoryData(WaterVertexFactoryType));
+			ComponentParams.VertexFactoryDataList.Add(FPSOPrecacheVertexFactoryData(WaterVertexFactoryNonIndirectType));
+			ComponentParams.VertexFactoryDataList.Add(FPSOPrecacheVertexFactoryData(WaterVertexFactoryIndirectType));
+			ComponentParams.VertexFactoryDataList.Add(FPSOPrecacheVertexFactoryData(WaterVertexFactoryIndirectISRType));
 			ComponentParams.PSOPrecacheParams = BasePrecachePSOParams;
 		}
 	}
diff --git a/Engine/Plugins/Experimental/Water/Source/Runtime/Private/WaterMeshSceneProxy.cpp b/Engine/Plugins/Experimental/Water/Source/Runtime/Private/WaterMeshSceneProxy.cpp
--- a/Engine/Plugins/Experimental/Water/Source/Runtime/Private/WaterMeshSceneProxy.cpp
+++ b/Engine/Plugins/Experimental/Water/Source/Runtime/Private/WaterMeshSceneProxy.cpp
@@ -17,6 +17,7 @@
 #include "SceneInterface.h"
 #include "RenderGraphBuilder.h"
 #include "RenderGraphUtils.h"
+#include "StereoRenderUtils.h"
 
 DECLARE_STATS_GROUP(TEXT("Water Mesh"), STATGROUP_WaterMesh, STATCAT_Advanced);
 
@@ -179,8 +180,10 @@ FWaterMeshSceneProxy::FWaterMeshSceneProxy(UWaterMeshComponent* Component)
 	NumQuadsLOD0 = NumQuads;
 	NumQuadsPerIndirectDrawTile = FMath::Min((int32)FMath::RoundUpToPowerOfTwo(FMath::Clamp(CVarWaterMeshGPUQuadTreeNumQuads.GetValueOnGameThread(), 2, 128)), NumQuadsLOD0);
 
+	const bool bIsGPUQuadTree = WaterQuadTree.IsGPUQuadTree();
+
 	// Initialize Z bounds needed for GPU driven rendering
-	if (WaterQuadTree.IsGPUQuadTree())
+	if (bIsGPUQuadTree)
 	{
 		WaterQuadTreeMinHeight = DBL_MAX;
 		WaterQuadTreeMaxHeight = -DBL_MAX;
@@ -206,9 +209,22 @@ FWaterMeshSceneProxy::FWaterMeshSceneProxy(UWaterMeshComponent* Component)
 	const FVector QuadTreeCorner = FVector(WaterQuadTree.GetTileRegion().Min, WaterQuadTreeMinHeight);
 	const float WaterQuadTreeDepthRange = WaterQuadTreeMaxHeight - WaterQuadTreeMinHeight;
 	const float LeafSize = WaterQuadTree.GetLeafSize();
-	WaterVertexFactoryIndirectDraw = new FWaterVertexFactoryIndirectDrawType(GetScene().GetFeatureLevel(), QuadTreeCorner, NumQuadsPerIndirectDrawTile, NumQuadsLOD0, DensityCount, LeafSize, LODScale, WaterQuadTreeDepthRange);
-	BeginInitResource(WaterVertexFactoryIndirectDraw);
+	if (bIsGPUQuadTree)
+	{
+		// bIsGPUQuadTree is constant over the lifetime of this scene proxy, so we know up front if we need the indirect draw vertex factory
+		WaterVertexFactoryIndirectDraw = new FWaterVertexFactoryIndirectDrawType(GetScene().GetFeatureLevel(), QuadTreeCorner, NumQuadsPerIndirectDrawTile, NumQuadsLOD0, DensityCount, LeafSize, LODScale, WaterQuadTreeDepthRange);
+		BeginInitResource(WaterVertexFactoryIndirectDraw);
 
+		// Only create the ISR vertex factory if ISR is enabled in the first place
+		const UE::StereoRenderUtils::FStereoShaderAspects Aspects(GetScene().GetShaderPlatform());
+		if (Aspects.IsInstancedStereoEnabled())
+		{
+			WaterVertexFactoryIndirectDrawISR = new FWaterVertexFactoryIndirectDrawISRType(GetScene().GetFeatureLevel(), QuadTreeCorner, NumQuadsPerIndirectDrawTile, NumQuadsLOD0, DensityCount, LeafSize, LODScale, WaterQuadTreeDepthRange);
+			BeginInitResource(WaterVertexFactoryIndirectDrawISR);
+		}
+	}
+
+	// We always need the basic CPU-driven rendering vertex factory because the far mesh uses it
 	WaterVertexFactories.Reserve(WaterQuadTree.GetTreeDepth());
 	for (uint8 i = 0; i < WaterQuadTree.GetTreeDepth(); i++)
 	{
@@ -250,7 +266,7 @@ FWaterMeshSceneProxy::FWaterMeshSceneProxy(UWaterMeshComponent* Component)
 	OcclusionCullingBounds = WaterQuadTree.ComputeNodeBounds(MaxQueries, CVarWaterMeshOcclusionCullExpandBoundsAmountXY.GetValueOnGameThread(), bIncludeFarMeshOcclusionQueries, &OcclusionResultsFarMeshOffset);
 	EmptyOcclusionCullingBounds.Add(WaterQuadTree.GetBoundsIncludingFarMesh());
 	// If this is a GPU quadtree, the CPU root node will have an invalid bounding box, so derive conservative bounds now
-	if (WaterQuadTree.IsGPUQuadTree() && !OcclusionCullingBounds.IsEmpty())
+	if (bIsGPUQuadTree && !OcclusionCullingBounds.IsEmpty())
 	{
 		OcclusionCullingBounds[0] = FBox(FVector(WaterQuadTree.GetTileRegion().Min, WaterQuadTreeMinHeight), FVector(WaterQuadTree.GetTileRegion().Max, WaterQuadTreeMaxHeight));
 	}
@@ -258,7 +274,7 @@ FWaterMeshSceneProxy::FWaterMeshSceneProxy(UWaterMeshComponent* Component)
 	// When using a GPU quadtree, this scene proxy allocates pooled buffers in GetDynamicMeshElements. AllocatePooledBuffer must be called on the renderthread.
 	// The callback on GWaterMeshGPUWork seems to be invoked after all GetDynamicMeshElements tasks finish, so there should be no race condition there.
 	// There is usually only a single instance of this proxy (in certain cases there might be two), so opting out of parallel GDME shouldn't have much of an impact (?).
-	bSupportsParallelGDME = !WaterQuadTree.IsGPUQuadTree();
+	bSupportsParallelGDME = !bIsGPUQuadTree;
 }
 
 FWaterMeshSceneProxy::~FWaterMeshSceneProxy()
@@ -268,8 +284,16 @@ FWaterMeshSceneProxy::~FWaterMeshSceneProxy()
 		WaterFactory->ReleaseResource();
 		delete WaterFactory;
 	}
-	WaterVertexFactoryIndirectDraw->ReleaseResource();
-	delete WaterVertexFactoryIndirectDraw;
+	if (WaterVertexFactoryIndirectDraw)
+	{
+		WaterVertexFactoryIndirectDraw->ReleaseResource();
+		delete WaterVertexFactoryIndirectDraw;
+	}
+	if (WaterVertexFactoryIndirectDrawISR)
+	{
+		WaterVertexFactoryIndirectDrawISR->ReleaseResource();
+		delete WaterVertexFactoryIndirectDrawISR;
+	}
 
 	delete WaterInstanceDataBuffers;
 
@@ -386,10 +410,11 @@ void FWaterMeshSceneProxy::GetDynamicMeshElements(const TArray<const FSceneView*
 
 	const int32 NumWaterMaterials = WaterQuadTree.GetWaterMaterials().Num();
 	const int32 NumBuckets = NumWaterMaterials * DensityCount;
+	const int32 NumBucketsIndirect = NumWaterMaterials; // Indirect draws use a single density mesh tile
 
 	// If this is a GPU quadtree, we are using GPU driven rendering to first create a quadtree on the GPU (once) and then traverse it every frame to build a list of indirect draws
 	const bool bIsGPUQuadTree = WaterQuadTree.IsGPUQuadTree();
-	if (bIsGPUQuadTree && NumBuckets > 0)
+	if (bIsGPUQuadTree && NumBucketsIndirect > 0)
 	{
 		bool bEncounteredISRView = false;
 		int32 InstanceFactor = 1;
@@ -413,7 +438,7 @@ void FWaterMeshSceneProxy::GetDynamicMeshElements(const TArray<const FSceneView*
 		}
 
 		const int32 NumVisibleViews = VisibleViews.Num();
-		const int32 NumIndirectDrawCalls = NumBuckets * NumVisibleViews;
+		const int32 NumIndirectDrawCalls = NumBucketsIndirect * NumVisibleViews;
 		const int32 LeafCountUpperBound = WaterQuadTree.GetMaxLeafCount() * NumVisibleViews;
 		const int32 NumInstanceDataSlots = FMath::Max(1, static_cast<int32>(LeafCountUpperBound * FMath::Clamp(CVarWaterMeshGPUQuadInstanceDataAllocMult.GetValueOnRenderThread(), 0.0f, 8.0f)));
 
@@ -421,6 +446,7 @@ void FWaterMeshSceneProxy::GetDynamicMeshElements(const TArray<const FSceneView*
 		struct FIndirectDrawResources
 		{
 			TRefCountPtr<FRDGPooledBuffer> IndirectArgs;
+			TRefCountPtr<FRDGPooledBuffer> InstanceDataOffsets;
 			TRefCountPtr<FRDGPooledBuffer> InstanceData0;
 			TRefCountPtr<FRDGPooledBuffer> InstanceData1;
 			TRefCountPtr<FRDGPooledBuffer> InstanceData2;
@@ -428,6 +454,7 @@ void FWaterMeshSceneProxy::GetDynamicMeshElements(const TArray<const FSceneView*
 		};
 		FIndirectDrawResources& IndirectDrawResources = Collector.AllocateOneFrameResource<FIndirectDrawResources>();
 		IndirectDrawResources.IndirectArgs = AllocatePooledBuffer(FRDGBufferDesc::CreateIndirectDesc<FRHIDrawIndexedIndirectParameters>(FMath::Max(1, NumIndirectDrawCalls)), TEXT("WaterQuadTree.IndirectArgsBuffer"));
+		IndirectDrawResources.InstanceDataOffsets = AllocatePooledBuffer(FRDGBufferDesc::CreateBufferDesc(sizeof(uint32), FMath::Max(1, NumIndirectDrawCalls)), TEXT("WaterQuadTree.InstanceDataOffsetsBuffer"));
 		IndirectDrawResources.InstanceData0 = AllocatePooledBuffer(FRDGBufferDesc::CreateBufferDesc(sizeof(uint32), FMath::Max(1, NumInstanceDataSlots)), TEXT("WaterQuadTree.InstanceDataBuffer0"));
 		IndirectDrawResources.InstanceData1 = AllocatePooledBuffer(FRDGBufferDesc::CreateBufferDesc(sizeof(uint32), FMath::Max(1, NumInstanceDataSlots)), TEXT("WaterQuadTree.InstanceDataBuffer1"));
 		IndirectDrawResources.InstanceData2 = AllocatePooledBuffer(FRDGBufferDesc::CreateBufferDesc(sizeof(uint32), FMath::Max(1, NumInstanceDataSlots)), TEXT("WaterQuadTree.InstanceDataBuffer2"));
@@ -442,6 +469,7 @@ void FWaterMeshSceneProxy::GetDynamicMeshElements(const TArray<const FSceneView*
 		{
 			WaterQuadTreeGPUTraverseParams = {};
 			WaterQuadTreeGPUTraverseParams.OutIndirectArgsBuffer = IndirectDrawResources.IndirectArgs;
+			WaterQuadTreeGPUTraverseParams.OutInstanceDataOffsetsBuffer = IndirectDrawResources.InstanceDataOffsets;
 			WaterQuadTreeGPUTraverseParams.OutInstanceData0Buffer = IndirectDrawResources.InstanceData0;
 			WaterQuadTreeGPUTraverseParams.OutInstanceData1Buffer = IndirectDrawResources.InstanceData1;
 			WaterQuadTreeGPUTraverseParams.OutInstanceData2Buffer = IndirectDrawResources.InstanceData2;
@@ -451,7 +479,6 @@ void FWaterMeshSceneProxy::GetDynamicMeshElements(const TArray<const FSceneView*
 			WaterQuadTreeGPUTraverseParams.CullingBounds = WaterInfoBounds.ShiftBy(-WaterQuadTree.GetTileRegion().Min);
 			WaterQuadTreeGPUTraverseParams.NumDensities = DensityCount;
 			WaterQuadTreeGPUTraverseParams.NumMaterials = NumWaterMaterials;
-			WaterQuadTreeGPUTraverseParams.NumViews = NumVisibleViews;
 			WaterQuadTreeGPUTraverseParams.NumQuadsLOD0 = NumQuadsLOD0;
 			WaterQuadTreeGPUTraverseParams.NumQuadsPerTileSide = NumQuadsPerIndirectDrawTile;
 			WaterQuadTreeGPUTraverseParams.ForceCollapseDensityLevel = ForceCollapseDensityLevel;
@@ -475,6 +502,11 @@ void FWaterMeshSceneProxy::GetDynamicMeshElements(const TArray<const FSceneView*
 			UserDataWrapper.UserData.IndirectInstanceData1 = IndirectDrawResources.InstanceData1->GetRHI();
 			UserDataWrapper.UserData.IndirectInstanceData2 = IndirectDrawResources.InstanceData2->GetRHI();
 			UserDataWrapper.UserData.IndirectInstanceData3 = (WITH_WATER_SELECTION_SUPPORT != 0) ? IndirectDrawResources.InstanceData3->GetRHI() : nullptr;
+			UserDataWrapper.UserData.IndirectInstanceDataOffsetsSRV = IndirectDrawResources.InstanceDataOffsets->GetSRV(RHICmdList, FRHIBufferSRVCreateInfo(PF_R32_UINT));
+			UserDataWrapper.UserData.IndirectInstanceData0SRV = IndirectDrawResources.InstanceData0->GetSRV(RHICmdList, FRHIBufferSRVCreateInfo(PF_R32_UINT));
+			UserDataWrapper.UserData.IndirectInstanceData1SRV = IndirectDrawResources.InstanceData1->GetSRV(RHICmdList, FRHIBufferSRVCreateInfo(PF_R32_UINT));
+			UserDataWrapper.UserData.IndirectInstanceData2SRV = IndirectDrawResources.InstanceData2->GetSRV(RHICmdList, FRHIBufferSRVCreateInfo(PF_R32_UINT));
+			UserDataWrapper.UserData.IndirectInstanceData3SRV = (WITH_WATER_SELECTION_SUPPORT != 0) ? IndirectDrawResources.InstanceData3->GetSRV(RHICmdList, FRHIBufferSRVCreateInfo(PF_R32_UINT)) : nullptr;
 			
 			UserDataWrappers[(int32)RenderGroup] = &UserDataWrapper;
 		}
@@ -488,6 +520,8 @@ void FWaterMeshSceneProxy::GetDynamicMeshElements(const TArray<const FSceneView*
 			{
 				TRACE_CPUPROFILER_EVENT_SCOPE(BucketsPerView);
 
+				const bool bInstancedStereo = Views[ViewIndex]->bIsInstancedStereoEnabled;
+
 				for (int32 MaterialIndex = 0; MaterialIndex < NumWaterMaterials; ++MaterialIndex)
 				{
 					TRACE_CPUPROFILER_EVENT_SCOPE(MaterialBucket);
@@ -517,7 +551,7 @@ void FWaterMeshSceneProxy::GetDynamicMeshElements(const TArray<const FSceneView*
 							// Set up mesh batch
 							FMeshBatch& Mesh = Collector.AllocateMesh();
 							Mesh.bWireframe = bWireframe;
-							Mesh.VertexFactory = WaterVertexFactoryIndirectDraw;
+							Mesh.VertexFactory = bInstancedStereo ? (FVertexFactory*)WaterVertexFactoryIndirectDrawISR : (FVertexFactory*)WaterVertexFactoryIndirectDraw;
 							Mesh.MaterialRenderProxy = MaterialRenderProxy;
 							Mesh.ReverseCulling = IsLocalToWorldDeterminantNegative();
 							Mesh.Type = PT_TriangleList;
@@ -545,11 +579,11 @@ void FWaterMeshSceneProxy::GetDynamicMeshElements(const TArray<const FSceneView*
 								BatchElement.FirstIndex = 0;
 								BatchElement.NumPrimitives = 0; // Must be 0 to enable usage of IndirectArgsBuffer
 								BatchElement.IndirectArgsBuffer = IndirectDrawResources.IndirectArgs->GetRHI();
-								BatchElement.IndirectArgsOffset = (CompactViewIndex * NumBuckets + BucketIndex) * sizeof(FRHIDrawIndexedIndirectParameters);
+								BatchElement.IndirectArgsOffset = (CompactViewIndex * NumBucketsIndirect + BucketIndex) * sizeof(FRHIDrawIndexedIndirectParameters);
 								BatchElement.UserData = (void*)&UserDataWrappers[(int32)RenderGroup]->UserData;
-								BatchElement.UserIndex = CompactViewIndex * NumBuckets + BucketIndex;
+								BatchElement.UserIndex = CompactViewIndex * NumBucketsIndirect + BucketIndex;
 
-								BatchElement.IndexBuffer = WaterVertexFactoryIndirectDraw->IndexBuffer;
+								BatchElement.IndexBuffer = bInstancedStereo ? WaterVertexFactoryIndirectDrawISR->IndexBuffer : WaterVertexFactoryIndirectDraw->IndexBuffer;
 								BatchElement.PrimitiveIdMode = PrimID_ForceZero;
 
 								// We need the uniform buffer of this primitive because it stores the proper value for the bOutputVelocity flag.
diff --git a/Engine/Plugins/Experimental/Water/Source/Runtime/Private/WaterQuadTreeGPU.cpp b/Engine/Plugins/Experimental/Water/Source/Runtime/Private/WaterQuadTreeGPU.cpp
--- a/Engine/Plugins/Experimental/Water/Source/Runtime/Private/WaterQuadTreeGPU.cpp
+++ b/Engine/Plugins/Experimental/Water/Source/Runtime/Private/WaterQuadTreeGPU.cpp
@@ -321,6 +321,7 @@ public:
 		SHADER_PARAMETER_RDG_BUFFER_SRV(ByteAddressBuffer, BucketCounts)
 		SHADER_PARAMETER(uint32, NumBuckets)
 		SHADER_PARAMETER(uint32, OutputOffset)
+		SHADER_PARAMETER(uint32, bWriteTotalSumAtBufferEnd)
 	END_SHADER_PARAMETER_STRUCT()
 
 	static void ModifyCompilationEnvironment(const FGlobalShaderPermutationParameters& Parameters, FShaderCompilerEnvironment& OutEnvironment)
@@ -366,8 +367,10 @@ public:
 		SHADER_PARAMETER(float, LeafSize)
 		SHADER_PARAMETER(float, LODScale)
 		SHADER_PARAMETER(float, CaptureDepthRange)
+		SHADER_PARAMETER(uint32, StereoPassInstanceFactor)
 		SHADER_PARAMETER(uint32, bWithWaterSelectionSupport)
 		SHADER_PARAMETER(uint32, bLODMorphingEnabled)
+		SHADER_PARAMETER(uint32, bInstancedStereoRendering)
 	END_SHADER_PARAMETER_STRUCT()
 
 	static void ModifyCompilationEnvironment(const FGlobalShaderPermutationParameters& Parameters, FShaderCompilerEnvironment& OutEnvironment)
@@ -751,12 +754,15 @@ void FWaterQuadTreeGPU::Traverse(FRDGBuilder& GraphBuilder, const FTraverseParam
 	RDG_EVENT_SCOPE(GraphBuilder, "FWaterQuadTreeGPU::Traverse");
 	RDG_GPU_STAT_SCOPE(GraphBuilder, FWaterQuadTreeGPU_Traverse);
 
+	const uint32 NumViews = Params.Views.Num();
+
 	FGlobalShaderMap* ShaderMap = GetGlobalShaderMap(GMaxRHIFeatureLevel);
 
 	FRDGTexture* QuadTreeTextureRDG = GraphBuilder.RegisterExternalTexture(CreateRenderTarget(QuadTreeTexture, TEXT("WaterQuadTree.QuadTree")));
 	FRDGTexture* WaterZBoundsTextureRDG = GraphBuilder.RegisterExternalTexture(CreateRenderTarget(WaterZBoundsTexture, TEXT("WaterQuadTree.WaterSurfaceHeight")));
 	FRDGBuffer* WaterBodyRenderDataBufferRDG = GraphBuilder.RegisterExternalBuffer(WaterBodyRenderDataBuffer);
 	FRDGBuffer* IndirectArgsBuffer = GraphBuilder.RegisterExternalBuffer(Params.OutIndirectArgsBuffer);
+	FRDGBuffer* InstanceDataOffsetsBuffer = GraphBuilder.RegisterExternalBuffer(Params.OutInstanceDataOffsetsBuffer);
 	FRDGBuffer* InstanceData0Buffer = GraphBuilder.RegisterExternalBuffer(Params.OutInstanceData0Buffer);
 	FRDGBuffer* InstanceData1Buffer = GraphBuilder.RegisterExternalBuffer(Params.OutInstanceData1Buffer);
 	FRDGBuffer* InstanceData2Buffer = GraphBuilder.RegisterExternalBuffer(Params.OutInstanceData2Buffer);
@@ -775,16 +781,14 @@ void FWaterQuadTreeGPU::Traverse(FRDGBuilder& GraphBuilder, const FTraverseParam
 	FRDGBufferUAV* OcclusionQueryIndirectArgsBufferUAV = nullptr;
 	if (bPixelPreciseOcclusionQueries)
 	{
-		OcclusionQueryIndirectArgsBuffer = GraphBuilder.CreateBuffer(FRDGBufferDesc::CreateIndirectDesc<FRHIDrawIndexedIndirectParameters>(Params.NumViews), TEXT("WaterQuadTree.OcclusionQueryIndirectArgs"));
+		OcclusionQueryIndirectArgsBuffer = GraphBuilder.CreateBuffer(FRDGBufferDesc::CreateIndirectDesc<FRHIDrawIndexedIndirectParameters>(NumViews), TEXT("WaterQuadTree.OcclusionQueryIndirectArgs"));
 		OcclusionQueryIndirectArgsBufferUAV = GraphBuilder.CreateUAV(FRDGBufferUAVDesc(OcclusionQueryIndirectArgsBuffer, PF_R32_UINT));
 	}
 
-	const uint32 NumBucketsPerView = Params.NumDensities * Params.NumMaterials;
-	const uint32 NumBucketsTotal = NumBucketsPerView * Params.NumViews;
+	const uint32 NumBucketsPerView = Params.NumMaterials; // The GPU-driven water rendering path only uses a single density mesh tile to draw everything, so NumBuckets is equal to NumMaterials.
+	const uint32 NumBucketsTotal = NumBucketsPerView * NumViews;
 	const FIntPoint QuadTreeResolution = QuadTreeTextureRDG->Desc.Extent;
 
-	FRDGBuffer* PrefixSumBuffer = GraphBuilder.CreateBuffer(FRDGBufferDesc::CreateBufferDesc(sizeof(uint32), NumBucketsTotal), TEXT("WaterQuadTree.PrefixSum"));
-
 	// Initialize indirect args
 	{
 		FWaterQuadTreeInitializeIndirectArgsCS::FPermutationDomain PermutationDomain;
@@ -795,14 +799,14 @@ void FWaterQuadTreeGPU::Traverse(FRDGBuilder& GraphBuilder, const FTraverseParam
 		PassParameters->IndirectArgs = IndirectArgsBufferUAV;
 		PassParameters->OcclusionQueryArgs = OcclusionQueryIndirectArgsBufferUAV;
 		PassParameters->NumDrawBuckets = Params.NumMaterials;
-		PassParameters->NumViews = Params.NumViews;
+		PassParameters->NumViews = NumViews;
 		PassParameters->NumQuads = Params.NumQuadsPerTileSide;
 		
 		FComputeShaderUtils::AddPass(GraphBuilder, RDG_EVENT_NAME("WaterQuadTreeInitIndirectArgs"), ComputeShader, PassParameters, FComputeShaderUtils::GetGroupCount(NumBucketsTotal, 64));
 	}
 
 	// Iterate over all views for which to create indirect water draws
-	for (uint32 ViewIndex = 0; ViewIndex < Params.NumViews; ++ViewIndex)
+	for (uint32 ViewIndex = 0; ViewIndex < NumViews; ++ViewIndex)
 	{
 		const FSceneView* View = Params.Views[ViewIndex];
 		const FViewInfo* ViewInfo = View->bIsViewInfo ? reinterpret_cast<const FViewInfo*>(View) : nullptr;
@@ -968,10 +972,11 @@ void FWaterQuadTreeGPU::Traverse(FRDGBuilder& GraphBuilder, const FTraverseParam
 			TShaderMapRef<FWaterQuadTreeBucketPrefixSumCS> ComputeShader(ShaderMap, PermutationVector);
 
 			FWaterQuadTreeBucketPrefixSumCS::FParameters* PassParameters = GraphBuilder.AllocParameters<FWaterQuadTreeBucketPrefixSumCS::FParameters>();
-			PassParameters->BucketPrefixSums = GraphBuilder.CreateUAV(FRDGBufferUAVDesc(PrefixSumBuffer, PF_R32_UINT));
+			PassParameters->BucketPrefixSums = GraphBuilder.CreateUAV(FRDGBufferUAVDesc(InstanceDataOffsetsBuffer, PF_R32_UINT));
 			PassParameters->BucketCounts = GraphBuilder.CreateSRV(FRDGBufferSRVDesc(BucketCounts));
 			PassParameters->NumBuckets = NumBucketsPerView;
 			PassParameters->OutputOffset = ViewIndex * NumBucketsPerView;
+			PassParameters->bWriteTotalSumAtBufferEnd =  Params.Views.IsValidIndex(ViewIndex + 1); // Propagate total sum so the next iteration/view can compute correct global prefix sums/offsets.
 
 			FComputeShaderUtils::AddPass(GraphBuilder, RDG_EVENT_NAME("WaterQuadTreeDrawBucketPrefixSums(View: %i)", ViewIndex), ComputeShader, PassParameters, FIntVector3(1, 1, 1));
 		}
@@ -992,7 +997,7 @@ void FWaterQuadTreeGPU::Traverse(FRDGBuilder& GraphBuilder, const FTraverseParam
 			PassParameters->WaterZBoundsTexture = WaterZBoundsTextureRDG;
 			PassParameters->WaterBodyRenderData = WaterBodyRenderDataBufferSRV;
 			PassParameters->PackedNodes = PackedNodesSRV;
-			PassParameters->InstanceDataOffsets = GraphBuilder.CreateSRV(FRDGBufferSRVDesc(PrefixSumBuffer, PF_R32_UINT));
+			PassParameters->InstanceDataOffsets = GraphBuilder.CreateSRV(FRDGBufferSRVDesc(InstanceDataOffsetsBuffer, PF_R32_UINT));
 			PassParameters->OcclusionResults = OcclusionQueryResultsSRV;
 			PassParameters->QuadTreePosition = QuadTreePositionTranslatedWorldSpace;
 			PassParameters->ObserverPosition = ObserverPositionTranslatedWorldSpace;
@@ -1008,8 +1013,10 @@ void FWaterQuadTreeGPU::Traverse(FRDGBuilder& GraphBuilder, const FTraverseParam
 			PassParameters->LeafSize = Params.LeafSize;
 			PassParameters->LODScale = Params.LODScale;
 			PassParameters->CaptureDepthRange = CaptureDepthRange;
+			PassParameters->StereoPassInstanceFactor = View->GetStereoPassInstanceFactor();
 			PassParameters->bWithWaterSelectionSupport = Params.bWithWaterSelectionSupport;
 			PassParameters->bLODMorphingEnabled = Params.bLODMorphingEnabled;
+			PassParameters->bInstancedStereoRendering = View->bIsInstancedStereoEnabled;
 
 			FComputeShaderUtils::AddPass(GraphBuilder, RDG_EVENT_NAME("WaterQuadTreeGenerateDraws(View: %i)", ViewIndex), ComputeShader, PassParameters, FComputeShaderUtils::GetGroupCount(PassParameters->NumDispatchedThreads, 64));
 		}
@@ -1052,11 +1059,12 @@ void FWaterQuadTreeGPU::Traverse(FRDGBuilder& GraphBuilder, const FTraverseParam
 	}
 
 	GraphBuilder.UseExternalAccessMode(IndirectArgsBuffer, ERHIAccess::IndirectArgs);
-	GraphBuilder.UseExternalAccessMode(InstanceData0Buffer, ERHIAccess::VertexOrIndexBuffer);
-	GraphBuilder.UseExternalAccessMode(InstanceData1Buffer, ERHIAccess::VertexOrIndexBuffer);
-	GraphBuilder.UseExternalAccessMode(InstanceData2Buffer, ERHIAccess::VertexOrIndexBuffer);
+	GraphBuilder.UseExternalAccessMode(InstanceDataOffsetsBuffer, ERHIAccess::SRVMask);
+	GraphBuilder.UseExternalAccessMode(InstanceData0Buffer, ERHIAccess::SRVMask | ERHIAccess::VertexOrIndexBuffer);
+	GraphBuilder.UseExternalAccessMode(InstanceData1Buffer, ERHIAccess::SRVMask | ERHIAccess::VertexOrIndexBuffer);
+	GraphBuilder.UseExternalAccessMode(InstanceData2Buffer, ERHIAccess::SRVMask | ERHIAccess::VertexOrIndexBuffer);
 	if (Params.bWithWaterSelectionSupport)
 	{
-		GraphBuilder.UseExternalAccessMode(InstanceData3Buffer, ERHIAccess::VertexOrIndexBuffer);
+		GraphBuilder.UseExternalAccessMode(InstanceData3Buffer, ERHIAccess::SRVMask | ERHIAccess::VertexOrIndexBuffer);
 	}
 }
diff --git a/Engine/Plugins/Experimental/Water/Source/Runtime/Private/WaterQuadTreeGPU.h b/Engine/Plugins/Experimental/Water/Source/Runtime/Private/WaterQuadTreeGPU.h
--- a/Engine/Plugins/Experimental/Water/Source/Runtime/Private/WaterQuadTreeGPU.h
+++ b/Engine/Plugins/Experimental/Water/Source/Runtime/Private/WaterQuadTreeGPU.h
@@ -61,6 +61,7 @@ public:
 	struct FTraverseParams
 	{
 		FRDGPooledBuffer* OutIndirectArgsBuffer = nullptr;
+		FRDGPooledBuffer* OutInstanceDataOffsetsBuffer = nullptr;
 		FRDGPooledBuffer* OutInstanceData0Buffer = nullptr;
 		FRDGPooledBuffer* OutInstanceData1Buffer = nullptr;
 		FRDGPooledBuffer* OutInstanceData2Buffer = nullptr;
@@ -70,7 +71,6 @@ public:
 		FBox2D CullingBounds = FBox2D(ForceInit);
 		uint32 NumDensities = 0;
 		uint32 NumMaterials = 0;
-		uint32 NumViews = 0;
 		uint32 NumQuadsLOD0 = 0;
 		uint32 NumQuadsPerTileSide = 0;
 		int32 ForceCollapseDensityLevel = -1;
diff --git a/Engine/Plugins/Experimental/Water/Source/Runtime/Private/WaterVertexFactory.cpp b/Engine/Plugins/Experimental/Water/Source/Runtime/Private/WaterVertexFactory.cpp
--- a/Engine/Plugins/Experimental/Water/Source/Runtime/Private/WaterVertexFactory.cpp
+++ b/Engine/Plugins/Experimental/Water/Source/Runtime/Private/WaterVertexFactory.cpp
@@ -13,23 +13,38 @@ IMPLEMENT_GLOBAL_SHADER_PARAMETER_STRUCT(FWaterVertexFactoryRaytracingParameters
 /**
  * Shader parameters for water vertex factory.
  */
-template <bool bWithWaterSelectionSupport, bool bIndirectDraws>
+template <bool bWithWaterSelectionSupport, EWaterVertexFactoryDrawMode DrawMode>
 class TWaterVertexFactoryShaderParameters : public FVertexFactoryShaderParameters
 {
-	using WaterVertexFactoryShaderParametersType = TWaterVertexFactoryShaderParameters<bWithWaterSelectionSupport, bIndirectDraws>;
+	using WaterVertexFactoryShaderParametersType = TWaterVertexFactoryShaderParameters<bWithWaterSelectionSupport, DrawMode>;
 	DECLARE_TYPE_LAYOUT(WaterVertexFactoryShaderParametersType, NonVirtual);
 
 public:
-	using WaterVertexFactoryType = TWaterVertexFactory<bWithWaterSelectionSupport, bIndirectDraws>;
+	using WaterVertexFactoryType = TWaterVertexFactory<bWithWaterSelectionSupport, DrawMode>;
 	using WaterMeshUserDataType = TWaterMeshUserData<bWithWaterSelectionSupport>;
 	using WaterInstanceDataBuffersType = TWaterInstanceDataBuffers<bWithWaterSelectionSupport>;
 
 	void Bind(const FShaderParameterMap& ParameterMap)
 	{
-		if (bIndirectDraws)
+		if (DrawMode == EWaterVertexFactoryDrawMode::Indirect || DrawMode == EWaterVertexFactoryDrawMode::IndirectInstancedStereo)
 		{
 			QuadTreePositionParameter.Bind(ParameterMap, TEXT("QuadTreePosition"));
 			LODMorphingEnabledParameter.Bind(ParameterMap, TEXT("bLODMorphingEnabled"));
+
+			if (DrawMode == EWaterVertexFactoryDrawMode::IndirectInstancedStereo)
+			{
+				InstanceDataOffsetsBufferParameter.Bind(ParameterMap, TEXT("InstanceDataOffsetsBuffer"));
+				InstanceData0BufferParameter.Bind(ParameterMap, TEXT("InstanceData0Buffer"));
+				InstanceData1BufferParameter.Bind(ParameterMap, TEXT("InstanceData1Buffer"));
+				InstanceData2BufferParameter.Bind(ParameterMap, TEXT("InstanceData2Buffer"));
+				if (bWithWaterSelectionSupport)
+				{
+					InstanceData3BufferParameter.Bind(ParameterMap, TEXT("InstanceData3Buffer"));
+				}
+
+				DrawBucketIndexParameter.Bind(ParameterMap, TEXT("DrawBucketIndex"));
+				StereoPassInstanceFactorParameter.Bind(ParameterMap, TEXT("StereoPassInstanceFactor"));
+			}
 		}
 	}
 
@@ -59,9 +74,26 @@ public:
 		}
 #endif
 
-		if (bIndirectDraws)
+		if (DrawMode == EWaterVertexFactoryDrawMode::Indirect || DrawMode == EWaterVertexFactoryDrawMode::IndirectInstancedStereo)
 		{
-			if (VertexStreams.Num() > 0)
+			check((DrawMode == EWaterVertexFactoryDrawMode::IndirectInstancedStereo) == View->bIsInstancedStereoEnabled);
+
+			if (View->bIsInstancedStereoEnabled)
+			{
+				const uint32 DrawBucketIndex = BatchElement.UserIndex;
+				const uint32 StereoPassInstanceFactor = View->GetStereoPassInstanceFactor();
+				ShaderBindings.Add(InstanceDataOffsetsBufferParameter, WaterMeshUserData->IndirectInstanceDataOffsetsSRV);
+				ShaderBindings.Add(InstanceData0BufferParameter, WaterMeshUserData->IndirectInstanceData0SRV);
+				ShaderBindings.Add(InstanceData1BufferParameter, WaterMeshUserData->IndirectInstanceData1SRV);
+				ShaderBindings.Add(InstanceData2BufferParameter, WaterMeshUserData->IndirectInstanceData2SRV);
+				if (bWithWaterSelectionSupport)
+				{
+					ShaderBindings.Add(InstanceData3BufferParameter, WaterMeshUserData->IndirectInstanceData3SRV);
+				}
+				ShaderBindings.Add(DrawBucketIndexParameter, DrawBucketIndex);
+				ShaderBindings.Add(StereoPassInstanceFactorParameter, StereoPassInstanceFactor);
+			}
+			else if (VertexStreams.Num() > 0)
 			{
 				const int32 NumBuffers = bWithWaterSelectionSupport ? 4 : 3;
 				FRHIBuffer* InstanceVertexBuffers[] = 
@@ -112,106 +144,83 @@ public:
 private:
 	LAYOUT_FIELD(FShaderParameter, QuadTreePositionParameter);
 	LAYOUT_FIELD(FShaderParameter, LODMorphingEnabledParameter);
+	LAYOUT_FIELD(FShaderResourceParameter, InstanceDataOffsetsBufferParameter);
+	LAYOUT_FIELD(FShaderResourceParameter, InstanceData0BufferParameter);
+	LAYOUT_FIELD(FShaderResourceParameter, InstanceData1BufferParameter);
+	LAYOUT_FIELD(FShaderResourceParameter, InstanceData2BufferParameter);
+	LAYOUT_FIELD(FShaderResourceParameter, InstanceData3BufferParameter);
+	LAYOUT_FIELD(FShaderParameter, DrawBucketIndexParameter);
+	LAYOUT_FIELD(FShaderParameter, StereoPassInstanceFactorParameter);
 };
 
 // ----------------------------------------------------------------------------------
 
-using FWaterVertexFactoryParametersNoSelectionNoIndirect = TWaterVertexFactoryShaderParameters</*bWithWaterSelectionSupport = */ false, /*bIndirectDraws = */ false>;
-using FWaterVertexFactoryParametersNoSelectionWithIndirect = TWaterVertexFactoryShaderParameters</*bWithWaterSelectionSupport = */ false, /*bIndirectDraws = */ true>;
-using FWaterVertexFactoryNoSelectionNoIndirect = TWaterVertexFactory</*bWithWaterSelectionSupport = */ false, /*bIndirectDraws = */ false>;
-using FWaterVertexFactoryNoSelectionWithIndirect = TWaterVertexFactory</*bWithWaterSelectionSupport = */ false, /*bIndirectDraws = */ true>;
+#define IMPLEMENT_WATER_VERTEX_FACTORY_RAYTRACING(VertexFactoryType, ParametersType, bImplementRaytracing) IMPLEMENT_WATER_VERTEX_FACTORY_RAYTRACING_##bImplementRaytracing(VertexFactoryType, ParametersType)
+#define IMPLEMENT_WATER_VERTEX_FACTORY_RAYTRACING_false(VertexFactoryType, ParametersType)
+#define IMPLEMENT_WATER_VERTEX_FACTORY_RAYTRACING_true(VertexFactoryType, ParametersType)														\
+	IMPLEMENT_VERTEX_FACTORY_PARAMETER_TYPE(VertexFactoryType, SF_Compute, ParametersType);														\
+	IMPLEMENT_VERTEX_FACTORY_PARAMETER_TYPE(VertexFactoryType, SF_RayHitGroup, ParametersType);
+
+#define IMPLEMENT_WATER_VERTEX_FACTORY(bWidthWaterSelectionSupport, DrawMode, bImplementRaytracing, Suffix)										\
+	using FWaterVertexFactoryParameters##Suffix = TWaterVertexFactoryShaderParameters<bWidthWaterSelectionSupport, DrawMode>;					\
+	using FWaterVertexFactory##Suffix = TWaterVertexFactory<bWidthWaterSelectionSupport, DrawMode>;												\
+	IMPLEMENT_TEMPLATE_TYPE_LAYOUT(template<>, FWaterVertexFactoryParameters##Suffix);															\
+	IMPLEMENT_VERTEX_FACTORY_PARAMETER_TYPE(FWaterVertexFactory##Suffix, SF_Vertex, FWaterVertexFactoryParameters##Suffix);						\
+	IMPLEMENT_WATER_VERTEX_FACTORY_RAYTRACING(FWaterVertexFactory##Suffix, FWaterVertexFactoryParameters##Suffix, bImplementRaytracing)			\
+	IMPLEMENT_TEMPLATE_VERTEX_FACTORY_TYPE(template<>, FWaterVertexFactory##Suffix, "/Plugin/Water/Private/WaterMeshVertexFactory.ush",			\
+		EVertexFactoryFlags::UsedWithMaterials																									\
+		| EVertexFactoryFlags::SupportsDynamicLighting																							\
+		| EVertexFactoryFlags::SupportsPrecisePrevWorldPos																						\
+		| EVertexFactoryFlags::SupportsPrimitiveIdStream																						\
+		| EVertexFactoryFlags::SupportsRayTracing																								\
+		| EVertexFactoryFlags::SupportsRayTracingDynamicGeometry																				\
+		| EVertexFactoryFlags::SupportsPSOPrecaching																							\
+	);
 
 // Always implement the basic vertex factory so that it's there for both editor and non-editor builds :
-IMPLEMENT_TEMPLATE_TYPE_LAYOUT(template<>, FWaterVertexFactoryParametersNoSelectionNoIndirect);
-IMPLEMENT_TEMPLATE_TYPE_LAYOUT(template<>, FWaterVertexFactoryParametersNoSelectionWithIndirect);
-IMPLEMENT_VERTEX_FACTORY_PARAMETER_TYPE(FWaterVertexFactoryNoSelectionNoIndirect, SF_Vertex, FWaterVertexFactoryParametersNoSelectionNoIndirect);
-IMPLEMENT_VERTEX_FACTORY_PARAMETER_TYPE(FWaterVertexFactoryNoSelectionWithIndirect, SF_Vertex, FWaterVertexFactoryParametersNoSelectionWithIndirect);
 #if RHI_RAYTRACING
-IMPLEMENT_VERTEX_FACTORY_PARAMETER_TYPE(FWaterVertexFactoryNoSelectionNoIndirect, SF_Compute, FWaterVertexFactoryParametersNoSelectionNoIndirect);
-IMPLEMENT_VERTEX_FACTORY_PARAMETER_TYPE(FWaterVertexFactoryNoSelectionNoIndirect, SF_RayHitGroup, FWaterVertexFactoryParametersNoSelectionNoIndirect);
-#endif // RHI_RAYTRACING
-IMPLEMENT_TEMPLATE_VERTEX_FACTORY_TYPE(template<>, FWaterVertexFactoryNoSelectionNoIndirect, "/Plugin/Water/Private/WaterMeshVertexFactory.ush",
-	  EVertexFactoryFlags::UsedWithMaterials
-	| EVertexFactoryFlags::SupportsDynamicLighting
-	| EVertexFactoryFlags::SupportsPrecisePrevWorldPos
-	| EVertexFactoryFlags::SupportsPrimitiveIdStream
-	| EVertexFactoryFlags::SupportsRayTracing
-	| EVertexFactoryFlags::SupportsRayTracingDynamicGeometry
-	| EVertexFactoryFlags::SupportsPSOPrecaching
-);
-IMPLEMENT_TEMPLATE_VERTEX_FACTORY_TYPE(template<>, FWaterVertexFactoryNoSelectionWithIndirect, "/Plugin/Water/Private/WaterMeshVertexFactory.ush",
-	EVertexFactoryFlags::UsedWithMaterials
-	| EVertexFactoryFlags::SupportsDynamicLighting
-	| EVertexFactoryFlags::SupportsPrecisePrevWorldPos
-	| EVertexFactoryFlags::SupportsPrimitiveIdStream
-	| EVertexFactoryFlags::SupportsRayTracing
-	| EVertexFactoryFlags::SupportsRayTracingDynamicGeometry
-	| EVertexFactoryFlags::SupportsPSOPrecaching
-);
+IMPLEMENT_WATER_VERTEX_FACTORY(false /*bWithWaterSelectionSupport*/, EWaterVertexFactoryDrawMode::NonIndirect, true /*bImplementRaytracing*/, NoSelectionNonIndirect);
+#else
+IMPLEMENT_WATER_VERTEX_FACTORY(false /*bWithWaterSelectionSupport*/, EWaterVertexFactoryDrawMode::NonIndirect, false /*bImplementRaytracing*/, NoSelectionNonIndirect);
+#endif
+IMPLEMENT_WATER_VERTEX_FACTORY(false /*bWithWaterSelectionSupport*/, EWaterVertexFactoryDrawMode::Indirect, false /*bImplementRaytracing*/, NoSelectionWithIndirect);
+IMPLEMENT_WATER_VERTEX_FACTORY(false /*bWithWaterSelectionSupport*/, EWaterVertexFactoryDrawMode::IndirectInstancedStereo, false /*bImplementRaytracing*/, NoSelectionWithIndirectISR);
 
+// In editor builds, also implement the vertex factory that supports water selection:
 #if WITH_WATER_SELECTION_SUPPORT
 
-using FWaterVertexFactoryParametersWithSelectionNoIndirect = TWaterVertexFactoryShaderParameters</*bWithWaterSelectionSupport = */ true, /*bIndirectDraws = */ false>;
-using FWaterVertexFactoryParametersWithSelectionWithIndirect = TWaterVertexFactoryShaderParameters</*bWithWaterSelectionSupport = */ true, /*bIndirectDraws = */ true>;
-using FWaterVertexFactoryWithSelectionNoIndirect = TWaterVertexFactory</*bWithWaterSelectionSupport = */ true, /*bIndirectDraws = */ false>;
-using FWaterVertexFactoryWithSelectionWithIndirect = TWaterVertexFactory</*bWithWaterSelectionSupport = */ true, /*bIndirectDraws = */ true>;
-
-// In editor builds, also implement the vertex factory that supports water selection:
-IMPLEMENT_TEMPLATE_TYPE_LAYOUT(template<>, FWaterVertexFactoryParametersWithSelectionNoIndirect);
-IMPLEMENT_TEMPLATE_TYPE_LAYOUT(template<>, FWaterVertexFactoryParametersWithSelectionWithIndirect);
-IMPLEMENT_VERTEX_FACTORY_PARAMETER_TYPE(FWaterVertexFactoryWithSelectionNoIndirect, SF_Vertex, FWaterVertexFactoryParametersWithSelectionNoIndirect);
-IMPLEMENT_VERTEX_FACTORY_PARAMETER_TYPE(FWaterVertexFactoryWithSelectionWithIndirect, SF_Vertex, FWaterVertexFactoryParametersWithSelectionWithIndirect);
 #if RHI_RAYTRACING
-IMPLEMENT_VERTEX_FACTORY_PARAMETER_TYPE(FWaterVertexFactoryWithSelectionNoIndirect, SF_Compute, FWaterVertexFactoryParametersWithSelectionNoIndirect);
-IMPLEMENT_VERTEX_FACTORY_PARAMETER_TYPE(FWaterVertexFactoryWithSelectionNoIndirect, SF_RayHitGroup, FWaterVertexFactoryParametersWithSelectionNoIndirect);
+IMPLEMENT_WATER_VERTEX_FACTORY(true /*bWithWaterSelectionSupport*/, EWaterVertexFactoryDrawMode::NonIndirect, true /*bImplementRaytracing*/, WithSelectionNonIndirect);
+#else
+IMPLEMENT_WATER_VERTEX_FACTORY(true /*bWithWaterSelectionSupport*/, EWaterVertexFactoryDrawMode::NonIndirect, false /*bImplementRaytracing*/, WithSelectionNonIndirect);
 #endif
-IMPLEMENT_TEMPLATE_VERTEX_FACTORY_TYPE(template<>, FWaterVertexFactoryWithSelectionNoIndirect, "/Plugin/Water/Private/WaterMeshVertexFactory.ush",
-	  EVertexFactoryFlags::UsedWithMaterials
-	| EVertexFactoryFlags::SupportsDynamicLighting
-	| EVertexFactoryFlags::SupportsPrecisePrevWorldPos
-	| EVertexFactoryFlags::SupportsPrimitiveIdStream
-	| EVertexFactoryFlags::SupportsRayTracing
-	| EVertexFactoryFlags::SupportsRayTracingDynamicGeometry
-	| EVertexFactoryFlags::SupportsPSOPrecaching
-);
-IMPLEMENT_TEMPLATE_VERTEX_FACTORY_TYPE(template<>, FWaterVertexFactoryWithSelectionWithIndirect, "/Plugin/Water/Private/WaterMeshVertexFactory.ush",
-	EVertexFactoryFlags::UsedWithMaterials
-	| EVertexFactoryFlags::SupportsDynamicLighting
-	| EVertexFactoryFlags::SupportsPrecisePrevWorldPos
-	| EVertexFactoryFlags::SupportsPrimitiveIdStream
-	| EVertexFactoryFlags::SupportsRayTracing
-	| EVertexFactoryFlags::SupportsRayTracingDynamicGeometry
-	| EVertexFactoryFlags::SupportsPSOPrecaching
-);
 
+IMPLEMENT_WATER_VERTEX_FACTORY(true /*bWithWaterSelectionSupport*/, EWaterVertexFactoryDrawMode::Indirect, false /*bImplementRaytracing*/, WithSelectionWithIndirect);
+IMPLEMENT_WATER_VERTEX_FACTORY(true /*bWithWaterSelectionSupport*/, EWaterVertexFactoryDrawMode::IndirectInstancedStereo, false /*bImplementRaytracing*/, WithSelectionWithIndirectISR);
 #endif // WITH_WATER_SELECTION_SUPPORT
 
-const FVertexFactoryType* GetWaterVertexFactoryType(bool bWithWaterSelectionSupport, bool bIndirectDraws)
+const FVertexFactoryType* GetWaterVertexFactoryType(bool bWithWaterSelectionSupport, EWaterVertexFactoryDrawMode DrawMode)
 {
 #if WITH_WATER_SELECTION_SUPPORT
 	if (bWithWaterSelectionSupport)
 	{
-		if (bIndirectDraws)
+		switch (DrawMode)
 		{
-			return &TWaterVertexFactory</*bWithWaterSelectionSupport = */ true, /*bIndirectDraws = */ true>::StaticType;
-		}
-		else
-		{
-			return &TWaterVertexFactory</*bWithWaterSelectionSupport = */ true, /*bIndirectDraws = */ false>::StaticType;
+		case EWaterVertexFactoryDrawMode::NonIndirect: return &TWaterVertexFactory</*bWithWaterSelectionSupport = */ true, EWaterVertexFactoryDrawMode::NonIndirect>::StaticType;
+		case EWaterVertexFactoryDrawMode::Indirect: return &TWaterVertexFactory</*bWithWaterSelectionSupport = */ true, EWaterVertexFactoryDrawMode::Indirect>::StaticType;
+		case EWaterVertexFactoryDrawMode::IndirectInstancedStereo: return &TWaterVertexFactory</*bWithWaterSelectionSupport = */ true, EWaterVertexFactoryDrawMode::IndirectInstancedStereo>::StaticType;
+		default: checkNoEntry(); return nullptr;
 		}
 	}
 	else
 #endif // WITH_WATER_SELECTION_SUPPORT
 	{
-		check(!bWithWaterSelectionSupport);
-		if (bIndirectDraws)
-		{
-			return &TWaterVertexFactory</*bWithWaterSelectionSupport = */ false, /*bIndirectDraws = */ true>::StaticType;
-		}
-		else
+		switch (DrawMode)
 		{
-			return &TWaterVertexFactory</*bWithWaterSelectionSupport = */ false, /*bIndirectDraws = */ false>::StaticType;
+		case EWaterVertexFactoryDrawMode::NonIndirect: return &TWaterVertexFactory</*bWithWaterSelectionSupport = */ false, EWaterVertexFactoryDrawMode::NonIndirect>::StaticType;
+		case EWaterVertexFactoryDrawMode::Indirect: return &TWaterVertexFactory</*bWithWaterSelectionSupport = */ false, EWaterVertexFactoryDrawMode::Indirect>::StaticType;
+		case EWaterVertexFactoryDrawMode::IndirectInstancedStereo: return &TWaterVertexFactory</*bWithWaterSelectionSupport = */ false, EWaterVertexFactoryDrawMode::IndirectInstancedStereo>::StaticType;
+		default: checkNoEntry(); return nullptr;
 		}
-		
 	}
 }
diff --git a/Engine/Plugins/Experimental/Water/Source/Runtime/Public/WaterMeshSceneProxy.h b/Engine/Plugins/Experimental/Water/Source/Runtime/Public/WaterMeshSceneProxy.h
--- a/Engine/Plugins/Experimental/Water/Source/Runtime/Public/WaterMeshSceneProxy.h
+++ b/Engine/Plugins/Experimental/Water/Source/Runtime/Public/WaterMeshSceneProxy.h
@@ -64,8 +64,9 @@ public:
 #endif // WITH_WATER_SELECTION_SUPPORT
 
 	// At runtime, we only ever need one version of the vertex factory : with selection support (editor) or without : 
-	using FWaterVertexFactoryType = TWaterVertexFactory<WITH_WATER_SELECTION_SUPPORT, /*bIndirectDraws = */ false>;
-	using FWaterVertexFactoryIndirectDrawType = TWaterVertexFactory<WITH_WATER_SELECTION_SUPPORT, /*bIndirectDraws = */ true>;
+	using FWaterVertexFactoryType = TWaterVertexFactory<WITH_WATER_SELECTION_SUPPORT, EWaterVertexFactoryDrawMode::NonIndirect>;
+	using FWaterVertexFactoryIndirectDrawType = TWaterVertexFactory<WITH_WATER_SELECTION_SUPPORT, EWaterVertexFactoryDrawMode::Indirect>;
+	using FWaterVertexFactoryIndirectDrawISRType = TWaterVertexFactory<WITH_WATER_SELECTION_SUPPORT, EWaterVertexFactoryDrawMode::IndirectInstancedStereo>;
 	using FWaterInstanceDataBuffersType = TWaterInstanceDataBuffers<WITH_WATER_SELECTION_SUPPORT>;
 	using FWaterMeshUserDataBuffersType = TWaterMeshUserDataBuffers<WITH_WATER_SELECTION_SUPPORT>;
 
@@ -107,6 +108,7 @@ private:
 	// One vertex factory per LOD
 	TArray<FWaterVertexFactoryType*> WaterVertexFactories;
 	FWaterVertexFactoryIndirectDrawType* WaterVertexFactoryIndirectDraw;
+	FWaterVertexFactoryIndirectDrawISRType* WaterVertexFactoryIndirectDrawISR;
 
 	/** Tiles containing water, stored in a quad tree */
 	FWaterQuadTree WaterQuadTree;
diff --git a/Engine/Plugins/Experimental/Water/Source/Runtime/Public/WaterVertexFactory.h b/Engine/Plugins/Experimental/Water/Source/Runtime/Public/WaterVertexFactory.h
--- a/Engine/Plugins/Experimental/Water/Source/Runtime/Public/WaterVertexFactory.h
+++ b/Engine/Plugins/Experimental/Water/Source/Runtime/Public/WaterVertexFactory.h
@@ -169,17 +169,32 @@ enum class EWaterMeshRenderGroupType : uint8
 #endif // WITH_WATER_SELECTION_SUPPORT
 };
 
-template <bool bWithWaterSelectionSupport, bool bIndirectDraws>
+// Water supports both a CPU-driven (non-indirect draws) and a GPU-driven (indirect draws) path.
+// The regular GPU-driven path does not support ISR out of the box, so there is a special path for indirect draws supporting ISR.
+enum class EWaterVertexFactoryDrawMode
+{
+	// Non-indirect draw calls for CPU-driven water rendering
+	NonIndirect,
+	
+	// Indirect draw calls for GPU-driven water rendering using the GPU water quadtree to generate indirect draw calls.
+	// Uses vertex streams (fixed function) to push instance data into the vertex shader.
+	Indirect,
+
+	// Indirect draw calls with support for ISR. Uses manual fetching of instance data in the vertex shader/factory.
+	IndirectInstancedStereo,
+};
+
+template <bool bWithWaterSelectionSupport, EWaterVertexFactoryDrawMode DrawMode>
 class TWaterVertexFactory : public FVertexFactory
 {
 	DECLARE_VERTEX_FACTORY_TYPE(FWaterVertexFactoryType);
 
 public:
 	using Super = FVertexFactory;
-	using FWaterVertexFactoryType = TWaterVertexFactory<bWithWaterSelectionSupport, bIndirectDraws>;
+	using FWaterVertexFactoryType = TWaterVertexFactory<bWithWaterSelectionSupport, DrawMode>;
 
 	static constexpr int32 NumRenderGroups = bWithWaterSelectionSupport ? 3 : 1; // Must match EWaterMeshRenderGroupType
-	static constexpr int32 NumAdditionalVertexStreams = bIndirectDraws ? 0 : TWaterInstanceDataBuffers<bWithWaterSelectionSupport>::NumBuffers;
+	static constexpr int32 NumAdditionalVertexStreams = TWaterInstanceDataBuffers<bWithWaterSelectionSupport>::NumBuffers;
 
 	TWaterVertexFactory(ERHIFeatureLevel::Type InFeatureLevel, const FVector& InQuadTreePositionWS, int32 InNumQuadsPerSide, int32 InNumQuadsLOD0, int32 InNumDensities, float InLeafSize, float InLODScale, float InCaptureDepthRange);
 	~TWaterVertexFactory();
@@ -209,6 +224,9 @@ public:
 
 	inline FVector GetQuadTreePositionWS() const { return QuadTreePositionWS; }
 
+	static constexpr bool UsesIndirectDraws() { return DrawMode == EWaterVertexFactoryDrawMode::Indirect || DrawMode == EWaterVertexFactoryDrawMode::IndirectInstancedStereo; }
+	static constexpr bool UsesInstancedStereo() { return DrawMode == EWaterVertexFactoryDrawMode::IndirectInstancedStereo; }
+
 private:
 	void SetupUniformDataForGroup(EWaterMeshRenderGroupType InRenderGroupType);
 
@@ -228,7 +246,7 @@ private:
 	const float CaptureDepthRange = 0.0f;
 };
 
-extern const FVertexFactoryType* GetWaterVertexFactoryType(bool bWithWaterSelectionSupport, bool bIndirectDraws);
+extern const FVertexFactoryType* GetWaterVertexFactoryType(bool bWithWaterSelectionSupport, EWaterVertexFactoryDrawMode DrawMode);
 
 
 /**
@@ -256,6 +274,11 @@ struct TWaterMeshUserData
 	FRHIBuffer* IndirectInstanceData1 = nullptr;
 	FRHIBuffer* IndirectInstanceData2 = nullptr;
 	FRHIBuffer* IndirectInstanceData3 = nullptr;
+	FRHIShaderResourceView* IndirectInstanceDataOffsetsSRV = nullptr;
+	FRHIShaderResourceView* IndirectInstanceData0SRV = nullptr;
+	FRHIShaderResourceView* IndirectInstanceData1SRV = nullptr;
+	FRHIShaderResourceView* IndirectInstanceData2SRV = nullptr;
+	FRHIShaderResourceView* IndirectInstanceData3SRV = nullptr;
 };
 
 /**
@@ -285,7 +308,7 @@ struct TWaterMeshUserDataBuffers
 		return UserData[(int32)InRenderGroupType].Get();
 	}
 
-	TStaticArray<TUniquePtr<WaterMeshUserDataType>, TWaterVertexFactory<bWithWaterSelectionSupport, /*bIndirectDraws = */ false>::NumRenderGroups> UserData;
+	TStaticArray<TUniquePtr<WaterMeshUserDataType>, TWaterVertexFactory<bWithWaterSelectionSupport, EWaterVertexFactoryDrawMode::NonIndirect>::NumRenderGroups> UserData;
 };
 
 #include "WaterVertexFactory.inl"
diff --git a/Engine/Plugins/Experimental/Water/Source/Runtime/Public/WaterVertexFactory.inl b/Engine/Plugins/Experimental/Water/Source/Runtime/Public/WaterVertexFactory.inl
--- a/Engine/Plugins/Experimental/Water/Source/Runtime/Public/WaterVertexFactory.inl
+++ b/Engine/Plugins/Experimental/Water/Source/Runtime/Public/WaterVertexFactory.inl
@@ -13,8 +13,8 @@
 
 // ----------------------------------------------------------------------------------
 
-template <bool bWithWaterSelectionSupport, bool bIndirectDraws>
-TWaterVertexFactory<bWithWaterSelectionSupport, bIndirectDraws>::TWaterVertexFactory(ERHIFeatureLevel::Type InFeatureLevel, const FVector& InQuadTreePositionWS, int32 InNumQuadsPerSide, int32 InNumQuadsLOD0, int32 InNumDensities, float InLeafSize, float InLODScale, float InCaptureDepthRange)
+template <bool bWithWaterSelectionSupport, EWaterVertexFactoryDrawMode DrawMode>
+TWaterVertexFactory<bWithWaterSelectionSupport, DrawMode>::TWaterVertexFactory(ERHIFeatureLevel::Type InFeatureLevel, const FVector& InQuadTreePositionWS, int32 InNumQuadsPerSide, int32 InNumQuadsLOD0, int32 InNumDensities, float InLeafSize, float InLODScale, float InCaptureDepthRange)
 	: FVertexFactory(InFeatureLevel)
 	, QuadTreePositionWS(InQuadTreePositionWS)
 	, NumQuadsPerSide(InNumQuadsPerSide)
@@ -28,15 +28,15 @@ TWaterVertexFactory<bWithWaterSelectionSupport, bIndirectDraws>::TWaterVertexFac
 	IndexBuffer = new FWaterMeshIndexBuffer(NumQuadsPerSide);
 }
 
-template <bool bWithWaterSelectionSupport, bool bIndirectDraws>
-TWaterVertexFactory<bWithWaterSelectionSupport, bIndirectDraws>::~TWaterVertexFactory()
+template <bool bWithWaterSelectionSupport, EWaterVertexFactoryDrawMode DrawMode>
+TWaterVertexFactory<bWithWaterSelectionSupport, DrawMode>::~TWaterVertexFactory()
 {
 	delete VertexBuffer;
 	delete IndexBuffer;
 }
 
-template <bool bWithWaterSelectionSupport, bool bIndirectDraws>
-void TWaterVertexFactory<bWithWaterSelectionSupport, bIndirectDraws>::InitRHI(FRHICommandListBase& RHICmdList)
+template <bool bWithWaterSelectionSupport, EWaterVertexFactoryDrawMode DrawMode>
+void TWaterVertexFactory<bWithWaterSelectionSupport, DrawMode>::InitRHI(FRHICommandListBase& RHICmdList)
 {
 	Super::InitRHI(RHICmdList);
 
@@ -68,19 +68,23 @@ void TWaterVertexFactory<bWithWaterSelectionSupport, bIndirectDraws>::InitRHI(FR
 	Elements.Add(VertexPositionElement);
 
 	// Adds all streams
-	if constexpr (bIndirectDraws)
+	if constexpr (UsesIndirectDraws())
 	{
-		FVertexStream InstanceDataVertexStream;
-		InstanceDataVertexStream.VertexBuffer = nullptr;
-		InstanceDataVertexStream.Stride = sizeof(uint32);
-		InstanceDataVertexStream.Offset = 0;
-		InstanceDataVertexStream.VertexStreamUsage = EVertexStreamUsage::Instancing;
-
-		constexpr int NumBuffers = bWithWaterSelectionSupport ? 4 : 3;
-
-		for (int i = 0; i < NumBuffers; ++i)
+		// Instanced stereo manually fetches instance data from buffers instead
+		if (!UsesInstancedStereo())
 		{
-			Elements.Add(FVertexElement(Streams.Add(InstanceDataVertexStream), 0, VET_UInt, 8 + i, InstanceDataVertexStream.Stride, true));
+			FVertexStream InstanceDataVertexStream;
+			InstanceDataVertexStream.VertexBuffer = nullptr;
+			InstanceDataVertexStream.Stride = sizeof(uint32);
+			InstanceDataVertexStream.Offset = 0;
+			InstanceDataVertexStream.VertexStreamUsage = EVertexStreamUsage::Instancing;
+
+			constexpr int NumBuffers = bWithWaterSelectionSupport ? 4 : 3;
+
+			for (int i = 0; i < NumBuffers; ++i)
+			{
+				Elements.Add(FVertexElement(Streams.Add(InstanceDataVertexStream), 0, VET_UInt, 8 + i, InstanceDataVertexStream.Stride, true));
+			}
 		}
 	}
 	else if constexpr (NumAdditionalVertexStreams > 0)
@@ -116,8 +120,8 @@ void TWaterVertexFactory<bWithWaterSelectionSupport, bIndirectDraws>::InitRHI(FR
 	InitDeclaration(Elements);
 }
 
-template <bool bWithWaterSelectionSupport, bool bIndirectDraws>
-void TWaterVertexFactory<bWithWaterSelectionSupport, bIndirectDraws>::ReleaseRHI()
+template <bool bWithWaterSelectionSupport, EWaterVertexFactoryDrawMode DrawMode>
+void TWaterVertexFactory<bWithWaterSelectionSupport, DrawMode>::ReleaseRHI()
 {
 	for (auto& UniformBuffer : UniformBuffers)
 	{
@@ -137,8 +141,8 @@ void TWaterVertexFactory<bWithWaterSelectionSupport, bIndirectDraws>::ReleaseRHI
 	Super::ReleaseRHI();
 }
 
-template <bool bWithWaterSelectionSupport, bool bIndirectDraws>
-void TWaterVertexFactory<bWithWaterSelectionSupport, bIndirectDraws>::SetupUniformDataForGroup(EWaterMeshRenderGroupType InRenderGroupType)
+template <bool bWithWaterSelectionSupport, EWaterVertexFactoryDrawMode DrawMode>
+void TWaterVertexFactory<bWithWaterSelectionSupport, DrawMode>::SetupUniformDataForGroup(EWaterMeshRenderGroupType InRenderGroupType)
 {
 	FWaterVertexFactoryParameters UniformParams;
 	UniformParams.NumQuadsPerTileSide = NumQuadsPerSide;
@@ -158,8 +162,8 @@ void TWaterVertexFactory<bWithWaterSelectionSupport, bIndirectDraws>::SetupUnifo
 	UniformBuffers[(int32)InRenderGroupType] = FWaterVertexFactoryBufferRef::CreateUniformBufferImmediate(UniformParams, UniformBuffer_MultiFrame);
 }
 
-template <bool bWithWaterSelectionSupport, bool bIndirectDraws>
-bool TWaterVertexFactory<bWithWaterSelectionSupport, bIndirectDraws>::ShouldCompilePermutation(const FVertexFactoryShaderPermutationParameters& Parameters)
+template <bool bWithWaterSelectionSupport, EWaterVertexFactoryDrawMode DrawMode>
+bool TWaterVertexFactory<bWithWaterSelectionSupport, DrawMode>::ShouldCompilePermutation(const FVertexFactoryShaderPermutationParameters& Parameters)
 {
 	const bool bIsCompatibleWithWater = ((Parameters.MaterialParameters.MaterialDomain == MD_Surface) && Parameters.MaterialParameters.bIsUsedWithWater) || Parameters.MaterialParameters.bIsSpecialEngineMaterial;
 	if (bIsCompatibleWithWater)
@@ -170,8 +174,8 @@ bool TWaterVertexFactory<bWithWaterSelectionSupport, bIndirectDraws>::ShouldComp
 	return false;
 }
 
-template <bool bWithWaterSelectionSupport, bool bIndirectDraws>
-void TWaterVertexFactory<bWithWaterSelectionSupport, bIndirectDraws>::ModifyCompilationEnvironment(const FVertexFactoryShaderPermutationParameters& Parameters, FShaderCompilerEnvironment& OutEnvironment)
+template <bool bWithWaterSelectionSupport, EWaterVertexFactoryDrawMode DrawMode>
+void TWaterVertexFactory<bWithWaterSelectionSupport, DrawMode>::ModifyCompilationEnvironment(const FVertexFactoryShaderPermutationParameters& Parameters, FShaderCompilerEnvironment& OutEnvironment)
 {
 	OutEnvironment.SetDefine(TEXT("WATER_MESH_FACTORY"), 1);
 #if 0
@@ -189,21 +193,17 @@ void TWaterVertexFactory<bWithWaterSelectionSupport, bIndirectDraws>::ModifyComp
 		OutEnvironment.SetDefine(TEXT("USE_VERTEXFACTORY_HITPROXY_ID"), TEXT("1"));
 		OutEnvironment.SetDefine(TEXT("WITH_WATER_SELECTION_SUPPORT_VF"), TEXT("1"));
 	}
-	if (bIndirectDraws)
+	if (UsesIndirectDraws())
 	{
 		OutEnvironment.SetDefine(TEXT("WATER_MESH_DRAW_INDIRECT"), TEXT("1"));
+		OutEnvironment.CompilerFlags.Add(CFLAG_IndirectDraw);
 	}
 
 	OutEnvironment.SetDefine(TEXT("RAY_TRACING_DYNAMIC_MESH_IN_LOCAL_SPACE"), TEXT("1"));
-
-	if (bIndirectDraws)
-	{
-		OutEnvironment.CompilerFlags.Add(CFLAG_IndirectDraw);
-	}
 }
 
-template <bool bWithWaterSelectionSupport, bool bIndirectDraws>
-void TWaterVertexFactory<bWithWaterSelectionSupport, bIndirectDraws>::ValidateCompiledResult(const FVertexFactoryType* Type, EShaderPlatform Platform, const FShaderParameterMap& ParameterMap, TArray<FString>& OutErrors)
+template <bool bWithWaterSelectionSupport, EWaterVertexFactoryDrawMode DrawMode>
+void TWaterVertexFactory<bWithWaterSelectionSupport, DrawMode>::ValidateCompiledResult(const FVertexFactoryType* Type, EShaderPlatform Platform, const FShaderParameterMap& ParameterMap, TArray<FString>& OutErrors)
 {
 #if 0
 	if (Type->SupportsPrimitiveIdStream()
@@ -215,21 +215,25 @@ void TWaterVertexFactory<bWithWaterSelectionSupport, bIndirectDraws>::ValidateCo
 #endif
 }
 
-template <bool bWithWaterSelectionSupport, bool bIndirectDraws>
-void TWaterVertexFactory<bWithWaterSelectionSupport, bIndirectDraws>::GetPSOPrecacheVertexFetchElements(EVertexInputStreamType VertexInputStreamType, FVertexDeclarationElementList& Elements)
+template <bool bWithWaterSelectionSupport, EWaterVertexFactoryDrawMode DrawMode>
+void TWaterVertexFactory<bWithWaterSelectionSupport, DrawMode>::GetPSOPrecacheVertexFetchElements(EVertexInputStreamType VertexInputStreamType, FVertexDeclarationElementList& Elements)
 {
 	// Add position stream
 	Elements.Add(FVertexElement(0, 0, VET_Float4, 0, sizeof(FVector4f), false));
 
 	// Add all the additional streams
-	if constexpr (bIndirectDraws)
+	if constexpr (UsesIndirectDraws())
 	{
-		Elements.Add(FVertexElement(1, 0, VET_UInt, 8, sizeof(uint32), true));
-		Elements.Add(FVertexElement(2, 0, VET_UInt, 9, sizeof(uint32), true));
-		Elements.Add(FVertexElement(3, 0, VET_UInt, 10, sizeof(uint32), true));
-		if (bWithWaterSelectionSupport)
+		// Instanced stereo manually fetches instance data from buffers instead
+		if (!UsesInstancedStereo())
 		{
-			Elements.Add(FVertexElement(4, 0, VET_UInt, 11, sizeof(uint32), true));
+			Elements.Add(FVertexElement(1, 0, VET_UInt, 8, sizeof(uint32), true));
+			Elements.Add(FVertexElement(2, 0, VET_UInt, 9, sizeof(uint32), true));
+			Elements.Add(FVertexElement(3, 0, VET_UInt, 10, sizeof(uint32), true));
+			if (bWithWaterSelectionSupport)
+			{
+				Elements.Add(FVertexElement(4, 0, VET_UInt, 11, sizeof(uint32), true));
+			}
 		}
 	}
 	else if constexpr (NumAdditionalVertexStreams > 0)
