Water: prevent updating WaterZone aspects which are not actually affected by a water body when they attempt to mark their owning water zone as dirty.

diff --git a/Engine/Plugins/Experimental/Water/Source/Runtime/Private/WaterBodyComponent.cpp b/Engine/Plugins/Experimental/Water/Source/Runtime/Private/WaterBodyComponent.cpp
--- a/Engine/Plugins/Experimental/Water/Source/Runtime/Private/WaterBodyComponent.cpp
+++ b/Engine/Plugins/Experimental/Water/Source/Runtime/Private/WaterBodyComponent.cpp
@@ -432,14 +432,26 @@ void UWaterBodyComponent::MarkOwningWaterZoneForRebuild(EWaterZoneRebuildFlags I
 {
 	if (AWaterZone* WaterZone = GetWaterZone())
 	{
+		EWaterZoneRebuildFlags RebuildFlags = InRebuildFlags;
+
+		// Avoid rebuilding things which this water body should not affect:
+		if (!AffectsWaterInfo())
+		{
+			EnumRemoveFlags(RebuildFlags, EWaterZoneRebuildFlags::UpdateWaterInfoTexture);
+		}
+		if (!AffectsWaterMesh())
+		{
+			EnumRemoveFlags(RebuildFlags, EWaterZoneRebuildFlags::UpdateWaterMesh);
+		}
+
 		if (bInOnlyWithinWaterBodyBounds)
 		{
 			const FBox WaterBodyBounds = Bounds.GetBox();
-			WaterZone->MarkForRebuild(InRebuildFlags, FBox2D(FVector2D(WaterBodyBounds.Min), FVector2D(WaterBodyBounds.Max)), /* DebugRequestingObject = */ GetOwner());
+			WaterZone->MarkForRebuild(RebuildFlags, FBox2D(FVector2D(WaterBodyBounds.Min), FVector2D(WaterBodyBounds.Max)), /* DebugRequestingObject = */ GetOwner());
 		}
 		else
 		{
-			WaterZone->MarkForRebuild(InRebuildFlags, /* DebugRequestingObject = */ GetOwner());
+			WaterZone->MarkForRebuild(RebuildFlags, /* DebugRequestingObject = */ GetOwner());
 		}
 	}
 }
@@ -1922,13 +1934,7 @@ void UWaterBodyComponent::OnWaterBodyRenderDataUpdated()
 	const bool bAffectsWaterInfo = AffectsWaterInfo();
 	const bool bAffectsWaterMesh = AffectsWaterMesh();
 
-	if (bAffectsWaterInfo)
-	{
-		if (AWaterZone* WaterZone = GetWaterZone())
-		{
-			MarkOwningWaterZoneForRebuild(EWaterZoneRebuildFlags::UpdateWaterInfoTexture, /* bOnlyWithinWaterBodyBounds = */ false);
-		}
-	}
+	MarkOwningWaterZoneForRebuild(EWaterZoneRebuildFlags::UpdateWaterInfoTexture, /* bOnlyWithinWaterBodyBounds = */ false);
 
 	if (bAffectsWaterInfo || bAffectsWaterMesh)
 	{
