Generalized the info assigned to water zone and bind its lifetime to the water zone registration, so that we can eventually store other data than their bounds
#tests editor, win64 cooked build

diff --git a/Engine/Plugins/Experimental/Water/Source/Runtime/Private/WaterBodyManager.cpp b/Engine/Plugins/Experimental/Water/Source/Runtime/Private/WaterBodyManager.cpp
--- a/Engine/Plugins/Experimental/Water/Source/Runtime/Private/WaterBodyManager.cpp
+++ b/Engine/Plugins/Experimental/Water/Source/Runtime/Private/WaterBodyManager.cpp
@@ -38,15 +38,29 @@ void FWaterBodyManager::RemoveWaterBodyComponent(UWaterBodyComponent* InWaterBod
 
 int32 FWaterBodyManager::AddWaterZone(AWaterZone* InWaterZone)
 {
-	RequestGPUDataRebuild();
 	int32 LowestFreeIndex = 0;
-	return WaterZones.EmplaceAtLowestFreeIndex(LowestFreeIndex, InWaterZone);
+	int32 WaterZoneIndex = WaterZones.EmplaceAtLowestFreeIndex(LowestFreeIndex, InWaterZone);
+
+	if (WaterViewExtension)
+	{
+		WaterViewExtension->AddWaterZone(InWaterZone);
+	}
+
+	RequestGPUDataRebuild();
+
+	return WaterZoneIndex;
 }
 
 void FWaterBodyManager::RemoveWaterZone(AWaterZone* InWaterZone)
 {
-	RequestGPUDataRebuild();
+	if (WaterViewExtension)
+	{
+		WaterViewExtension->RemoveWaterZone(InWaterZone);
+	}
+
 	WaterZones.RemoveAt(InWaterZone->GetWaterZoneIndex());
+
+	RequestGPUDataRebuild();
 }
 
 void FWaterBodyManager::RequestGPUDataRebuild()
diff --git a/Engine/Plugins/Experimental/Water/Source/Runtime/Private/WaterViewExtension.cpp b/Engine/Plugins/Experimental/Water/Source/Runtime/Private/WaterViewExtension.cpp
--- a/Engine/Plugins/Experimental/Water/Source/Runtime/Private/WaterViewExtension.cpp
+++ b/Engine/Plugins/Experimental/Water/Source/Runtime/Private/WaterViewExtension.cpp
@@ -313,11 +313,15 @@ void FWaterViewExtension::SetupView(FSceneViewFamily& InViewFamily, FSceneView&
 	ON_SCOPE_EXIT { bUpdatingWaterInfo = false; };
 
 	const int32 WaterInfoRenderMethod = CVarWaterInfoRenderMethod.GetValueOnGameThread();
+	FSceneInterface* Scene = WorldPtr.Get()->Scene;
+	check(Scene != nullptr);
 
 	for (const TPair<AWaterZone*, UE::WaterInfo::FRenderingContext>& Pair : WaterInfoContextsToRender)
 	{
 		AWaterZone* WaterZone = Pair.Key;
-		check(WaterZone);
+		check(WaterZone != nullptr);
+		FWaterZoneInfo* WaterZoneInfo = WaterZoneInfos.Find(WaterZone);
+		check(WaterZoneInfo != nullptr);
 
 		if (WaterZone->IsLocalOnlyTessellationEnabled())
 		{
@@ -338,9 +342,7 @@ void FWaterViewExtension::SetupView(FSceneViewFamily& InViewFamily, FSceneView&
 
 			// Keep a minimum of <1., 1.> bounds to avoid updating every frame if the update margin is larger than the zone.
 			const FVector2D UpdateExtents = FVector2D::Max(FVector2D(1., 1.), WaterInfoHalfExtent - UpdateMargin);
-			const FBox2D UpdateBounds(FVector2D(WaterMeshCenter) - UpdateExtents, FVector2D(WaterMeshCenter) + UpdateExtents);
-
-			WaterInfoUpdateBounds.Add(WaterZone, UpdateBounds);
+			WaterZoneInfo->UpdateBounds.Emplace(FVector2D(WaterMeshCenter) - UpdateExtents, FVector2D(WaterMeshCenter) + UpdateExtents);
 
 			const FVector2D WaterQuadTreeHalfExtent = WaterMesh->GetExtentInTiles() * TileSize;
 			const FVector2D WaterQuadTreeCenter = WaterMesh->GetDynamicWaterMeshCenter();
@@ -355,18 +357,22 @@ void FWaterViewExtension::SetupView(FSceneViewFamily& InViewFamily, FSceneView&
 			// Mark GPU data dirty since we have a new WaterArea parameter and need to push this to water bodies.
 			MarkGPUDataDirty();
 		}
+		else
+		{
+			WaterZoneInfo->UpdateBounds.Reset();
+		}
 
 		// Old method of rendering the water info texture; uses scene captures
 		if (WaterInfoRenderMethod == 0)
 		{
 			const UE::WaterInfo::FRenderingContext& Context(Pair.Value);
-			UE::WaterInfo::UpdateWaterInfoRendering(WorldPtr.Get()->Scene, Context);
+			UE::WaterInfo::UpdateWaterInfoRendering(Scene, Context);
 		}
 		// Render the water info texture using custom render pass method
 		else if (WaterInfoRenderMethod == 2)
 		{
 			const UE::WaterInfo::FRenderingContext& Context(Pair.Value);
-			UE::WaterInfo::UpdateWaterInfoRendering_CustomRenderPass(WorldPtr.Get()->Scene, Context);
+			UE::WaterInfo::UpdateWaterInfoRendering_CustomRenderPass(Scene, Context);
 		}
 	}
 
@@ -385,10 +391,11 @@ void FWaterViewExtension::SetupView(FSceneViewFamily& InViewFamily, FSceneView&
 		// Check if the view location is no longer within the current update bounds of a water zone and if so, queue an update for it.
 		for (AWaterZone* WaterZone : TActorRange<AWaterZone>(WorldPtr.Get()))
 		{
+			FWaterZoneInfo* WaterZoneInfo = WaterZoneInfos.Find(WaterZone);
+			check(WaterZoneInfo != nullptr);
 			if (WaterZone->IsLocalOnlyTessellationEnabled())
 			{
-				const FBox2D* WaterInfoBounds = WaterInfoUpdateBounds.Find(WaterZone);
-				if (WaterInfoBounds == nullptr || !WaterInfoBounds->IsInside(FVector2D(ViewLocation)))
+				if (!WaterZoneInfo->UpdateBounds.IsSet() || !WaterZoneInfo->UpdateBounds->IsInside(FVector2D(ViewLocation)))
 				{
 					WaterZone->MarkForRebuild(EWaterZoneRebuildFlags::UpdateWaterInfoTexture);
 				}
@@ -433,3 +440,14 @@ void FWaterViewExtension::MarkGPUDataDirty()
 {
 	bRebuildGPUData = true;
 }
+
+void FWaterViewExtension::AddWaterZone(AWaterZone* InWaterZone)
+{
+	check(!WaterZoneInfos.Contains(InWaterZone));
+	WaterZoneInfos.Emplace(InWaterZone);
+}
+
+void FWaterViewExtension::RemoveWaterZone(AWaterZone* InWaterZone)
+{
+	WaterZoneInfos.FindAndRemoveChecked(InWaterZone);
+}
diff --git a/Engine/Plugins/Experimental/Water/Source/Runtime/Public/WaterViewExtension.h b/Engine/Plugins/Experimental/Water/Source/Runtime/Public/WaterViewExtension.h
--- a/Engine/Plugins/Experimental/Water/Source/Runtime/Public/WaterViewExtension.h
+++ b/Engine/Plugins/Experimental/Water/Source/Runtime/Public/WaterViewExtension.h
@@ -30,15 +30,23 @@ public:
 	void MarkWaterInfoTextureForRebuild(const UE::WaterInfo::FRenderingContext& RenderContext);
 
 	void MarkGPUDataDirty();
+
+	void AddWaterZone(AWaterZone* InWaterZone);
+	void RemoveWaterZone(AWaterZone* InWaterZone);
+
 private:
 	/** Queued Water Info rendering contexts to submit for rendering on the next SetupView call */
 	TMap<AWaterZone*, UE::WaterInfo::FRenderingContext> WaterInfoContextsToRender;
 
-	/** 
-	 * For each water zone, store the bounds of the tile from which the water zone was last rendered.
-	 * When the view location crosses the bounds, submit a new WaterInfo update to reflect the new active area
-	 */
-	TMap<AWaterZone*, FBox2D> WaterInfoUpdateBounds;
+	struct FWaterZoneInfo
+	{
+		/** 
+		 * For each water zone, store the bounds of the tile from which the water zone was last rendered.
+		 * When the view location crosses the bounds, submit a new WaterInfo update to reflect the new active area
+		 */
+		TOptional<FBox2D> UpdateBounds = FBox2D(ForceInit);
+	};
+	TMap<AWaterZone*, FWaterZoneInfo> WaterZoneInfos;
 
 	struct FWaterGPUResources
 	{
