Disabled water MID deduplication by default.

#rnx

diff --git a/Engine/Plugins/Experimental/Water/Source/Runtime/Private/WaterMeshComponent.cpp b/Engine/Plugins/Experimental/Water/Source/Runtime/Private/WaterMeshComponent.cpp
--- a/Engine/Plugins/Experimental/Water/Source/Runtime/Private/WaterMeshComponent.cpp
+++ b/Engine/Plugins/Experimental/Water/Source/Runtime/Private/WaterMeshComponent.cpp
@@ -65,7 +65,7 @@ TAutoConsoleVariable<int32> CVarWaterMeshEnabled(
 
 TAutoConsoleVariable<int32> CVarWaterMeshMIDDeduplication(
 	TEXT("r.Water.WaterMesh.MIDDeduplication"),
-	1,
+	0,
 	TEXT("Deduplicate per-water body MIDs"),
 	ECVF_RenderThreadSafe
 );
