Water: fixed FindLandscape using the component location rather than a 2d bounds check. It would be possible previously for components which have their root component locations outside the landscape but spline points individually are all within the landscape would not correctly pick it up.

diff --git a/Engine/Plugins/Experimental/Water/Source/Runtime/Private/WaterBodyComponent.cpp b/Engine/Plugins/Experimental/Water/Source/Runtime/Private/WaterBodyComponent.cpp
--- a/Engine/Plugins/Experimental/Water/Source/Runtime/Private/WaterBodyComponent.cpp
+++ b/Engine/Plugins/Experimental/Water/Source/Runtime/Private/WaterBodyComponent.cpp
@@ -1016,13 +1016,13 @@ ALandscapeProxy* UWaterBodyComponent::FindLandscape() const
 {
 	if (!Landscape.IsValid())
 	{
-		const FVector Location = GetComponentLocation();
+		const FBox ComponentBounds = Bounds.GetBox();
 		for (TObjectIterator<ALandscapeProxy> It; It; ++It)
 		{
 			if (It->GetWorld() == GetWorld())
 			{
-				FBox Box = It->GetComponentsBoundingBox();
-				if (Box.IsInsideOrOnXY(Location))
+				const FBox Box = It->GetComponentsBoundingBox();
+				if (Box.IntersectXY(ComponentBounds))
 				{
 					Landscape = *It;
 					return Landscape.Get();
diff --git a/Engine/Plugins/Experimental/Water/Source/Runtime/Private/WaterBodyCustomComponent.cpp b/Engine/Plugins/Experimental/Water/Source/Runtime/Private/WaterBodyCustomComponent.cpp
--- a/Engine/Plugins/Experimental/Water/Source/Runtime/Private/WaterBodyCustomComponent.cpp
+++ b/Engine/Plugins/Experimental/Water/Source/Runtime/Private/WaterBodyCustomComponent.cpp
@@ -47,6 +47,15 @@ TArray<UPrimitiveComponent*> UWaterBodyCustomComponent::GetStandardRenderableCom
 	return Result;
 }
 
+FBoxSphereBounds UWaterBodyCustomComponent::CalcBounds(const FTransform& LocalToWorld) const
+{
+	if (MeshComp)
+	{
+		return MeshComp->CalcBounds(MeshComp->GetRelativeTransform()).TransformBy(LocalToWorld);
+	}
+	return Super::CalcBounds(LocalToWorld);
+}
+
 void UWaterBodyCustomComponent::Reset()
 {
 	AActor* Owner = GetOwner();
diff --git a/Engine/Plugins/Experimental/Water/Source/Runtime/Public/WaterBodyCustomComponent.h b/Engine/Plugins/Experimental/Water/Source/Runtime/Public/WaterBodyCustomComponent.h
--- a/Engine/Plugins/Experimental/Water/Source/Runtime/Public/WaterBodyCustomComponent.h
+++ b/Engine/Plugins/Experimental/Water/Source/Runtime/Public/WaterBodyCustomComponent.h
@@ -15,14 +15,16 @@ class WATER_API UWaterBodyCustomComponent : public UWaterBodyComponent
 	GENERATED_UCLASS_BODY()
 	friend class AWaterBodyCustom;
 public:
-	/** AWaterBody Interface */
+	/** UWaterBodyComponent Interface */
 	virtual EWaterBodyType GetWaterBodyType() const override { return EWaterBodyType::Transition; }
 	virtual TArray<UPrimitiveComponent*> GetCollisionComponents(bool bInOnlyEnabledComponents = true) const override;
 	virtual TArray<UPrimitiveComponent*> GetStandardRenderableComponents() const override;
 	virtual bool CanEverAffectWaterMesh() const { return false; }
 
+	virtual FBoxSphereBounds CalcBounds(const FTransform& LocalToWorld) const override;
+
 protected:
-	/** AWaterBody Interface */
+	/** UWaterBodyComponent Interface */
 	virtual void Reset() override;
 	virtual void BeginUpdateWaterBody() override;
 	virtual void OnUpdateBody(bool bWithExclusionVolumes) override;
