Water: Made the water bodies list on exclusion volumes Blueprint Readable.

diff --git a/Engine/Plugins/Experimental/Water/Source/Runtime/Public/WaterBodyExclusionVolume.h b/Engine/Plugins/Experimental/Water/Source/Runtime/Public/WaterBodyExclusionVolume.h
--- a/Engine/Plugins/Experimental/Water/Source/Runtime/Public/WaterBodyExclusionVolume.h
+++ b/Engine/Plugins/Experimental/Water/Source/Runtime/Public/WaterBodyExclusionVolume.h
@@ -74,7 +74,7 @@ public:
 	EWaterExclusionMode ExclusionMode = EWaterExclusionMode::RemoveWaterBodiesListFromExclusion;
 
 	/** List of water bodies that will be added or removed from the exclusion volume based on the ExclusionMode parameter. */
-	UPROPERTY(EditInstanceOnly, Category = Water, meta = (DisplayAfter=ExclusionMode))
+	UPROPERTY(EditInstanceOnly, BlueprintReadOnly, Category = Water, meta = (DisplayAfter=ExclusionMode))
 	TArray<TSoftObjectPtr<AWaterBody>> WaterBodies;
 
 #if WITH_EDITORONLY_DATA
