World Partition - HLOD: When building HLODs for water bodies, handle more failure cases when trying to retrieve a WaterMesh mesh description

diff --git a/Engine/Plugins/Experimental/Water/Source/Runtime/Private/WaterBodyComponent.cpp b/Engine/Plugins/Experimental/Water/Source/Runtime/Private/WaterBodyComponent.cpp
--- a/Engine/Plugins/Experimental/Water/Source/Runtime/Private/WaterBodyComponent.cpp
+++ b/Engine/Plugins/Experimental/Water/Source/Runtime/Private/WaterBodyComponent.cpp
@@ -2108,12 +2108,9 @@ FMeshDescription UWaterBodyComponent::GetHLODMeshDescription() const
 	{
 		const int32 LastLOD = WaterMeshOverride->GetNumLODs() - 1;
 
-		// If source model is valid, return the mesh description directly, otherwise recreate it from the render data
-		if (WaterMeshOverride->IsSourceModelValid(LastLOD))
-		{
-			MeshDescription = *WaterMeshOverride->GetMeshDescription(LastLOD);
-		}
-		else
+		// If source model is valid, clone the mesh description, otherwise recreate it from the render data
+		bool bIsMeshDescriptionValid = WaterMeshOverride->CloneMeshDescription(LastLOD, MeshDescription);
+		if (!bIsMeshDescriptionValid)
 		{
 			const IMeshMergeUtilities& MeshMergeUtilities = FModuleManager::Get().LoadModuleChecked<IMeshMergeModule>("MeshMergeUtilities").GetUtilities();
 			MeshMergeUtilities.RetrieveMeshDescription(WaterMeshOverride, LastLOD, MeshDescription);
