Water: added some additional checks in the WaterViewExtension to track down an issue where water zones may not have registered water zone info.

diff --git a/Engine/Plugins/Experimental/Water/Source/Runtime/Private/WaterViewExtension.cpp b/Engine/Plugins/Experimental/Water/Source/Runtime/Private/WaterViewExtension.cpp
--- a/Engine/Plugins/Experimental/Water/Source/Runtime/Private/WaterViewExtension.cpp
+++ b/Engine/Plugins/Experimental/Water/Source/Runtime/Private/WaterViewExtension.cpp
@@ -391,14 +391,17 @@ void FWaterViewExtension::SetupView(FSceneViewFamily& InViewFamily, FSceneView&
 		// Check if the view location is no longer within the current update bounds of a water zone and if so, queue an update for it.
 		for (AWaterZone* WaterZone : TActorRange<AWaterZone>(WorldPtr.Get()))
 		{
+			checkf(IsValid(WaterZone), TEXT("TActorRange<AWaterZone> gave us an invalid Water Zone actor!"));
 			FWaterZoneInfo* WaterZoneInfo = WaterZoneInfos.Find(WaterZone);
-			check(WaterZoneInfo != nullptr);
-			if (WaterZone->IsLocalOnlyTessellationEnabled())
+			if (ensureMsgf(WaterZoneInfo != nullptr, TEXT("Water Zone (%s) has no corresponding water zone info!"), *GetNameSafe(WaterZone)))
 			{
-				if (!WaterZoneInfo->UpdateBounds.IsSet() || !WaterZoneInfo->UpdateBounds->IsInside(FVector2D(ViewLocation)))
-				{
-					WaterZone->MarkForRebuild(EWaterZoneRebuildFlags::UpdateWaterInfoTexture);
-				}
+				 if (WaterZone->IsLocalOnlyTessellationEnabled())
+				 {
+					 if (!WaterZoneInfo->UpdateBounds.IsSet() || !WaterZoneInfo->UpdateBounds->IsInside(FVector2D(ViewLocation)))
+					 {
+						 WaterZone->MarkForRebuild(EWaterZoneRebuildFlags::UpdateWaterInfoTexture);
+					 }
+				 }
 			}
 		}
 	}
