Fixed PSO precaching error on ISR water vertex factory. Fixed by only compiling the ISR/non-ISR variant if ISR is actually enabled/disabled. Also removed indirect VFs from the list of used VFs of the far water mesh material as that one is always rendered with the legacy CPU-driven rendering path.

#rnx

diff --git a/Engine/Plugins/Experimental/Water/Source/Runtime/Private/WaterMeshComponent.cpp b/Engine/Plugins/Experimental/Water/Source/Runtime/Private/WaterMeshComponent.cpp
--- a/Engine/Plugins/Experimental/Water/Source/Runtime/Private/WaterMeshComponent.cpp
+++ b/Engine/Plugins/Experimental/Water/Source/Runtime/Private/WaterMeshComponent.cpp
@@ -99,8 +99,6 @@ void UWaterMeshComponent::CollectPSOPrecacheData(const FPSOPrecacheParams& BaseP
 		ComponentParams.Priority = EPSOPrecachePriority::High;
 		ComponentParams.MaterialInterface = FarDistanceMaterial;
 		ComponentParams.VertexFactoryDataList.Add(FPSOPrecacheVertexFactoryData(WaterVertexFactoryNonIndirectType));
-		ComponentParams.VertexFactoryDataList.Add(FPSOPrecacheVertexFactoryData(WaterVertexFactoryIndirectType));
-		ComponentParams.VertexFactoryDataList.Add(FPSOPrecacheVertexFactoryData(WaterVertexFactoryIndirectISRType));
 		ComponentParams.PSOPrecacheParams = BasePrecachePSOParams;
 	}
 	for (UMaterialInterface* MaterialInterface : UsedMaterials)
diff --git a/Engine/Plugins/Experimental/Water/Source/Runtime/Private/WaterMeshSceneProxy.cpp b/Engine/Plugins/Experimental/Water/Source/Runtime/Private/WaterMeshSceneProxy.cpp
--- a/Engine/Plugins/Experimental/Water/Source/Runtime/Private/WaterMeshSceneProxy.cpp
+++ b/Engine/Plugins/Experimental/Water/Source/Runtime/Private/WaterMeshSceneProxy.cpp
@@ -209,19 +209,22 @@ FWaterMeshSceneProxy::FWaterMeshSceneProxy(UWaterMeshComponent* Component)
 	const FVector QuadTreeCorner = FVector(WaterQuadTree.GetTileRegion().Min, WaterQuadTreeMinHeight);
 	const float WaterQuadTreeDepthRange = WaterQuadTreeMaxHeight - WaterQuadTreeMinHeight;
 	const float LeafSize = WaterQuadTree.GetLeafSize();
+	
+	// bIsGPUQuadTree is constant over the lifetime of this scene proxy, so we know up front if we need the indirect draw vertex factory
 	if (bIsGPUQuadTree)
 	{
-		// bIsGPUQuadTree is constant over the lifetime of this scene proxy, so we know up front if we need the indirect draw vertex factory
-		WaterVertexFactoryIndirectDraw = new FWaterVertexFactoryIndirectDrawType(GetScene().GetFeatureLevel(), QuadTreeCorner, NumQuadsPerIndirectDrawTile, NumQuadsLOD0, DensityCount, LeafSize, LODScale, WaterQuadTreeDepthRange);
-		BeginInitResource(WaterVertexFactoryIndirectDraw);
-
-		// Only create the ISR vertex factory if ISR is enabled in the first place
+		// Only create one indirect draw vertex factory, depending on whether ISR is enabled
 		const UE::StereoRenderUtils::FStereoShaderAspects Aspects(GetScene().GetShaderPlatform());
 		if (Aspects.IsInstancedStereoEnabled())
 		{
 			WaterVertexFactoryIndirectDrawISR = new FWaterVertexFactoryIndirectDrawISRType(GetScene().GetFeatureLevel(), QuadTreeCorner, NumQuadsPerIndirectDrawTile, NumQuadsLOD0, DensityCount, LeafSize, LODScale, WaterQuadTreeDepthRange);
 			BeginInitResource(WaterVertexFactoryIndirectDrawISR);
 		}
+		else
+		{
+			WaterVertexFactoryIndirectDraw = new FWaterVertexFactoryIndirectDrawType(GetScene().GetFeatureLevel(), QuadTreeCorner, NumQuadsPerIndirectDrawTile, NumQuadsLOD0, DensityCount, LeafSize, LODScale, WaterQuadTreeDepthRange);
+			BeginInitResource(WaterVertexFactoryIndirectDraw);
+		}
 	}
 
 	// We always need the basic CPU-driven rendering vertex factory because the far mesh uses it
diff --git a/Engine/Plugins/Experimental/Water/Source/Runtime/Public/WaterVertexFactory.inl b/Engine/Plugins/Experimental/Water/Source/Runtime/Public/WaterVertexFactory.inl
--- a/Engine/Plugins/Experimental/Water/Source/Runtime/Public/WaterVertexFactory.inl
+++ b/Engine/Plugins/Experimental/Water/Source/Runtime/Public/WaterVertexFactory.inl
@@ -9,6 +9,7 @@
 #include "MeshMaterialShader.h"
 #include "WaterInstanceDataBuffer.h"
 #include "DataDrivenShaderPlatformInfo.h"
+#include "StereoRenderUtils.h"
 
 
 // ----------------------------------------------------------------------------------
@@ -168,6 +169,16 @@ bool TWaterVertexFactory<bWithWaterSelectionSupport, DrawMode>::ShouldCompilePer
 	const bool bIsCompatibleWithWater = ((Parameters.MaterialParameters.MaterialDomain == MD_Surface) && Parameters.MaterialParameters.bIsUsedWithWater) || Parameters.MaterialParameters.bIsSpecialEngineMaterial;
 	if (bIsCompatibleWithWater)
 	{
+		// Only compile the ISR/non-ISR indirect draw version if ISR is enabled/disabled
+		if (UsesIndirectDraws())
+		{
+			const UE::StereoRenderUtils::FStereoShaderAspects Aspects(Parameters.Platform);
+			const bool bMatchingISRConfig = UsesInstancedStereo() == Aspects.IsInstancedStereoEnabled();
+			if (!bMatchingISRConfig)
+			{
+				return false;
+			}
+		}
 		// Only let the PC platform compile the permutations supporting selection : 
 		return (!bWithWaterSelectionSupport || IsPCPlatform(Parameters.Platform));
 	}
