Hack to allow MRQ to render the water info texture. This will be removed when per-view water info texture is supported
#tests editor, MRQ

diff --git a/Engine/Plugins/Experimental/Water/Source/Runtime/Private/WaterZoneActor.cpp b/Engine/Plugins/Experimental/Water/Source/Runtime/Private/WaterZoneActor.cpp
--- a/Engine/Plugins/Experimental/Water/Source/Runtime/Private/WaterZoneActor.cpp
+++ b/Engine/Plugins/Experimental/Water/Source/Runtime/Private/WaterZoneActor.cpp
@@ -38,7 +38,12 @@ TAutoConsoleVariable<float> CVarWaterFallbackDepth(
 	TEXT("Depth to use for all water when there are no ground actors defined."),
 	ECVF_Default);
 
-
+// HACK [jonathan.bard] : (details underneath)
+TAutoConsoleVariable<int32> CVarSkipWaterInfoTextureRenderWhenWorldRenderingDisabled(
+	TEXT("r.Water.SkipWaterInfoTextureRenderWhenWorldRenderingDisabled"),
+	1,
+	TEXT("Use this to prevent the water info from rendering when world rendering is disabled."),
+	ECVF_Default);
 
 AWaterZone::AWaterZone(const FObjectInitializer& Initializer)
 	: Super(Initializer)
@@ -580,7 +585,11 @@ bool AWaterZone::UpdateWaterInfoTexture()
 		// The render path for rendering the water info texture without scene captures is executed within the scene renderer.
 		// If world rendering is disabled like in a loading screen, the scene renderer is not called and the water info texture will not be drawn.
 		UGameViewportClient* GameViewport = World->GetGameViewport();
-		if (GameViewport && GameViewport->bDisableWorldRendering)
+		// HACK [jonathan.bard] : CVarSkipWaterInfoTextureRenderWhenWorldRenderingDisabled is a temporary hack for MRQ because bDisableWorldRendering is true when capturing for MRQ. 
+		//  The proper solution is to refactor the water info texture and make it per-view (also for split screen support) : FWaterViewExtension::SetupViewFamily/SetupView
+		//  are called on the game thread so they might be a good place for doing this. 
+		//  For now, we just use this CVar trick to let the water info texture be rendered when MRQ runs :
+		if ((CVarSkipWaterInfoTextureRenderWhenWorldRenderingDisabled.GetValueOnGameThread() != 0) && GameViewport && GameViewport->bDisableWorldRendering)
 		{
 			return false;
 		}
diff --git a/Engine/Plugins/MovieScene/MovieRenderPipeline/Source/MovieRenderPipelineCore/Private/MoviePipelineGameOverrideSetting.cpp b/Engine/Plugins/MovieScene/MovieRenderPipeline/Source/MovieRenderPipelineCore/Private/MoviePipelineGameOverrideSetting.cpp
--- a/Engine/Plugins/MovieScene/MovieRenderPipeline/Source/MovieRenderPipelineCore/Private/MoviePipelineGameOverrideSetting.cpp
+++ b/Engine/Plugins/MovieScene/MovieRenderPipeline/Source/MovieRenderPipelineCore/Private/MoviePipelineGameOverrideSetting.cpp
@@ -137,6 +137,9 @@ void UMoviePipelineGameOverrideSetting::ApplyCVarSettings(const bool bOverrideVa
 	// Cloth's time step smoothing messes up the change in number of simulation substeps that fixes the cloth simulation behavior when using Temporal Samples.
 	MOVIEPIPELINE_STORE_AND_OVERRIDE_CVAR_INT(PreviousChaosClothUseTimeStepSmoothing, TEXT("p.ChaosCloth.UseTimeStepSmoothing"), 0, bOverrideValues);
 
+	// Water skips water info texture when the world's game viewport rendering is disabled so we need to prevent this from happening.
+	MOVIEPIPELINE_STORE_AND_OVERRIDE_CVAR_INT(PreviousSkipWaterInfoTextureRenderWhenWorldRenderingDisabled, TEXT("r.Water.SkipWaterInfoTextureRenderWhenWorldRenderingDisabled"), 0, bOverrideValues);
+
 	// Must come after the above cvars so that if one of those cvars is also specified by the Scalability level, then we restore to the value in the original scalability level
 	// not the value we cached in the Cinematic level (if applied).
 	if (bCinematicQualitySettings)
diff --git a/Engine/Plugins/MovieScene/MovieRenderPipeline/Source/MovieRenderPipelineCore/Public/MoviePipelineGameOverrideSetting.h b/Engine/Plugins/MovieScene/MovieRenderPipeline/Source/MovieRenderPipelineCore/Public/MoviePipelineGameOverrideSetting.h
--- a/Engine/Plugins/MovieScene/MovieRenderPipeline/Source/MovieRenderPipelineCore/Public/MoviePipelineGameOverrideSetting.h
+++ b/Engine/Plugins/MovieScene/MovieRenderPipeline/Source/MovieRenderPipelineCore/Public/MoviePipelineGameOverrideSetting.h
@@ -148,4 +148,5 @@ private:
 	float PreviousChaosImmPhysicsMinStepTime;
 	int32 PreviousSkipRedundantTransformUpdate;
 	int32 PreviousChaosClothUseTimeStepSmoothing;
+	int32 PreviousSkipWaterInfoTextureRenderWhenWorldRenderingDisabled;
 };
\ No newline at end of file
