Set r.Water.WaterInfo.RenderMethod to 1 (was 2 before) to work around artifacts that were observed with the recently introduced method 2. Season 28 shipped with method 1, so this is a known good path.

#rnx

diff --git a/Engine/Plugins/Experimental/Water/Source/Runtime/Private/WaterViewExtension.cpp b/Engine/Plugins/Experimental/Water/Source/Runtime/Private/WaterViewExtension.cpp
--- a/Engine/Plugins/Experimental/Water/Source/Runtime/Private/WaterViewExtension.cpp
+++ b/Engine/Plugins/Experimental/Water/Source/Runtime/Private/WaterViewExtension.cpp
@@ -28,7 +28,7 @@ static TAutoConsoleVariable<float> CVarLocalTessellationUpdateMargin(
 extern void OnCVarWaterInfoSceneProxiesValueChanged(IConsoleVariable*);
 static TAutoConsoleVariable<int32> CVarWaterInfoRenderMethod(
 	TEXT("r.Water.WaterInfo.RenderMethod"),
-	2,
+	1,
 	TEXT("0: SceneCaptures, 1: Custom, 2: CustomRenderPasses"),
 	FConsoleVariableDelegate::CreateStatic(OnCVarWaterInfoSceneProxiesValueChanged),
 	ECVF_Default | ECVF_RenderThreadSafe);
