Water: fixed changing river width/velocity/audio intensity from details panel not working.

diff --git a/Engine/Plugins/Experimental/Water/Source/Editor/Private/WaterSplineMetadataDetails.cpp b/Engine/Plugins/Experimental/Water/Source/Editor/Private/WaterSplineMetadataDetails.cpp
--- a/Engine/Plugins/Experimental/Water/Source/Editor/Private/WaterSplineMetadataDetails.cpp
+++ b/Engine/Plugins/Experimental/Water/Source/Editor/Private/WaterSplineMetadataDetails.cpp
@@ -161,7 +161,7 @@ void FWaterSplineMetadataDetails::OnSetRiverWidth(float NewValue, ETextCommit::T
 {
 	if (UWaterSplineMetadata* Metadata = GetMetadata())
 	{
-		if (CheckIfDifferent(*this, Metadata->Depth.Points, NewValue))
+		if (CheckIfDifferent(*this, Metadata->RiverWidth.Points, NewValue))
 		{
 			SetValues(*this, Metadata->RiverWidth.Points, NewValue, CommitInfo, LOCTEXT("SetSplineWaterWidth", "Set spline point river width"));
 		}
@@ -172,7 +172,7 @@ void FWaterSplineMetadataDetails::OnSetVelocity(float NewValue, ETextCommit::Typ
 {
 	if (UWaterSplineMetadata* Metadata = GetMetadata())
 	{
-		if (CheckIfDifferent(*this, Metadata->Depth.Points, NewValue))
+		if (CheckIfDifferent(*this, Metadata->WaterVelocityScalar.Points, NewValue))
 		{
 			SetValues(*this, Metadata->WaterVelocityScalar.Points, NewValue, CommitInfo, LOCTEXT("SetSplineWaterVelocity", "Set spline point water velocity"));
 		}
@@ -183,7 +183,7 @@ void FWaterSplineMetadataDetails::OnSetAudioIntensity(float NewValue, ETextCommi
 {
 	if (UWaterSplineMetadata* Metadata = GetMetadata())
 	{
-		if (CheckIfDifferent(*this, Metadata->Depth.Points, NewValue))
+		if (CheckIfDifferent(*this, Metadata->AudioIntensity.Points, NewValue))
 		{
 			SetValues(*this, Metadata->AudioIntensity.Points, NewValue, CommitInfo, LOCTEXT("SetSpline point audio intensity", "Set spline point audio intensity"));
 		}
