Water: Implemented the ability to tag arbitrary actors as "ground" by adding a WaterTerrainComponent.

diff --git a/Engine/Plugins/Experimental/Water/Source/Runtime/Private/WaterSubsystem.cpp b/Engine/Plugins/Experimental/Water/Source/Runtime/Private/WaterSubsystem.cpp
--- a/Engine/Plugins/Experimental/Water/Source/Runtime/Private/WaterSubsystem.cpp
+++ b/Engine/Plugins/Experimental/Water/Source/Runtime/Private/WaterSubsystem.cpp
@@ -16,6 +16,7 @@
 #include "WaterBodyIslandActor.h"
 #include "WaterModule.h"
 #include "WaterRuntimeSettings.h"
+#include "WaterTerrainComponent.h"
 #include "WaterUtils.h"
 #include "WaterViewExtension.h"
 #include "Algo/MaxElement.h"
@@ -30,6 +31,9 @@ extern UNREALED_API UEditorEngine* GEditor;
 #include "BuoyancyTypes.h"
 #endif // WITH_DITOR
 
+#include "LandscapeComponent.h"
+#include "LandscapeProxy.h"
+
 #include UE_INLINE_GENERATED_CPP_BY_NAME(WaterSubsystem)
 
 // ----------------------------------------------------------------------------------
@@ -302,7 +306,14 @@ void UWaterSubsystem::PostInitialize()
 		check(!OnHeightmapStreamedHandle.IsValid());
 		OnHeightmapStreamedHandle = LandscapeSubsystem->GetOnHeightmapStreamedDelegate().AddUObject(this, &UWaterSubsystem::OnHeightmapStreamed);
 	}
+
+	if (GEngine)
+	{
+		GEngine->OnActorMoved().AddUObject(this, &UWaterSubsystem::OnActorMoved);
+	}
 #endif // WITH_EDITOR
+
+	UActorComponent::MarkRenderStateDirtyEvent.AddUObject(this, &UWaterSubsystem::OnMarkRenderStateDirty);
 }
 
 void UWaterSubsystem::Deinitialize()
@@ -319,8 +330,15 @@ void UWaterSubsystem::Deinitialize()
 		}
 		OnHeightmapStreamedHandle.Reset();
 	}
+
+	if (GEngine)
+	{
+		GEngine->OnActorMoved().RemoveAll(this);
+	}
 #endif // WITH_EDITOR
 
+	UActorComponent::MarkRenderStateDirtyEvent.RemoveAll(this);
+
 	FConsoleVariableDelegate NullCallback;
 	CVarShallowWaterSimulationRenderTargetSize->SetOnChangedCallback(NullCallback);
 	CVarShallowWaterSim->SetOnChangedCallback(NullCallback);
@@ -368,7 +386,59 @@ void UWaterSubsystem::ApplyRuntimeSettings(const UWaterRuntimeSettings* Settings
 #endif // WITH_EDITOR
 }
 
+
+void UWaterSubsystem::OnMarkRenderStateDirty(UActorComponent& Component)
+{
+	TRACE_CPUPROFILER_EVENT_SCOPE(UWaterSubsystem::OnMarkRenderStateDirty);
+
+	const AActor* ComponentOwner = Component.GetOwner();
+
+	if (WaterTerrainActors.Find(ComponentOwner) != nullptr)
+	{
+		OnWaterTerrainActorChanged(ComponentOwner);
+	}
+}
+
+void UWaterSubsystem::OnWaterTerrainActorChanged(const AActor* TerrainActor)
+{
+	TArray<TWeakObjectPtr<UWaterTerrainComponent>, TInlineAllocator<4>> WaterTerrainComponentPtrs;
+	WaterTerrainActors.MultiFind(TerrainActor, WaterTerrainComponentPtrs);
+
+	check(TerrainActor != nullptr && TerrainActor->GetWorld() == GetWorld());
+
+	for (TWeakObjectPtr<UWaterTerrainComponent> WaterTerrainComponentPtr : WaterTerrainComponentPtrs)
+	{
+		 if (!WaterTerrainComponentPtr.IsValid())
+		 {
+			 continue;
+		 }
+
+		 if (const UWaterTerrainComponent* WaterTerrainComponent = WaterTerrainComponentPtr.Get())
+		 {
+			 const FBox2D TerrainBounds = WaterTerrainComponent->GetTerrainBounds();
+			 for (AWaterZone* WaterZone : TActorRange<AWaterZone>(GetWorld()))
+			 {
+				 if (WaterTerrainComponent->AffectsWaterZone(WaterZone))
+				 {
+					  WaterZone->MarkForRebuild(EWaterZoneRebuildFlags::UpdateWaterInfoTexture, TerrainBounds, TerrainActor);
+				 }
+			 }
+		 }
+		 
+	}
+}
+
 #if WITH_EDITOR
+void UWaterSubsystem::OnActorMoved(AActor* MovedActor)
+{
+	TRACE_CPUPROFILER_EVENT_SCOPE(UWaterSubsystem::OnActorMoved);
+
+	if (WaterTerrainActors.Find(MovedActor) != nullptr)
+	{
+		OnWaterTerrainActorChanged(MovedActor);
+	}
+}
+
 void UWaterSubsystem::OnHeightmapStreamed(const FOnHeightmapStreamedContext& InContext)
 {
 	UE_LOG(LogWater, Verbose, TEXT("UWaterSubsystem::OnHeightmapStreamed() -- Rebuilding Water Info Texture..."));
@@ -624,6 +694,36 @@ TSoftObjectPtr<AWaterZone> UWaterSubsystem::FindWaterZone(const FBox2D& Bounds,
 	return FindWaterZone(GetWorld(), Bounds, PreferredLevel);
 }
 
+void UWaterSubsystem::RegisterWaterTerrainComponent(UWaterTerrainComponent* InWaterTerrainComponent)
+{
+	check(InWaterTerrainComponent);
+	if (const AActor* TerrainActor = InWaterTerrainComponent->GetOwner())
+	{
+		 WaterTerrainActors.Add(TerrainActor,  InWaterTerrainComponent);
+	}
+}
+
+void UWaterSubsystem::UnregisterWaterTerrainComponent(UWaterTerrainComponent* InWaterTerrainComponent)
+{
+	check(InWaterTerrainComponent);
+	if (const AActor* TerrainActor = InWaterTerrainComponent->GetOwner())
+	{
+		WaterTerrainActors.RemoveSingle(TerrainActor, InWaterTerrainComponent);
+	}
+}
+
+void UWaterSubsystem::GetWaterTerrainComponents(TArray<UWaterTerrainComponent*>& OutWaterTerrainComponents) const
+{
+	OutWaterTerrainComponents.Empty(WaterTerrainActors.Num());
+	for (const TTuple<const AActor*, TWeakObjectPtr<UWaterTerrainComponent>>& Pair : WaterTerrainActors)
+	{
+		if (UWaterTerrainComponent* WaterTerrainComponent = Pair.Value.Get())
+		{
+			OutWaterTerrainComponents.Add(WaterTerrainComponent);
+		}
+	}
+}
+
 void UWaterSubsystem::NotifyWaterScalabilityChangedInternal(IConsoleVariable* CVar)
 {
 	OnWaterScalabilityChanged.Broadcast();
diff --git a/Engine/Plugins/Experimental/Water/Source/Runtime/Private/WaterTerrainComponent.cpp b/Engine/Plugins/Experimental/Water/Source/Runtime/Private/WaterTerrainComponent.cpp
new file mode 100644
--- /dev/null
+++ b/Engine/Plugins/Experimental/Water/Source/Runtime/Private/WaterTerrainComponent.cpp
@@ -0,0 +1,104 @@
+﻿// Copyright Epic Games, Inc. All Rights Reserved.
+
+#include "WaterTerrainComponent.h"
+
+#include "WaterSubsystem.h"
+#include "GameFramework/Actor.h"
+#include "Components/PrimitiveComponent.h"
+#include "EngineUtils.h"
+
+TArray<UPrimitiveComponent*> UWaterTerrainComponent::GetTerrainPrimitives() const
+{
+	TArray<UPrimitiveComponent*> Result;
+	if (const AActor* Owner = GetOwner())
+	{
+		Owner->GetComponents(Result, /* bNonColliding */ true);
+	}
+	return Result;
+}
+
+FBox2D UWaterTerrainComponent::GetTerrainBounds() const
+{
+	FBox2D Result(ForceInitToZero);
+	if (const AActor* Owner = GetOwner())
+	{
+		const FBox OwnerBounds = Owner->GetComponentsBoundingBox(/* bNonColliding */ true);
+		Result = FBox2D(FVector2D(OwnerBounds.Min), FVector2D(OwnerBounds.Max));
+	}
+	return Result;
+}
+
+bool UWaterTerrainComponent::AffectsWaterZone(AWaterZone* WaterZone) const
+{
+	// if we have a water zone override set, and the passed water zone is not it, we do not affect it
+	if (!WaterZoneOverride.IsNull() && WaterZoneOverride.Get() != WaterZone)
+	{
+		return false;
+	}
+
+	const FBox2D WaterZoneBounds = WaterZone->GetZoneBounds2D();
+
+	return GetTerrainBounds().Intersect(WaterZoneBounds);
+}
+
+void UWaterTerrainComponent::OnRegister()
+{
+	Super::OnRegister();
+
+	if (GetOwner() == nullptr)
+	{
+		return;
+	}
+	
+	if (UWaterSubsystem* WaterSubsystem = UWaterSubsystem::GetWaterSubsystem(GetWorld()))
+	{
+		WaterSubsystem->RegisterWaterTerrainComponent(this);
+	}
+
+	for (AWaterZone* WaterZone : TActorRange<AWaterZone>(GetWorld()))
+	{
+		if (AffectsWaterZone(WaterZone))
+		{
+			WaterZone->MarkForRebuild(EWaterZoneRebuildFlags::UpdateWaterInfoTexture, GetTerrainBounds(), GetOwner());
+		}
+	}
+}
+
+void UWaterTerrainComponent::OnUnregister()
+{
+	if (GetOwner() != nullptr)
+	{
+		for (AWaterZone* WaterZone : TActorRange<AWaterZone>(GetWorld()))
+		{
+			if (AffectsWaterZone(WaterZone))
+			{
+				 WaterZone->MarkForRebuild(EWaterZoneRebuildFlags::UpdateWaterInfoTexture, GetTerrainBounds(), GetOwner());
+			}
+		}
+	}
+
+	if (UWaterSubsystem* WaterSubsystem = UWaterSubsystem::GetWaterSubsystem(GetWorld()))
+	{
+		WaterSubsystem->UnregisterWaterTerrainComponent(this);
+	}
+	
+	Super::OnUnregister();
+}
+
+#if WITH_EDITOR
+void UWaterTerrainComponent::PostEditChangeProperty(FPropertyChangedEvent& PropertyChangedEvent)
+{
+	Super::PostEditChangeProperty(PropertyChangedEvent);
+
+	const FName PropertyName = PropertyChangedEvent.MemberProperty ? PropertyChangedEvent.MemberProperty->GetFName() : NAME_None;
+	if (PropertyName == GET_MEMBER_NAME_CHECKED(UWaterTerrainComponent, WaterZoneOverride))
+	{
+		// If the override pointer is set, mark all water zones to re-render since we can't know which ones we used to affect:
+		if (UWaterSubsystem* WaterSubsystem = UWaterSubsystem::GetWaterSubsystem(GetWorld()))
+		{
+			WaterSubsystem->MarkWaterZonesInRegionForRebuild(GetTerrainBounds(), EWaterZoneRebuildFlags::UpdateWaterInfoTexture, GetOwner());
+		}
+	}
+}
+#endif // WITH_EDITOR
+
diff --git a/Engine/Plugins/Experimental/Water/Source/Runtime/Private/WaterZoneActor.cpp b/Engine/Plugins/Experimental/Water/Source/Runtime/Private/WaterZoneActor.cpp
--- a/Engine/Plugins/Experimental/Water/Source/Runtime/Private/WaterZoneActor.cpp
+++ b/Engine/Plugins/Experimental/Water/Source/Runtime/Private/WaterZoneActor.cpp
@@ -15,6 +15,7 @@
 #include "Algo/AnyOf.h"
 #include "Engine/GameViewportClient.h"
 #include "WaterBodyInfoMeshComponent.h"
+#include "WaterTerrainComponent.h"
 
 #include UE_INLINE_GENERATED_CPP_BY_NAME(WaterZoneActor)
 
@@ -202,7 +203,7 @@ void AWaterZone::MarkForRebuild(EWaterZoneRebuildFlags Flags, const FBox2D& Upda
 		// Suppress water mesh updates which occur outside the bounds of the water quad tree.
 		if ((!UpdateRegion.bIsValid) || UpdateRegion.Intersect(WaterQuadTreeBounds))
 		{
-			UE_LOG(LogWater, Verbose, TEXT("AWaterZone::MarkForRebuild (UpdateWaterMesh) in region {%s} (triggered by %s)"), *UpdateRegion.ToString(), *GetNameSafe(DebugRequestingObject));
+			UE_LOG(LogWater, Verbose, TEXT("AWaterZone (%s) UpdateWaterMesh in region {%s} (triggered by %s)"), *GetNameSafe(this), *UpdateRegion.ToString(), *GetNameSafe(DebugRequestingObject));
 			WaterMesh->MarkWaterMeshGridDirty();
 			WaterMesh->MarkRenderStateDirty();
 		}
@@ -213,7 +214,7 @@ void AWaterZone::MarkForRebuild(EWaterZoneRebuildFlags Flags, const FBox2D& Upda
 		const FBox2D WaterInfoBounds2D(FVector2D(WaterInfoBounds.Min), FVector2D(WaterInfoBounds.Max));
 		if ((!UpdateRegion.bIsValid) || UpdateRegion.Intersect(WaterInfoBounds2D))
 		{
-			UE_LOG(LogWater, Verbose, TEXT("AWaterZone::MarkForRebuild (UpdateWaterInfoTexture) in region {%s} (triggered by %s)"), *UpdateRegion.ToString(), *GetNameSafe(DebugRequestingObject));
+			UE_LOG(LogWater, Verbose, TEXT("AWaterZone (%s) UpdateWaterInfoTexture in region {%s} (triggered by %s)"), *GetNameSafe(this), *UpdateRegion.ToString(), *GetNameSafe(DebugRequestingObject));
 			bNeedsWaterInfoRebuild = true;
 		}
 	}
@@ -371,6 +372,10 @@ void AWaterZone::PostEditChangeProperty(FPropertyChangedEvent& PropertyChangedEv
 	{
 		MarkForRebuild(EWaterZoneRebuildFlags::All, /* DebugRequestingObject = */ this);
 	}
+	else if (PropertyName == GET_MEMBER_NAME_CHECKED(AWaterZone, bAutoIncludeLandscapesAsTerrain))
+	{
+		MarkForRebuild(EWaterZoneRebuildFlags::UpdateWaterInfoTexture, /* DebugRequestingObject = */ this);
+	}
 }
 
 void AWaterZone::OnActorSelectionChanged(const TArray<UObject*>& NewSelection, bool bForceRefresh)
@@ -524,16 +529,38 @@ bool AWaterZone::UpdateWaterInfoTexture()
 		TArray<TWeakObjectPtr<UPrimitiveComponent>> GroundPrimitiveComponents;
 
 		const FBox WaterZoneBounds = GetZoneBounds();
-		for (ALandscapeProxy* LandscapeProxy : TActorRange<ALandscapeProxy>(World))
+		if (bAutoIncludeLandscapesAsTerrain)
 		{
-			const FBox LandscapeBox = LandscapeProxy->GetComponentsBoundingBox(/*bIncludeNonColliding = */ true);
-			// Only consider landscapes which this zone intersects with in XY and if the landscape volume is not zero sized
-			if (WaterZoneBounds.IntersectXY(LandscapeBox) && LandscapeBox.GetVolume() > 0.0)
+			for (ALandscapeProxy* LandscapeProxy : TActorRange<ALandscapeProxy>(World))
 			{
-				GroundZMin = FMath::Min(GroundZMin, LandscapeBox.Min.Z);
-				GroundZMax = FMath::Max(GroundZMax, LandscapeBox.Max.Z);
-				TInlineComponentArray<ULandscapeComponent*> LandscapeComponents(LandscapeProxy);
-				GroundPrimitiveComponents.Append(LandscapeComponents);
+				 const FBox LandscapeBox = LandscapeProxy->GetComponentsBoundingBox(/*bIncludeNonColliding = */ true);
+				 // Only consider landscapes which this zone intersects with in XY and if the landscape volume is not zero sized
+				 if (WaterZoneBounds.IntersectXY(LandscapeBox) && LandscapeBox.GetVolume() > 0.0)
+				 {
+					 GroundZMin = FMath::Min(GroundZMin, LandscapeBox.Min.Z);
+					 GroundZMax = FMath::Max(GroundZMax, LandscapeBox.Max.Z);
+					 TInlineComponentArray<ULandscapeComponent*> LandscapeComponents(LandscapeProxy);
+					 GroundPrimitiveComponents.Append(LandscapeComponents);
+				 }
+			}
+		}
+
+		UWaterSubsystem* WaterSubsystem = UWaterSubsystem::GetWaterSubsystem(World);
+		check(WaterSubsystem);
+
+		TArray<UWaterTerrainComponent*> WaterTerrainComponents;
+		WaterSubsystem->GetWaterTerrainComponents(WaterTerrainComponents);
+		
+		for (UWaterTerrainComponent* TerrainComponent : WaterTerrainComponents)
+		{
+			if (TerrainComponent->AffectsWaterZone(this))
+			{
+				TArray<UPrimitiveComponent*> TerrainPrimitives = TerrainComponent->GetTerrainPrimitives();
+
+				const FBox GroundActorBounds = TerrainComponent->GetOwner()->GetComponentsBoundingBox(true);
+				GroundZMin = FMath::Min(GroundZMin, GroundActorBounds.Min.Z);
+				GroundZMax = FMath::Max(GroundZMax, GroundActorBounds.Max.Z);
+				GroundPrimitiveComponents.Append(TerrainPrimitives);
 			}
 		}
 
@@ -550,9 +577,9 @@ bool AWaterZone::UpdateWaterInfoTexture()
 		{
 			if (UPrimitiveComponent* GroundPrimComp = GroundPrimCompPtr.Get())
 			{
-					TArray<UMaterialInterface*> TmpUsedMaterials;
-					GroundPrimComp->GetUsedMaterials(TmpUsedMaterials, false);
-					UsedMaterials.Append(TmpUsedMaterials);
+				TArray<UMaterialInterface*> TmpUsedMaterials;
+				GroundPrimComp->GetUsedMaterials(TmpUsedMaterials, false);
+				UsedMaterials.Append(TmpUsedMaterials);
 			}
 		}
 
@@ -622,7 +649,7 @@ bool AWaterZone::UpdateWaterInfoTexture()
 			WaterViewExtension->MarkWaterInfoTextureForRebuild(Context);
 		}
 
-		UE_LOG(LogWater, Verbose, TEXT("Queued Water Info texture update"));
+		UE_LOG(LogWater, Verbose, TEXT("Water Zone (%s) queued Water Info texture update"), *GetNameSafe(this));
 	}
 
 	return true;
diff --git a/Engine/Plugins/Experimental/Water/Source/Runtime/Public/WaterSubsystem.h b/Engine/Plugins/Experimental/Water/Source/Runtime/Public/WaterSubsystem.h
--- a/Engine/Plugins/Experimental/Water/Source/Runtime/Public/WaterSubsystem.h
+++ b/Engine/Plugins/Experimental/Water/Source/Runtime/Public/WaterSubsystem.h
@@ -7,9 +7,11 @@
 #include "Subsystems/WorldSubsystem.h"
 #include "Interfaces/Interface_PostProcessVolume.h"
 #include "WaterBodyManager.h"
+#include "WaterTerrainComponent.h"
 #include "WaterZoneActor.h"
 #include "WaterSubsystem.generated.h"
 
+class UWaterTerrainComponent;
 class UStaticMesh;
 
 DECLARE_STATS_GROUP(TEXT("Water"), STATGROUP_Water, STATCAT_Advanced);
@@ -171,7 +173,15 @@ public:
 	static TSoftObjectPtr<AWaterZone> FindWaterZone(const UWorld* World, const FBox2D& Bounds, const TSoftObjectPtr<const ULevel> PreferredLevel = {});
 	TSoftObjectPtr<AWaterZone> FindWaterZone(const FBox2D& Bounds, const TSoftObjectPtr<const ULevel> PreferredLevel = {}) const;
 
+	void RegisterWaterTerrainComponent(UWaterTerrainComponent* WaterTerrainComponent);
+	void UnregisterWaterTerrainComponent(UWaterTerrainComponent* WaterTerrainComponent);
+
+	/** Returns a list of all water terrain components registered to the water subsystem. Can be used instead of searching all actors to find them. */
+	void GetWaterTerrainComponents(TArray<UWaterTerrainComponent*>& OutWaterTerrainComponents) const;
+
 #if WITH_EDITOR
+	void OnActorMoved(AActor* MovedActor);
+
 	/** Little scope object to temporarily change the value of bAllowWaterSubsystemOnPreviewWorld */
 	struct WATER_API FScopedAllowWaterSubsystemOnPreviewWorld
 	{
@@ -198,6 +208,10 @@ private:
 	void AdjustUnderwaterWaterInfoQueryFlags(EWaterBodyQueryFlags& InOutFlags);
 	void ApplyRuntimeSettings(const UWaterRuntimeSettings* Settings, EPropertyChangeType::Type ChangeType);
 
+	void OnMarkRenderStateDirty(UActorComponent& Component);
+
+	void OnWaterTerrainActorChanged(const AActor* TerrainActor);
+
 	FWaterBodyManager& GetWaterBodyManagerInternal();
 
 #if !(UE_BUILD_SHIPPING || UE_BUILD_TEST)
@@ -248,6 +262,12 @@ private:
 
 	FWaterBodyManager WaterBodyManager;
 
+	/**
+	 * Keeps track of all actors that have WaterTerrainComponents so we can avoid extra work iterating
+	 * every actors components to find one when global events are triggered.
+	 */
+	TMultiMap<const AActor*, TWeakObjectPtr<UWaterTerrainComponent>> WaterTerrainActors;
+
 #if WITH_EDITOR
 	FDelegateHandle OnHeightmapStreamedHandle;
 	void OnHeightmapStreamed(const struct FOnHeightmapStreamedContext& InContext);
diff --git a/Engine/Plugins/Experimental/Water/Source/Runtime/Public/WaterTerrainComponent.h b/Engine/Plugins/Experimental/Water/Source/Runtime/Public/WaterTerrainComponent.h
new file mode 100644
--- /dev/null
+++ b/Engine/Plugins/Experimental/Water/Source/Runtime/Public/WaterTerrainComponent.h
@@ -0,0 +1,40 @@
+﻿// Copyright Epic Games, Inc. All Rights Reserved.
+
+#pragma once
+
+#include "CoreMinimal.h"
+#include "Components/ActorComponent.h"
+#include "WaterTerrainComponent.generated.h"
+
+class AWaterZone;
+
+/**
+ * Water Terrain Component can be attached to any actor with primitive components to allow them to render into a Water Info Texture as the terrain.
+ */
+UCLASS(meta=(BlueprintSpawnableComponent))
+class WATER_API UWaterTerrainComponent : public UActorComponent
+{
+	GENERATED_BODY()
+
+public:
+	/** Returns a list of all Primitive Components that should render as the terrain */
+	virtual TArray<UPrimitiveComponent*> GetTerrainPrimitives() const;
+
+	virtual FBox2D GetTerrainBounds() const;
+
+	virtual void OnRegister() override;
+	virtual void OnUnregister() override;
+
+	virtual bool AffectsWaterZone(AWaterZone* WaterZone) const;
+protected:
+
+#if WITH_EDITOR
+	virtual void PostEditChangeProperty( struct FPropertyChangedEvent& PropertyChangedEvent) override;
+#endif // WITH_EDITOR
+protected:
+	/** By default, the terrain component will be rendering into any overlapping water zone.
+	 * If the override is set, it will only render to that specific water zone. */
+	UPROPERTY(EditAnywhere, Category = Water)
+	TSoftObjectPtr<AWaterZone> WaterZoneOverride;
+};
+
diff --git a/Engine/Plugins/Experimental/Water/Source/Runtime/Public/WaterZoneActor.h b/Engine/Plugins/Experimental/Water/Source/Runtime/Public/WaterZoneActor.h
--- a/Engine/Plugins/Experimental/Water/Source/Runtime/Public/WaterZoneActor.h
+++ b/Engine/Plugins/Experimental/Water/Source/Runtime/Public/WaterZoneActor.h
@@ -186,6 +186,10 @@ private:
 	UPROPERTY(Category = LocalTessellation, EditAnywhere, meta = (EditCondition = "bEnableLocalOnlyTessellation"))
 	FVector LocalTessellationExtent;
 
+	/** When set to true, all landscape proxies that intersect with the bounds of this water zone will be included as ground actors regardless if they have WaterTerrain components. */
+	UPROPERTY(Category = Terrain, EditAnywhere)
+	bool bAutoIncludeLandscapesAsTerrain = true;
+
 	bool bNeedsWaterInfoRebuild = true;
 
 	FVector2f WaterHeightExtents;
