World Partition: make BuoyancyManager non-spatially loaded, as it returns invalid streaming bounds.

diff --git a/Engine/Plugins/Experimental/Water/Source/Runtime/Private/BuoyancyManager.cpp b/Engine/Plugins/Experimental/Water/Source/Runtime/Private/BuoyancyManager.cpp
--- a/Engine/Plugins/Experimental/Water/Source/Runtime/Private/BuoyancyManager.cpp
+++ b/Engine/Plugins/Experimental/Water/Source/Runtime/Private/BuoyancyManager.cpp
@@ -14,6 +14,10 @@ ABuoyancyManager::ABuoyancyManager(const FObjectInitializer& ObjectInitializer)
 	: Super(ObjectInitializer)
 {
 	Timestamp = 0;
+
+#if WITH_EDITORONLY_DATA
+	bIsSpatiallyLoaded = false;
+#endif
 }
 
 ABuoyancyManager* ABuoyancyManager::Get(const UObject* WorldContextObject)
diff --git a/Engine/Plugins/Experimental/Water/Source/Runtime/Public/BuoyancyManager.h b/Engine/Plugins/Experimental/Water/Source/Runtime/Public/BuoyancyManager.h
--- a/Engine/Plugins/Experimental/Water/Source/Runtime/Public/BuoyancyManager.h
+++ b/Engine/Plugins/Experimental/Water/Source/Runtime/Public/BuoyancyManager.h
@@ -54,6 +54,8 @@ public:
 #if WITH_EDITOR
 	// Prevent import/export of buoyancy manager actors since they should be transient and are always spawned when the map is loaded.
 	virtual bool ShouldImport(FStringView ActorPropString, bool IsMovingLevel) override { return false; }
+	
+	virtual bool CanChangeIsSpatiallyLoadedFlag() const override { return false; }
 #endif // WITH_EDITOR
 private:
 	void InitializeAsyncAux(UBuoyancyComponent* Component);
