Water: restored default water info rendering method to the CustomRenderPass method. Nanite bug has been resolved.

diff --git a/Engine/Plugins/Experimental/Water/Source/Runtime/Private/WaterViewExtension.cpp b/Engine/Plugins/Experimental/Water/Source/Runtime/Private/WaterViewExtension.cpp
--- a/Engine/Plugins/Experimental/Water/Source/Runtime/Private/WaterViewExtension.cpp
+++ b/Engine/Plugins/Experimental/Water/Source/Runtime/Private/WaterViewExtension.cpp
@@ -28,7 +28,7 @@ static TAutoConsoleVariable<float> CVarLocalTessellationUpdateMargin(
 extern void OnCVarWaterInfoSceneProxiesValueChanged(IConsoleVariable*);
 static TAutoConsoleVariable<int32> CVarWaterInfoRenderMethod(
 	TEXT("r.Water.WaterInfo.RenderMethod"),
-	1,
+	2,
 	TEXT("0: SceneCaptures, 1: Custom, 2: CustomRenderPasses"),
 	FConsoleVariableDelegate::CreateStatic(OnCVarWaterInfoSceneProxiesValueChanged),
 	ECVF_Default | ECVF_RenderThreadSafe);
@@ -362,16 +362,15 @@ void FWaterViewExtension::SetupView(FSceneViewFamily& InViewFamily, FSceneView&
 			WaterZoneInfo->UpdateBounds.Reset();
 		}
 
+		const UE::WaterInfo::FRenderingContext& Context(Pair.Value);
 		// Old method of rendering the water info texture; uses scene captures
 		if (WaterInfoRenderMethod == 0)
 		{
-			const UE::WaterInfo::FRenderingContext& Context(Pair.Value);
 			UE::WaterInfo::UpdateWaterInfoRendering(Scene, Context);
 		}
 		// Render the water info texture using custom render pass method
 		else if (WaterInfoRenderMethod == 2)
 		{
-			const UE::WaterInfo::FRenderingContext& Context(Pair.Value);
 			UE::WaterInfo::UpdateWaterInfoRendering_CustomRenderPass(Scene, Context);
 		}
 	}
