Added tooltip for Gerstner waves class
#skipci

diff --git a/Engine/Plugins/Experimental/Water/Source/Runtime/Public/GerstnerWaterWaves.h b/Engine/Plugins/Experimental/Water/Source/Runtime/Public/GerstnerWaterWaves.h
--- a/Engine/Plugins/Experimental/Water/Source/Runtime/Public/GerstnerWaterWaves.h
+++ b/Engine/Plugins/Experimental/Water/Source/Runtime/Public/GerstnerWaterWaves.h
@@ -162,6 +162,10 @@ public:
 	TArray<FGerstnerWaveOctave> Octaves;
 };
 
+/**
+	Waves implementation based off Gerstner waves. 
+	They provide a decent look, are relatively cheap to evaluate and have the advantage of being only dependent on the time variable, which makes them perfect for online simulations. 
+*/
 UCLASS(EditInlineNew, BlueprintType, MinimalAPI)
 class UGerstnerWaterWaves : public UWaterWaves
 {
