Moved CustomRenderPassInput list to the private scene implementation. Exposed a public virtual function to queue the CRPs for the next time the scene is rendered.

diff --git a/Engine/Plugins/Experimental/Water/Source/Runtime/Private/WaterInfoRendering.cpp b/Engine/Plugins/Experimental/Water/Source/Runtime/Private/WaterInfoRendering.cpp
--- a/Engine/Plugins/Experimental/Water/Source/Runtime/Private/WaterInfoRendering.cpp
+++ b/Engine/Plugins/Experimental/Water/Source/Runtime/Private/WaterInfoRendering.cpp
@@ -1050,6 +1050,7 @@ static TMap<uint32, int32> GatherLandscapeLODOverrides(const UWorld* World, cons
 
 void UpdateWaterInfoRendering_CustomRenderPass(
 	FSceneInterface* Scene,
+	const FSceneViewFamily& ViewFamily,
 	const WaterInfo::FRenderingContext& Context)
 {
 	TRACE_CPUPROFILER_EVENT_SCOPE(WaterInfo::UpdateWaterInfoRendering_CustomRenderPass);
@@ -1112,7 +1113,7 @@ void UpdateWaterInfoRendering_CustomRenderPass(
 		DepthPass->PerformRenderCapture(FCustomRenderPassBase::ERenderCaptureType::BeginCapture);
 	}
 	PassInput.CustomRenderPass = DepthPass;
-	Scene->CustomRenderPassRendererInputs.Add(PassInput);
+	Scene->AddCustomRenderPass(&ViewFamily, PassInput);
 
 	TSet<FPrimitiveComponentId> ComponentsToRenderInColorPass;
 	TSet<FPrimitiveComponentId> ComponentsToRenderInDilationPass;
@@ -1139,7 +1140,7 @@ void UpdateWaterInfoRendering_CustomRenderPass(
 	PassInput.ShowOnlyPrimitives = MoveTemp(ComponentsToRenderInColorPass);
 	FWaterInfoRenderingColorPass* ColorPass = new FWaterInfoRenderingColorPass(RenderTargetSize);
 	PassInput.CustomRenderPass = ColorPass;
-	Scene->CustomRenderPassRendererInputs.Add(PassInput);
+	Scene->AddCustomRenderPass(&ViewFamily, PassInput);
 
 	FUpdateWaterInfoParams Params;
 	Params.CaptureZ = ViewLocation.Z;
@@ -1160,7 +1161,7 @@ void UpdateWaterInfoRendering_CustomRenderPass(
 		DilationPass->PerformRenderCapture(FCustomRenderPassBase::ERenderCaptureType::EndCapture);
 	}
 	PassInput.CustomRenderPass = DilationPass;
-	Scene->CustomRenderPassRendererInputs.Add(PassInput);
+	Scene->AddCustomRenderPass(&ViewFamily, PassInput);
 }
 
 } // namespace WaterInfo
diff --git a/Engine/Plugins/Experimental/Water/Source/Runtime/Private/WaterInfoRendering.h b/Engine/Plugins/Experimental/Water/Source/Runtime/Private/WaterInfoRendering.h
--- a/Engine/Plugins/Experimental/Water/Source/Runtime/Private/WaterInfoRendering.h
+++ b/Engine/Plugins/Experimental/Water/Source/Runtime/Private/WaterInfoRendering.h
@@ -10,6 +10,7 @@ class FSceneInterface;
 class UTextureRenderTarget2D;
 class UPrimitiveComponent;
 class FSceneView;
+class FSceneViewFamily;
 
 namespace UE::WaterInfo
 {
@@ -33,6 +34,7 @@ void UpdateWaterInfoRendering2(
 
 void UpdateWaterInfoRendering_CustomRenderPass(
 	FSceneInterface* Scene,
+	const FSceneViewFamily& ViewFamily,
 	const FRenderingContext& Context);
 
 const FName& GetWaterInfoDepthPassName();
diff --git a/Engine/Plugins/Experimental/Water/Source/Runtime/Private/WaterViewExtension.cpp b/Engine/Plugins/Experimental/Water/Source/Runtime/Private/WaterViewExtension.cpp
--- a/Engine/Plugins/Experimental/Water/Source/Runtime/Private/WaterViewExtension.cpp
+++ b/Engine/Plugins/Experimental/Water/Source/Runtime/Private/WaterViewExtension.cpp
@@ -371,7 +371,7 @@ void FWaterViewExtension::SetupView(FSceneViewFamily& InViewFamily, FSceneView&
 		// Render the water info texture using custom render pass method
 		else if (WaterInfoRenderMethod == 2)
 		{
-			UE::WaterInfo::UpdateWaterInfoRendering_CustomRenderPass(Scene, Context);
+			UE::WaterInfo::UpdateWaterInfoRendering_CustomRenderPass(Scene, InViewFamily, Context);
 		}
 	}
 
diff --git a/Engine/Source/Runtime/Engine/Public/SceneInterface.h b/Engine/Source/Runtime/Engine/Public/SceneInterface.h
--- a/Engine/Source/Runtime/Engine/Public/SceneInterface.h
+++ b/Engine/Source/Runtime/Engine/Public/SceneInterface.h
@@ -68,6 +68,7 @@ struct FPersistentPrimitiveIndex;
 template<int32 MaxSHOrder> class TSHVectorRGB;
 using FSHVectorRGB3 = TSHVectorRGB<3>;
 struct FCustomPrimitiveData;
+class FSceneViewFamily;
 
 
 struct FPrimitiveSceneDesc;
@@ -706,7 +707,9 @@ public:
 
 		class FCustomRenderPassBase* CustomRenderPass = nullptr;
 	};
-	TArray<FCustomRenderPassRendererInput> CustomRenderPassRendererInputs;
+
+	/** Enqueues a new custom render pass to execute the next time this scene is rendered by ANY scene renderer. It will be immediately removed from the scene afterwards. */
+	virtual bool AddCustomRenderPass(const FSceneViewFamily* ViewFamily, const FCustomRenderPassRendererInput& CustomRenderPassInput) { return false; }
 
 protected:
 	virtual ~FSceneInterface() {}
diff --git a/Engine/Source/Runtime/Renderer/Private/RendererScene.cpp b/Engine/Source/Runtime/Renderer/Private/RendererScene.cpp
--- a/Engine/Source/Runtime/Renderer/Private/RendererScene.cpp
+++ b/Engine/Source/Runtime/Renderer/Private/RendererScene.cpp
@@ -7233,6 +7233,12 @@ void FScene::DebugRender(TArrayView<FViewInfo> Views)
 }
 #endif
 
+bool FScene::AddCustomRenderPass(const FSceneViewFamily* ViewFamily, const FCustomRenderPassRendererInput& CustomRenderPassInput)
+{
+	CustomRenderPassRendererInputs.Add(CustomRenderPassInput);
+	return true;
+}
+
 void FScene::UpdateCachedShadowState(const FScenePreUpdateChangeSet &ScenePreUpdateChangeSet, const FScenePostUpdateChangeSet &ScenePostUpdateChangeSet)
 {
 	for (const FLightSceneInfo* LightSceneInfo : DirectionalLights)
diff --git a/Engine/Source/Runtime/Renderer/Private/SceneCaptureRendering.cpp b/Engine/Source/Runtime/Renderer/Private/SceneCaptureRendering.cpp
--- a/Engine/Source/Runtime/Renderer/Private/SceneCaptureRendering.cpp
+++ b/Engine/Source/Runtime/Renderer/Private/SceneCaptureRendering.cpp
@@ -938,7 +938,8 @@ void FScene::UpdateSceneCaptureContents(USceneCaptureComponent2D* CaptureCompone
 			GetShowOnlyAndHiddenComponents(CaptureComponent, PassInput.HiddenPrimitives, PassInput.ShowOnlyPrimitives);
 
 			// Caching scene capture info to be passed to the scene renderer.
-			CustomRenderPassRendererInputs.Add(PassInput);
+			// #todo: We cannot (yet) guarantee for which ViewFamily this CRP will eventually be rendered since it will just execute the next time the scene is rendered by any FSceneRenderer. This seems quite problematic and could easily lead to unexpected behavior...
+			AddCustomRenderPass(nullptr, PassInput);
 			return;
 		}
 
diff --git a/Engine/Source/Runtime/Renderer/Private/ScenePrivate.h b/Engine/Source/Runtime/Renderer/Private/ScenePrivate.h
--- a/Engine/Source/Runtime/Renderer/Private/ScenePrivate.h
+++ b/Engine/Source/Runtime/Renderer/Private/ScenePrivate.h
@@ -3238,6 +3238,9 @@ public:
 	/** Collection of scene render extensions. */
 	FSceneExtensions SceneExtensions;
 
+	/** List of all the custom render passes that will run the next time the scene is rendered. */
+	TArray<FCustomRenderPassRendererInput> CustomRenderPassRendererInputs;
+
 	/** Initialization constructor. */
 	FScene(UWorld* InWorld, bool bInRequiresHitProxies,bool bInIsEditorScene, bool bCreateFXSystem, ERHIFeatureLevel::Type InFeatureLevel);
 
@@ -3750,6 +3753,8 @@ public:
 	template<typename TExtension>
 	const TExtension& GetExtension() const { return SceneExtensions.GetExtension<TExtension>(); }
 
+	virtual bool AddCustomRenderPass(const FSceneViewFamily* ViewFamily, const FCustomRenderPassRendererInput& CustomRenderPassInput);
+
 	FSceneCulling* SceneCulling = nullptr;
 
 	class FInstanceCullingOcclusionQueryRenderer* InstanceCullingOcclusionQueryRenderer = nullptr;
