Water: Fixed PostEditMove update on water body actors potentially registering a collision component while the world is being cleaned up.

diff --git a/Engine/Plugins/Experimental/Water/Source/Runtime/Private/WaterBodyActor.cpp b/Engine/Plugins/Experimental/Water/Source/Runtime/Private/WaterBodyActor.cpp
--- a/Engine/Plugins/Experimental/Water/Source/Runtime/Private/WaterBodyActor.cpp
+++ b/Engine/Plugins/Experimental/Water/Source/Runtime/Private/WaterBodyActor.cpp
@@ -149,6 +149,13 @@ void AWaterBody::PostEditMove(bool bFinished)
 		return;
 	}
 
+	// We want to avoid modifying the water body if the world is being cleaned up.
+	// World Partition unloads will call PostEditMove during the cleanup process and we don't want to unnecessarily make changes to the water body now.
+	if (GetWorld() && GetWorld()->IsBeingCleanedUp())
+	{
+		return;
+	}
+
 	WaterBodyComponent->FixupEditorTransform();
 
 	if (bFinished)
