Water: fixed a potential crash due to water info contexts rendering after a water zone has unregisted.

diff --git a/Engine/Plugins/Experimental/Water/Source/Runtime/Private/WaterViewExtension.cpp b/Engine/Plugins/Experimental/Water/Source/Runtime/Private/WaterViewExtension.cpp
--- a/Engine/Plugins/Experimental/Water/Source/Runtime/Private/WaterViewExtension.cpp
+++ b/Engine/Plugins/Experimental/Water/Source/Runtime/Private/WaterViewExtension.cpp
@@ -321,45 +321,46 @@ void FWaterViewExtension::SetupView(FSceneViewFamily& InViewFamily, FSceneView&
 		AWaterZone* WaterZone = Pair.Key;
 		check(WaterZone != nullptr);
 		FWaterZoneInfo* WaterZoneInfo = WaterZoneInfos.Find(WaterZone);
-		check(WaterZoneInfo != nullptr);
-
-		if (WaterZone->IsLocalOnlyTessellationEnabled())
+		if (ensureMsgf(WaterZoneInfo != nullptr, TEXT("We are trying to rendering a water info texture for a water zone that is not registered!")))
 		{
-			UWaterMeshComponent* WaterMesh= WaterZone->GetWaterMeshComponent();
-			check(WaterMesh);
+			if (WaterZone->IsLocalOnlyTessellationEnabled())
+			{
+				UWaterMeshComponent* WaterMesh= WaterZone->GetWaterMeshComponent();
+				check(WaterMesh);
+
+				const double TileSize = WaterMesh->GetTileSize();
 
-			const double TileSize = WaterMesh->GetTileSize();
+				const FVector WaterMeshCenter(ViewLocation.GridSnap(TileSize));
 
-			const FVector WaterMeshCenter(ViewLocation.GridSnap(TileSize));
+				const FVector2D WaterInfoHalfExtent(WaterZone->GetDynamicWaterInfoExtent() / 2.0);
+				const FBox2D WaterInfoBounds(FVector2D(WaterMeshCenter) - WaterInfoHalfExtent, FVector2D(WaterMeshCenter) + WaterInfoHalfExtent);
 
-			const FVector2D WaterInfoHalfExtent(WaterZone->GetDynamicWaterInfoExtent() / 2.0);
-			const FBox2D WaterInfoBounds(FVector2D(WaterMeshCenter) - WaterInfoHalfExtent, FVector2D(WaterMeshCenter) + WaterInfoHalfExtent);
+				WaterZone->SetLocalTessellationCenter(WaterMeshCenter);
 
-			WaterZone->SetLocalTessellationCenter(WaterMeshCenter);
+				// Trigger the next update when the camera is <UpdateMargin> units away from the border of the current window.
+				const FVector2D UpdateMargin(CVarLocalTessellationUpdateMargin.GetValueOnGameThread());
 
-			// Trigger the next update when the camera is <UpdateMargin> units away from the border of the current window.
-			const FVector2D UpdateMargin(CVarLocalTessellationUpdateMargin.GetValueOnGameThread());
+				// Keep a minimum of <1., 1.> bounds to avoid updating every frame if the update margin is larger than the zone.
+				const FVector2D UpdateExtents = FVector2D::Max(FVector2D(1., 1.), WaterInfoHalfExtent - UpdateMargin);
+				WaterZoneInfo->UpdateBounds.Emplace(FVector2D(WaterMeshCenter) - UpdateExtents, FVector2D(WaterMeshCenter) + UpdateExtents);
 
-			// Keep a minimum of <1., 1.> bounds to avoid updating every frame if the update margin is larger than the zone.
-			const FVector2D UpdateExtents = FVector2D::Max(FVector2D(1., 1.), WaterInfoHalfExtent - UpdateMargin);
-			WaterZoneInfo->UpdateBounds.Emplace(FVector2D(WaterMeshCenter) - UpdateExtents, FVector2D(WaterMeshCenter) + UpdateExtents);
+				const FVector2D WaterQuadTreeHalfExtent = WaterMesh->GetExtentInTiles() * TileSize;
+				const FVector2D WaterQuadTreeCenter = WaterMesh->GetDynamicWaterMeshCenter();
+				const FBox2D WaterQuadTreeBounds(WaterQuadTreeCenter - WaterQuadTreeHalfExtent, WaterQuadTreeCenter + WaterQuadTreeHalfExtent);
 
-			const FVector2D WaterQuadTreeHalfExtent = WaterMesh->GetExtentInTiles() * TileSize;
-			const FVector2D WaterQuadTreeCenter = WaterMesh->GetDynamicWaterMeshCenter();
-			const FBox2D WaterQuadTreeBounds(WaterQuadTreeCenter - WaterQuadTreeHalfExtent, WaterQuadTreeCenter + WaterQuadTreeHalfExtent);
+				// If the new water info bounds would be outside the water mesh, recenter and rebuild the water mesh
+				if (!WaterQuadTreeBounds.IsInside(WaterInfoBounds.ExpandBy(WaterInfoBounds.GetExtent())))
+				{
+					WaterMesh->SetDynamicWaterMeshCenter(FVector2D(WaterMeshCenter));
+				}
 
-			// If the new water info bounds would be outside the water mesh, recenter and rebuild the water mesh
-			if (!WaterQuadTreeBounds.IsInside(WaterInfoBounds.ExpandBy(WaterInfoBounds.GetExtent())))
+				// Mark GPU data dirty since we have a new WaterArea parameter and need to push this to water bodies.
+				MarkGPUDataDirty();
+			}
+			else
 			{
-				WaterMesh->SetDynamicWaterMeshCenter(FVector2D(WaterMeshCenter));
+				WaterZoneInfo->UpdateBounds.Reset();
 			}
-
-			// Mark GPU data dirty since we have a new WaterArea parameter and need to push this to water bodies.
-			MarkGPUDataDirty();
-		}
-		else
-		{
-			WaterZoneInfo->UpdateBounds.Reset();
 		}
 
 		const UE::WaterInfo::FRenderingContext& Context(Pair.Value);
@@ -451,4 +452,6 @@ void FWaterViewExtension::AddWaterZone(AWaterZone* InWaterZone)
 void FWaterViewExtension::RemoveWaterZone(AWaterZone* InWaterZone)
 {
 	WaterZoneInfos.FindAndRemoveChecked(InWaterZone);
+	// Remove any WaterInfoContexts that might be trying to rendering for this water zone this frame. They are no longer valid.
+	WaterInfoContextsToRender.Remove(InWaterZone);
 }
