Water: implemented actionable message when water body meshes are outdated.

diff --git a/Engine/Plugins/Experimental/Water/Source/Editor/Private/WaterEditorSubsystem.cpp b/Engine/Plugins/Experimental/Water/Source/Editor/Private/WaterEditorSubsystem.cpp
--- a/Engine/Plugins/Experimental/Water/Source/Editor/Private/WaterEditorSubsystem.cpp
+++ b/Engine/Plugins/Experimental/Water/Source/Editor/Private/WaterEditorSubsystem.cpp
@@ -1,17 +1,38 @@
 // Copyright Epic Games, Inc. All Rights Reserved.
 
 #include "WaterEditorSubsystem.h"
+
+#include "ActionableMessageSubsystem.h"
 #include "Engine/Texture2D.h"
 #include "UObject/ConstructorHelpers.h"
 #include "WaterEditorSettings.h"
 #include "Materials/MaterialParameterCollection.h"
 #include "WaterModule.h"
 #include "Modules/ModuleManager.h"
+#include "Subsystems/UnrealEditorSubsystem.h"
+#include "Editor.h"
+#include "LevelEditor.h"
+#include "WaterSubsystem.h"
+#include "Algo/Count.h"
 
 #include UE_INLINE_GENERATED_CPP_BY_NAME(WaterEditorSubsystem)
 
 #define LOCTEXT_NAMESPACE "WaterEditorSubsystem"
 
+namespace UE::Water
+{
+	static const FName WaterBodiesNeedSaveKey(TEXT("WaterBodiesNeedResave_MessageKey"));
+
+	static void MarkOutdatedPackagesDirty()
+	{
+		const IWaterModuleInterface& WaterModule = FModuleManager::GetModuleChecked<IWaterModuleInterface>("Water");
+		if (IWaterEditorServices* WaterEditorServices = WaterModule.GetWaterEditorServices())
+		{
+			WaterEditorServices->DirtyAllModifiedPackages();
+		}
+	}
+}
+
 UWaterEditorSubsystem::UWaterEditorSubsystem()
 {
 	struct FConstructorStatics
@@ -39,10 +60,18 @@ void UWaterEditorSubsystem::Initialize(FSubsystemCollectionBase& Collection)
 
 	IWaterModuleInterface& WaterModule = FModuleManager::GetModuleChecked<IWaterModuleInterface>("Water");
 	WaterModule.SetWaterEditorServices(this);
+
+	UPackage::PackageDirtyStateChangedEvent.AddUObject(this, &UWaterEditorSubsystem::OnPackageDirtied);
+	FLevelEditorModule& LevelEditor = FModuleManager::LoadModuleChecked<FLevelEditorModule>("LevelEditor");
+	LevelEditor.OnMapChanged().AddUObject(this, &UWaterEditorSubsystem::OnMapChanged);
 }
 
 void UWaterEditorSubsystem::Deinitialize()
 {
+	FLevelEditorModule& LevelEditor = FModuleManager::LoadModuleChecked<FLevelEditorModule>("LevelEditor");
+	LevelEditor.OnMapChanged().RemoveAll(this);
+	UPackage::PackageDirtyStateChangedEvent.RemoveAll(this);
+
 	IWaterModuleInterface& WaterModule = FModuleManager::GetModuleChecked<IWaterModuleInterface>("Water");
 	if (IWaterEditorServices* WaterEditorServices = WaterModule.GetWaterEditorServices())
 	{
@@ -80,5 +109,126 @@ UTexture2D* UWaterEditorSubsystem::GetWaterActorSprite(UClass* InClass) const
 	return DefaultWaterActorSprite;
 }
 
+UWorld* UWaterEditorSubsystem::GetEditorWorld() const
+{
+	if (UUnrealEditorSubsystem* UnrealEditorSubsystem = GEditor->GetEditorSubsystem<UUnrealEditorSubsystem>())
+	{
+		return UnrealEditorSubsystem->GetEditorWorld();
+	}
+	return GetWorld();
+}
+
+void UWaterEditorSubsystem::UpdateModifiedPackagesMessage()
+{
+	if (UWorld* EditorWorld = GetEditorWorld())
+	{
+		if (UActionableMessageSubsystem* ActionableMessageSubsystem = EditorWorld->GetSubsystem<UActionableMessageSubsystem>())
+		{
+			if (PackagesNeedingDirtying.Num() > 0)
+			{
+				FActionableMessage Message;
+				Message.Message = LOCTEXT("WaterBodiesNeedSave_Message", "Some water bodies are not up to date");
+				Message.Tooltip = LOCTEXT("WaterBodiesNeedSave_Tooltip", "Some water bodies have some derived data that is being rebuilt on load. This is causing editor loading to slow down and can be resolved by simply resaving the water bodies.");
+				Message.ActionMessage = LOCTEXT("WaterBodiesNeedSave_ActionMsg", "Mark Dirty");
+				Message.ActionCallback = UE::Water::MarkOutdatedPackagesDirty;
+				ActionableMessageSubsystem->SetActionableMessage(UE::Water::WaterBodiesNeedSaveKey, Message);
+			}
+			else
+			{
+				ActionableMessageSubsystem->ClearActionableMessage(UE::Water::WaterBodiesNeedSaveKey);
+			}
+		}
+	}
+}
+
+void UWaterEditorSubsystem::OnPackageDirtied(UPackage* Package)
+{
+	if (!bSuppressOnDirtyEvents)
+	{
+		if (Package && Package->IsDirty())
+		{
+			PackagesNeedingDirtying.Remove(Package);
+		}
+		UpdateModifiedPackagesMessage();
+	}
+}
+
+void UWaterEditorSubsystem::OnMapChanged(UWorld* InWorld, EMapChangeType InChangeType)
+{
+	// Stop tracking packages for maps that are no longer loaded.
+	if (InChangeType == EMapChangeType::TearDownWorld)
+	{
+		ClearModifiedPackages();
+	}
+}
+
+bool UWaterEditorSubsystem::TryMarkPackageAsModified(UPackage* InPackage)
+{
+	// Don't need to mark it as dirty while transacting since it will be marked dirty later anyways:
+	if (GIsTransacting)
+	{
+		return false;
+	}
+
+	// Don't need to do anything if the package is already dirty:
+	if (InPackage->IsDirty())
+	{
+		return false;
+	}
+
+	// We don't care about temp packages:
+	if (FPackageName::IsTempPackage(InPackage->GetName()))
+	{
+		return false;
+	}
+
+	PackagesNeedingDirtying.Add(InPackage);
+
+	UpdateModifiedPackagesMessage();
+
+	return true;
+}
+
+bool UWaterEditorSubsystem::HasAnyModifiedPackages() const
+{
+	return IntCastChecked<int32>(Algo::CountIf(PackagesNeedingDirtying, [](const TWeakObjectPtr<UPackage>& InWeakPackagePtr) { return InWeakPackagePtr.IsValid(); })) > 0;
+}
+
+void UWaterEditorSubsystem::ForEachModifiedPackage(const TFunctionRef<bool(UPackage*)>& Func) const
+{
+	for (TWeakObjectPtr<UPackage> WeakPackage : PackagesNeedingDirtying)
+	{
+		if (WeakPackage.IsValid())
+		{
+			if (!Func(WeakPackage.Get()))
+			{
+				break;
+			}
+		}
+	}
+}
+
+void UWaterEditorSubsystem::ClearModifiedPackages()
+{
+	PackagesNeedingDirtying.Empty();
+	UpdateModifiedPackagesMessage();
+}
+
+void UWaterEditorSubsystem::DirtyAllModifiedPackages()
+{
+	TGuardValue<bool> SuppressOnDirtyEvents(bSuppressOnDirtyEvents, true);
+
+	ForEachModifiedPackage([](UPackage* Package)
+	{
+		if (!Package->IsDirty())
+		{
+			Package->SetDirtyFlag(true);
+		}
+		return true;
+	});
+
+	ClearModifiedPackages();
+}
+
 #undef LOCTEXT_NAMESPACE
 
diff --git a/Engine/Plugins/Experimental/Water/Source/Editor/Public/WaterEditorSubsystem.h b/Engine/Plugins/Experimental/Water/Source/Editor/Public/WaterEditorSubsystem.h
--- a/Engine/Plugins/Experimental/Water/Source/Editor/Public/WaterEditorSubsystem.h
+++ b/Engine/Plugins/Experimental/Water/Source/Editor/Public/WaterEditorSubsystem.h
@@ -4,6 +4,7 @@
 
 #include "EditorSubsystem.h"
 #include "WaterEditorServices.h"
+#include "UnrealEdMisc.h"
 #include "WaterEditorSubsystem.generated.h"
 
 class AWaterZone;
@@ -25,11 +26,26 @@ public:
 
 	UMaterialParameterCollection* GetLandscapeMaterialParameterCollection() const {	return LandscapeMaterialParameterCollection; }
 
-	// ~ IWaterEditorServices interface
+	//~ Begin IWaterEditorServices interface
 	virtual void RegisterWaterActorSprite(UClass* InClass, UTexture2D* Texture) override;
 	virtual UTexture2D* GetWaterActorSprite(UClass* InClass) const override;
 	virtual UTexture2D* GetErrorSprite() const override { return ErrorSprite; }
+
+	virtual bool TryMarkPackageAsModified(UPackage* ModifiedPackage) override;
+	virtual bool HasAnyModifiedPackages() const override;
+	virtual void ForEachModifiedPackage(const TFunctionRef<bool(UPackage*)>& Func) const override;
+	virtual void ClearModifiedPackages() override;
+	virtual void DirtyAllModifiedPackages() override;
+	//~ End IWaterEditorServices interface
+
 private:
+	UWorld* GetEditorWorld() const;
+
+	void UpdateModifiedPackagesMessage();
+
+	void OnPackageDirtied(UPackage* Package);
+	void OnMapChanged(UWorld* InWorld, EMapChangeType InChangeType);
+
 	UPROPERTY()
 	TObjectPtr<UMaterialParameterCollection> LandscapeMaterialParameterCollection;
 
@@ -39,6 +55,11 @@ private:
 	TObjectPtr<UTexture2D> DefaultWaterActorSprite;
 	UPROPERTY()
 	TObjectPtr<UTexture2D> ErrorSprite;
+
+	/** Set of water body packages which have been silently modified but not dirtied. */
+	TSet<TWeakObjectPtr<UPackage>, TWeakObjectPtrSetKeyFuncs<TWeakObjectPtr<UPackage>>> PackagesNeedingDirtying;
+
+	bool bSuppressOnDirtyEvents = false;
 };
 
 #if UE_ENABLE_INCLUDE_ORDER_DEPRECATED_IN_5_2
diff --git a/Engine/Plugins/Experimental/Water/Source/Editor/WaterEditor.Build.cs b/Engine/Plugins/Experimental/Water/Source/Editor/WaterEditor.Build.cs
--- a/Engine/Plugins/Experimental/Water/Source/Editor/WaterEditor.Build.cs
+++ b/Engine/Plugins/Experimental/Water/Source/Editor/WaterEditor.Build.cs
@@ -31,6 +31,7 @@ public class WaterEditor : ModuleRules
 				"DeveloperSettings",
 				"AdvancedPreviewScene",
 				"PlacementMode",
+				"ActionableMessage",
 			});
 
 		PublicDependencyModuleNames.AddRange(
diff --git a/Engine/Plugins/Experimental/Water/Source/Runtime/Private/WaterBodyComponent.cpp b/Engine/Plugins/Experimental/Water/Source/Runtime/Private/WaterBodyComponent.cpp
--- a/Engine/Plugins/Experimental/Water/Source/Runtime/Private/WaterBodyComponent.cpp
+++ b/Engine/Plugins/Experimental/Water/Source/Runtime/Private/WaterBodyComponent.cpp
@@ -1507,8 +1507,16 @@ void UWaterBodyComponent::OnPostRegisterAllComponents()
 	UWaterBodyInfoMeshComponent* WaterInfoMeshComponent = GetWaterInfoMeshComponent();
 	const bool bHasConservativeRasterMesh = IsValid(WaterInfoMeshComponent) && WaterInfoMeshComponent->bIsConservativeRasterCompatible;
 	const bool bShouldHaveConservativeRastermesh = CVarWaterBodyBuildConservativeRasterizationMesh.GetValueOnGameThread() != 0;
-	if ((bHasConservativeRasterMesh != bShouldHaveConservativeRastermesh) || GetLinkerCustomVersion(FFortniteMainBranchObjectVersion::GUID) < FFortniteMainBranchObjectVersion::WaterBodyStaticMeshFixup)
+	if (AffectsWaterInfo() && ((bHasConservativeRasterMesh != bShouldHaveConservativeRastermesh) || GetLinkerCustomVersion(FFortniteMainBranchObjectVersion::GUID) < FFortniteMainBranchObjectVersion::WaterBodyStaticMeshFixup))
 	{
+		const IWaterModuleInterface& WaterModule = FModuleManager::GetModuleChecked<IWaterModuleInterface>("Water");
+		if (IWaterEditorServices* WaterEditorServices = WaterModule.GetWaterEditorServices())
+		{
+			if (GetWorld() && GetWorld()->WorldType == EWorldType::Editor)
+			{
+				WaterEditorServices->TryMarkPackageAsModified(GetPackage());
+			}
+		}
 		UpdateWaterBodyRenderData();
 	}
 
diff --git a/Engine/Plugins/Experimental/Water/Source/Runtime/Public/WaterEditorServices.h b/Engine/Plugins/Experimental/Water/Source/Runtime/Public/WaterEditorServices.h
--- a/Engine/Plugins/Experimental/Water/Source/Runtime/Public/WaterEditorServices.h
+++ b/Engine/Plugins/Experimental/Water/Source/Runtime/Public/WaterEditorServices.h
@@ -16,6 +16,16 @@ public:
 	virtual void RegisterWaterActorSprite(UClass* InClass, UTexture2D* Texture) = 0;
 	virtual UTexture2D* GetWaterActorSprite(UClass* InClass) const = 0;
 	virtual UTexture2D* GetErrorSprite() const = 0;
+
+	/**
+	 * Mark a package as modified and needs to be dirtied manually by the user. This will be picked up by the editor module and notify the user that they need to resave the package manually.
+	 * This can happen when a package has outdated derived data (such as a water body's saved WaterInfoMesh) is being rebuilt on load due to some external change that affects the derived data.
+	 */
+	virtual bool TryMarkPackageAsModified(UPackage* InPackage) = 0;
+	virtual bool HasAnyModifiedPackages() const = 0;
+	virtual void ForEachModifiedPackage(const TFunctionRef<bool(UPackage*)>& Func) const = 0;
+	virtual void ClearModifiedPackages() = 0;
+	virtual void DirtyAllModifiedPackages() = 0;
 };
 #endif // WITH_EDITORONLY_DATA
 
