Water: Only update ModifiedPackages message in MarkPackageDirty if we actually made a change to the list of modified packages.
Additionally, changed native call to UUnrealEditorSubsystem::GetEditorWorld to use the same implementation without accessing the subsystem which is only designed for use in scripts

diff --git a/Engine/Plugins/Experimental/Water/Source/Editor/Private/WaterEditorSubsystem.cpp b/Engine/Plugins/Experimental/Water/Source/Editor/Private/WaterEditorSubsystem.cpp
--- a/Engine/Plugins/Experimental/Water/Source/Editor/Private/WaterEditorSubsystem.cpp
+++ b/Engine/Plugins/Experimental/Water/Source/Editor/Private/WaterEditorSubsystem.cpp
@@ -9,7 +9,6 @@
 #include "Materials/MaterialParameterCollection.h"
 #include "WaterModule.h"
 #include "Modules/ModuleManager.h"
-#include "Subsystems/UnrealEditorSubsystem.h"
 #include "Editor.h"
 #include "LevelEditor.h"
 #include "WaterSubsystem.h"
@@ -111,11 +110,7 @@ UTexture2D* UWaterEditorSubsystem::GetWaterActorSprite(UClass* InClass) const
 
 UWorld* UWaterEditorSubsystem::GetEditorWorld() const
 {
-	if (UUnrealEditorSubsystem* UnrealEditorSubsystem = GEditor->GetEditorSubsystem<UUnrealEditorSubsystem>())
-	{
-		return UnrealEditorSubsystem->GetEditorWorld();
-	}
-	return GetWorld();
+	return GEditor ? GEditor->GetEditorWorldContext(false).World() : nullptr;
 }
 
 void UWaterEditorSubsystem::UpdateModifiedPackagesMessage()
@@ -147,9 +142,12 @@ void UWaterEditorSubsystem::OnPackageDirtied(UPackage* Package)
 	{
 		if (Package && Package->IsDirty())
 		{
-			PackagesNeedingDirtying.Remove(Package);
+			// Only update the modified packages messsage if we actually made a change:
+			if (PackagesNeedingDirtying.Remove(Package) > 0)
+			{
+				UpdateModifiedPackagesMessage();
+			}
 		}
-		UpdateModifiedPackagesMessage();
 	}
 }
 
