Water: fixed GetWaterTimeSeconds so that it can be called outside of the game thread.

diff --git a/Engine/Plugins/Experimental/Water/Source/Runtime/Private/WaterSubsystem.cpp b/Engine/Plugins/Experimental/Water/Source/Runtime/Private/WaterSubsystem.cpp
--- a/Engine/Plugins/Experimental/Water/Source/Runtime/Private/WaterSubsystem.cpp
+++ b/Engine/Plugins/Experimental/Water/Source/Runtime/Private/WaterSubsystem.cpp
@@ -62,9 +62,10 @@ static FAutoConsoleVariableRef CVarFreezeWaves(
 	ECVF_Cheat
 );
 
-static TAutoConsoleVariable<float> CVarOverrideWavesTime(
+static float OverrideWavesTime = -1.f;
+static FAutoConsoleVariableRef CVarOverrideWavesTime(
 	TEXT("r.Water.OverrideWavesTime"),
-	-1.0f,
+	OverrideWavesTime,
 	TEXT("Forces the time used for waves if >= 0.0"),
 	ECVF_Cheat
 );
@@ -486,10 +487,9 @@ bool UWaterSubsystem::IsWaterRenderingEnabled() const
 
 float UWaterSubsystem::GetWaterTimeSeconds() const
 {
-	float ForceWavesTimeValue = CVarOverrideWavesTime.GetValueOnGameThread();
-	if (ForceWavesTimeValue >= 0.0f)
+	if (OverrideWavesTime >= 0.0f)
 	{
-		return ForceWavesTimeValue;
+		return OverrideWavesTime;
 	}
 
 	if (UWorld* World = GetWorld())
