Water: allow water info texture renders to occur when not all materials have finished compiling. We will render every frame until the WIT is rendered with a full set of compiled shaders.

diff --git a/Engine/Plugins/Experimental/Water/Source/Runtime/Private/WaterZoneActor.cpp b/Engine/Plugins/Experimental/Water/Source/Runtime/Private/WaterZoneActor.cpp
--- a/Engine/Plugins/Experimental/Water/Source/Runtime/Private/WaterZoneActor.cpp
+++ b/Engine/Plugins/Experimental/Water/Source/Runtime/Private/WaterZoneActor.cpp
@@ -602,11 +602,6 @@ bool AWaterZone::UpdateWaterInfoTexture()
 				}
 			}
 		}
-
-		if (bHasIncompleteShaderMaps)
-		{
-			return false;
-		}
 #endif // WITH_EDITOR
 
 		// The render path for rendering the water info texture without scene captures is executed within the scene renderer.
@@ -650,6 +645,14 @@ bool AWaterZone::UpdateWaterInfoTexture()
 		}
 
 		UE_LOG(LogWater, Verbose, TEXT("Water Zone (%s) queued Water Info texture update"), *GetNameSafe(this));
+
+#if WITH_EDITOR
+		// Ensure that we try to render again next frame if the materials haven't all been compiled.
+		if (bHasIncompleteShaderMaps)
+		{
+			return false;
+		}
+#endif // WITH_EDITOR
 	}
 
 	return true;
