Moved water quadtree creation to WaterMeshSceneProxy using a new FWaterQuadTreeBuilder class which can be used to construct quadtrees at different positions using a certain set of fixed parameters. This is the first step towards managing per-view quadtrees in the scene proxy. There should be no difference in behavior with/without this CL; this is merely a change in architecture to enable future improvements. There's a new cvar to toggle between the old and now path: r.Water.WaterMesh.BuildQuadTreeInSceneProxy. In order for that cvar to take effect, the water mesh needs to be rebuild. This can easily be achieved with r.Water.WaterMesh.ForceRebuildMeshPerFrame 1.

#rnx
#tests Local testing on PC in editor (QA-Waterworks, helios_terrain and creating new water bodies in a new empty level; with and without far distance material/mesh) as well as in a cooked instance of QAGame running QA-Waterworks

diff --git a/Engine/Plugins/Experimental/Water/Source/Runtime/Private/WaterMeshComponent.cpp b/Engine/Plugins/Experimental/Water/Source/Runtime/Private/WaterMeshComponent.cpp
--- a/Engine/Plugins/Experimental/Water/Source/Runtime/Private/WaterMeshComponent.cpp
+++ b/Engine/Plugins/Experimental/Water/Source/Runtime/Private/WaterMeshComponent.cpp
@@ -70,6 +70,13 @@ TAutoConsoleVariable<int32> CVarWaterMeshMIDDeduplication(
 	ECVF_RenderThreadSafe
 );
 
+TAutoConsoleVariable<int32> CVarWaterMeshBuildQuadTreeInSceneProxy(
+	TEXT("r.Water.WaterMesh.BuildQuadTreeInSceneProxy"),
+	1,
+	TEXT("Moves water quadtree construction into the water mesh scene proxy."),
+	ECVF_RenderThreadSafe
+);
+
 extern TAutoConsoleVariable<float> CVarWaterSplineResampleMaxDistance;
 
 
@@ -192,6 +199,12 @@ FMaterialRelevance UWaterMeshComponent::GetWaterMaterialRelevance(ERHIFeatureLev
 	return Result;
 }
 
+bool UWaterMeshComponent::ShouldBuildQuadTreeInSceneProxy() const
+{
+	const bool bShouldBuildQuadTreeInSceneProxy = CVarWaterMeshBuildQuadTreeInSceneProxy.GetValueOnGameThread() != 0;
+	return bShouldBuildQuadTreeInSceneProxy;
+}
+
 void UWaterMeshComponent::SetDynamicWaterMeshCenter(const FVector2D& NewCenter)
 {
 	if (!DynamicWaterMeshCenter.Equals(NewCenter))
@@ -229,14 +242,30 @@ FIntPoint UWaterMeshComponent::GetExtentInTiles() const
 
 FBoxSphereBounds UWaterMeshComponent::CalcBounds(const FTransform& LocalToWorld) const
 {
-	// Always return valid bounds (tree is initialized with invalid bounds and if nothing is inserted, the tree bounds will stay invalid)
-	FBox NewBounds = WaterQuadTree.GetBounds();
-
-	if (NewBounds.Min.Z >= NewBounds.Max.Z)
+	FBox NewBounds = FBox();
+	if (ShouldBuildQuadTreeInSceneProxy())
+	{
+		// With the water quadtree building moved into the scene proxy, the scene proxy is free to create quadtree(s) anywhere within the water zone.
+		// To ensure that the scene proxy isn't frustum culled, we need to be conservative and use the water zone bounds instead. The quadtree does
+		// fine grained frustum culling internally anyways, so this shouldn't be a performance issue.
+		AWaterZone* WaterZone = Cast<AWaterZone>(GetOwner());
+		if (WaterZone)
+		{
+			NewBounds = WaterZone->GetZoneBounds();
+		}
+	}
+	else
 	{
-		NewBounds.Min.Z = 0.0f;
-		NewBounds.Max.Z = 100.0f;
+		// Always return valid bounds (tree is initialized with invalid bounds and if nothing is inserted, the tree bounds will stay invalid)
+		NewBounds = WaterQuadTree.GetBounds();
+
+		if (NewBounds.Min.Z >= NewBounds.Max.Z)
+		{
+			NewBounds.Min.Z = 0.0f;
+			NewBounds.Max.Z = 100.0f;
+		}
 	}
+	
 	// Add the far distance to the bounds if it's valid
 	if (FarDistanceMaterial)
 	{
@@ -256,6 +285,291 @@ void UWaterMeshComponent::RebuildWaterMesh(float InTileSize, const FIntPoint& In
 
 	AWaterZone* WaterZone = CastChecked<AWaterZone>(GetOwner());
 
+	bool bAnyWaterMeshesNotReady = false;
+	const bool bIsGPUQuadTree = CVarWaterMeshGPUQuadTree.GetValueOnGameThread() != 0;
+	const bool bBuildQuadTreeInSceneProxy = ShouldBuildQuadTreeInSceneProxy();
+
+	// Recreate WaterQuadTreeBuilder
+	if (bBuildQuadTreeInSceneProxy)
+	{
+		FMaterialRenderProxy* FarDistanceMaterialProxy = IsMaterialUsedWithWater(FarDistanceMaterial) ? FarDistanceMaterial->GetRenderProxy() : nullptr;
+		WaterQuadTreeBuilder.Init(WaterZone->GetZoneBounds2D(), InExtentInTiles, InTileSize, FarDistanceMaterialProxy, FarDistanceMeshExtent, bUseFarMeshWithoutOcean, bIsGPUQuadTree);
+
+		UsedMaterials.Empty();
+
+		bool bUsesFarDistanceMaterial = bUseFarMeshWithoutOcean;
+
+		WaterZone->ForEachWaterBodyComponent([&](UWaterBodyComponent* WaterBodyComponent) -> bool
+		{
+			check(WaterBodyComponent);
+			AActor* Actor = WaterBodyComponent->GetOwner();
+			check(Actor);
+
+			// Skip invisible water bodies
+			if (!WaterBodyComponent->ShouldRender() || !WaterBodyComponent->ShouldGenerateWaterMeshTile())
+			{
+				return true;
+			}
+
+			UWaterBodyInfoMeshComponent* WaterBodyInfoMeshComponent = WaterBodyComponent->GetWaterInfoMeshComponent();
+			UStaticMesh* StaticMesh = WaterBodyInfoMeshComponent ? WaterBodyInfoMeshComponent->GetStaticMesh() : nullptr;
+			bAnyWaterMeshesNotReady |= StaticMesh && StaticMesh->IsCompiling();
+			FStaticMeshRenderData* StaticMeshRenderData = StaticMesh ? StaticMesh->GetRenderData() : nullptr;
+			if (!ensure(WaterBodyInfoMeshComponent) || !ensure(StaticMesh) || !StaticMeshRenderData)
+			{
+				return true;
+			}
+
+			bUsesFarDistanceMaterial = bUsesFarDistanceMaterial || (WaterBodyComponent->GetWaterBodyType() == EWaterBodyType::Ocean);
+
+			auto GetMaterialProxy = [&](UMaterialInstanceDynamic* MID, bool bUseFallback) -> FMaterialRenderProxy*
+			{
+				UMaterialInterface* MaterialInterface = nullptr;
+				if (!MID || !IsMaterialUsedWithWater(MID))
+				{
+					MaterialInterface = bUseFallback ? UMaterial::GetDefaultMaterial(MD_Surface) : nullptr;
+				}
+				else
+				{
+					WaterBodyComponent->SetDynamicParametersOnMID(MID);
+					MaterialInterface = MID;
+				}
+				if (MaterialInterface)
+				{
+					UsedMaterials.Add(MaterialInterface);
+					return MaterialInterface->GetRenderProxy();
+				}
+				return nullptr;
+			};
+			
+			const bool bIsRiver = WaterBodyComponent->GetWaterBodyType() == EWaterBodyType::River;
+
+			FWaterQuadTreeBuilder::FWaterBody WaterBody = {};
+			WaterBody.Material = GetMaterialProxy(WaterBodyComponent->GetWaterMaterialInstance(), true);
+			WaterBody.RiverToLakeMaterial = bIsRiver ? GetMaterialProxy(WaterBodyComponent->GetRiverToLakeTransitionMaterialInstance(), false) : nullptr;
+			WaterBody.RiverToOceanMaterial = bIsRiver ? GetMaterialProxy(WaterBodyComponent->GetRiverToOceanTransitionMaterialInstance(), false) : nullptr;
+			WaterBody.StaticMeshRenderData = StaticMeshRenderData;
+			WaterBody.LocalToWorld = WaterBodyComponent->GetComponentTransform();
+			WaterBody.Bounds = WaterBodyComponent->Bounds;
+			WaterBody.OverlapMaterialPriority = WaterBodyComponent->GetOverlapMaterialPriority();
+			WaterBody.Type = WaterBodyComponent->GetWaterBodyType();
+			WaterBody.WaterBodyIndex = WaterBodyComponent->GetWaterBodyIndex();
+			WaterBody.SurfaceBaseHeight = WaterBodyComponent->GetComponentLocation().Z;
+			WaterBody.MaxWaveHeight = WaterBodyComponent->GetMaxWaveHeight();
+#if WITH_WATER_SELECTION_SUPPORT
+			WaterBody.HitProxy = new HActor(/*InActor = */Actor, /*InPrimComponent = */nullptr);
+			WaterBody.bWaterBodySelected = Actor->IsSelected();
+#endif // WITH_WATER_SELECTION_SUPPORT
+
+			if (!bIsGPUQuadTree)
+			{
+				switch (WaterBody.Type)
+				{
+				case EWaterBodyType::River:
+				{
+					TRACE_CPUPROFILER_EVENT_SCOPE(River);
+
+					TArray<FBox, TInlineAllocator<16>> Boxes;
+					TArray<UPrimitiveComponent*> CollisionComponents = WaterBodyComponent->GetCollisionComponents();
+					for (UPrimitiveComponent* Comp : CollisionComponents)
+					{
+						if (UBodySetup* BodySetup = (Comp != nullptr) ? Comp->GetBodySetup() : nullptr)
+						{
+							// Go through all sub shapes on the bodysetup to get a tight fit along water body
+							for (const FKConvexElem& ConvElem : BodySetup->AggGeom.ConvexElems)
+							{
+								TRACE_CPUPROFILER_EVENT_SCOPE(Add);
+
+								FBox SubBox = ConvElem.ElemBox.TransformBy(Comp->GetComponentTransform().ToMatrixWithScale());
+								SubBox.Max.Z += WaterBodyComponent->GetMaxWaveHeight();
+
+								Boxes.Add(SubBox);
+							}
+						}
+						else
+						{
+							// fallback on global AABB: 
+							FVector Center;
+							FVector Extent;
+							Actor->GetActorBounds(false, Center, Extent);
+							FBox Box(FBox::BuildAABB(Center, Extent));
+							Box.Max.Z += WaterBodyComponent->GetMaxWaveHeight();
+							Boxes.Add(Box);
+						}
+					}
+
+					for (const FBox& Box : Boxes)
+					{
+#if !(UE_BUILD_SHIPPING || UE_BUILD_TEST)
+						if (!!CVarWaterMeshShowTileGenerationGeometry.GetValueOnGameThread())
+						{
+							DrawDebugBox(GetWorld(), Box.GetCenter(), Box.GetExtent(), FColor::Red);
+						}
+#endif
+					}
+					WaterBody.RiverBoxes = Boxes;
+					break;
+				}
+				case EWaterBodyType::Lake:
+				{
+					TRACE_CPUPROFILER_EVENT_SCOPE(Lake);
+
+					const UWaterSplineComponent* SplineComp = WaterBodyComponent->GetWaterSpline();
+					const int32 NumOriginalSplinePoints = SplineComp->GetNumberOfSplinePoints();
+					float ConstantZ = SplineComp->GetLocationAtDistanceAlongSpline(0.0f, ESplineCoordinateSpace::World).Z;
+
+					FBox LakeBounds = Actor->GetComponentsBoundingBox(/* bNonColliding = */true);
+					LakeBounds.Max.Z += WaterBodyComponent->GetMaxWaveHeight();
+
+					// Skip lakes with less than 3 spline points
+					if (NumOriginalSplinePoints < 3)
+					{
+						return true;
+					}
+
+					TArray<TArray<FVector2D>> PolygonBatches;
+					// Reuse the convex hulls generated for the physics shape because the work has already been done, but we can fallback to a simple spline evaluation method in case there's no physics:
+					bool bUseFallbackMethod = true;
+
+					TArray<UPrimitiveComponent*> CollisionComponents = WaterBodyComponent->GetCollisionComponents();
+					if (CollisionComponents.Num() > 0)
+					{
+						UPrimitiveComponent* Comp = CollisionComponents[0];
+						if (UBodySetup* BodySetup = (Comp != nullptr) ? Comp->GetBodySetup() : nullptr)
+						{
+							FTransform CompTransform = Comp->GetComponentTransform();
+							// Go through all sub shapes on the bodysetup to get a tight fit along water body
+							for (const FKConvexElem& ConvElem : BodySetup->AggGeom.ConvexElems)
+							{
+								TRACE_CPUPROFILER_EVENT_SCOPE(Add);
+
+								// Vertex data contains the bottom vertices first, then the top ones :
+								int32 TotalNumVertices = ConvElem.VertexData.Num();
+								check(TotalNumVertices % 2 == 0);
+								int32 NumVertices = TotalNumVertices / 2;
+								if (NumVertices > 0)
+								{
+									bUseFallbackMethod = false;
+
+									// Because the physics shape is made of multiple convex hulls, we cannot simply add their vertices to one big list but have to have 1 batch of polygons per convex hull
+									TArray<FVector2D>& Polygon = PolygonBatches.Emplace_GetRef();
+									Polygon.SetNum(NumVertices);
+									for (int32 i = 0; i < NumVertices; ++i)
+									{
+										// Gather the top vertices :
+										Polygon[i] = FVector2D(CompTransform.TransformPosition(ConvElem.VertexData[NumVertices + i]));
+									}
+								}
+							}
+						}
+					}
+
+					if (bUseFallbackMethod)
+					{
+						TArray<FVector> PolyLineVertices;
+						SplineComp->ConvertSplineToPolyLine(ESplineCoordinateSpace::World, FMath::Square(CVarWaterSplineResampleMaxDistance.GetValueOnGameThread()), PolyLineVertices);
+
+						TArray<FVector2D>& Polygon = PolygonBatches.Emplace_GetRef();
+						Polygon.Reserve(PolyLineVertices.Num());
+						Algo::Transform(PolyLineVertices, Polygon, [](const FVector& Vertex) { return FVector2D(Vertex); });
+					}
+
+					for (const TArray<FVector2D>& Polygon : PolygonBatches)
+					{
+#if !(UE_BUILD_SHIPPING || UE_BUILD_TEST)
+						if (!!CVarWaterMeshShowTileGenerationGeometry.GetValueOnGameThread())
+						{
+							float Z = SplineComp->GetLocationAtDistanceAlongSpline(0.0f, ESplineCoordinateSpace::World).Z;
+							int32 NumVertices = Polygon.Num();
+							for (int32 i = 0; i < NumVertices; i++)
+							{
+								const FVector2D& Point0 = Polygon[i];
+								const FVector2D& Point1 = Polygon[(i + 1) % NumVertices];
+								DrawDebugLine(GetWorld(), FVector(Point0.X, Point0.Y, Z), FVector(Point1.X, Point1.Y, Z), FColor::Green);
+							}
+						}
+#endif // !(UE_BUILD_SHIPPING || UE_BUILD_TEST)
+					}
+
+					WaterBody.PolygonBounds = LakeBounds;
+					WaterBody.PolygonBatches = MoveTemp(PolygonBatches);
+
+					break;
+				}
+				case EWaterBodyType::Ocean:
+				{
+					TRACE_CPUPROFILER_EVENT_SCOPE(Ocean);
+
+					const UWaterSplineComponent* SplineComp = WaterBodyComponent->GetWaterSpline();
+					const int32 NumOriginalSplinePoints = SplineComp->GetNumberOfSplinePoints();
+
+					// Skip oceans with less than 3 spline points
+					if (NumOriginalSplinePoints < 3)
+					{
+						return true;
+					}
+
+					TArray<FVector2D> Polygon;
+
+					TArray<FVector> PolyLineVertices;
+					SplineComp->ConvertSplineToPolyLine(ESplineCoordinateSpace::World, FMath::Square(CVarWaterSplineResampleMaxDistance.GetValueOnGameThread()), PolyLineVertices);
+
+					Polygon.Reserve(PolyLineVertices.Num());
+					Algo::Transform(PolyLineVertices, Polygon, [](const FVector& Vertex) { return FVector2D(Vertex); });
+
+					FBox OceanBounds = WaterBodyComponent->Bounds.GetBox();
+					OceanBounds.Max.Z += WaterBodyComponent->GetMaxWaveHeight();
+
+#if !(UE_BUILD_SHIPPING || UE_BUILD_TEST)
+					if (!!CVarWaterMeshShowTileGenerationGeometry.GetValueOnGameThread())
+					{
+						float Z = SplineComp->GetLocationAtDistanceAlongSpline(0.0f, ESplineCoordinateSpace::World).Z;
+						int32 NumVertices = Polygon.Num();
+						for (int32 i = 0; i < NumVertices; i++)
+						{
+							const FVector2D& Point0 = Polygon[i];
+							const FVector2D& Point1 = Polygon[(i + 1) % NumVertices];
+							DrawDebugLine(GetWorld(), FVector(Point0.X, Point0.Y, Z), FVector(Point1.X, Point1.Y, Z), FColor::Blue);
+						}
+
+						DrawDebugBox(GetWorld(), OceanBounds.GetCenter(), OceanBounds.GetExtent(), FColor::Blue);
+					}
+#endif //!(UE_BUILD_SHIPPING || UE_BUILD_TEST)
+
+					WaterBody.PolygonBounds = OceanBounds;
+					WaterBody.PolygonBatches.Add(MoveTemp(Polygon));
+
+					break;
+				}
+				case EWaterBodyType::Transition:
+					// Transitions dont require rendering
+					break;
+				default:
+					ensureMsgf(false, TEXT("This water body type is not implemented and will not produce any water tiles. "));
+				}
+			}
+
+			WaterQuadTreeBuilder.AddWaterBody(MoveTemp(WaterBody));
+
+			return true;
+		});
+
+		if (bUsesFarDistanceMaterial && FarDistanceMaterialProxy && (FarDistanceMeshExtent > 0.0f))
+		{
+			UsedMaterials.Add(FarDistanceMaterial);
+		}
+
+		MarkRenderStateDirty();
+
+		// Force another rebuild next frame if water body meshes are still compiling
+		if (bIsGPUQuadTree && bAnyWaterMeshesNotReady)
+		{
+			bNeedsRebuild = true;
+		}
+
+		return;
+	}
+
 	// Position snapped to the grid
 	const FVector2D GridPosition = WaterZone->IsLocalOnlyTessellationEnabled() ? GetDynamicWaterMeshCenter() : FVector2D(FMath::GridSnap<FVector::FReal>(GetComponentLocation().X, InTileSize), FMath::GridSnap<FVector::FReal>(GetComponentLocation().Y, InTileSize));
 	const FVector2D WorldExtent = FVector2D(InTileSize * InExtentInTiles.X, InTileSize * InExtentInTiles.Y);
@@ -268,8 +582,6 @@ void UWaterMeshComponent::RebuildWaterMesh(float InTileSize, const FIntPoint& In
 		return;
 	}
 
-	const bool bIsGPUQuadTree = CVarWaterMeshGPUQuadTree.GetValueOnGameThread() != 0;
-
 	// This resets the tree to an initial state, ready for node insertion
 	WaterQuadTree.InitTree(WaterWorldBox, InTileSize, InExtentInTiles, bIsGPUQuadTree);
 
@@ -348,9 +660,6 @@ void UWaterMeshComponent::RebuildWaterMesh(float InTileSize, const FIntPoint& In
 		InMID->SetTextureParameterValue(UWaterBodyComponent::WaterVelocityAndHeightName, OwningZone->WaterInfoTexture);
 		return true;
 	};
-	
-
-	bool bAnyWaterMeshesNotReady = false;
 
 	OwningZone->ForEachWaterBodyComponent([this, bIsGPUQuadTree, WaterWorldBox, bIsFlooded, GlobalOceanHeight, OceanFlood, &FarMeshHeight, &bHasOcean, &SortedPriorities, &bAnyWaterMeshesNotReady, &SetDynamicParametersOnSharedMID](UWaterBodyComponent* WaterBodyComponent)
 	{
@@ -449,9 +758,9 @@ void UWaterMeshComponent::RebuildWaterMesh(float InTileSize, const FIntPoint& In
 					}
 				}
 
-				RenderData.Material = OutMaterials[0];
-				RenderData.RiverToLakeMaterial = OutMaterials[1];
-				RenderData.RiverToOceanMaterial = OutMaterials[2];
+				RenderData.Material = OutMaterials[0] ? OutMaterials[0]->GetRenderProxy() : nullptr;
+				RenderData.RiverToLakeMaterial = OutMaterials[1] ? OutMaterials[1]->GetRenderProxy() : nullptr;
+				RenderData.RiverToOceanMaterial = OutMaterials[2] ? OutMaterials[2]->GetRenderProxy() : nullptr;
 			}
 
 			RenderData.Priority = static_cast<int16>(FMath::Clamp(WaterBodyComponent->GetOverlapMaterialPriority(), MinWaterBodyPriority, MaxWaterBodyPriority));
@@ -476,22 +785,23 @@ void UWaterMeshComponent::RebuildWaterMesh(float InTileSize, const FIntPoint& In
 		{
 			// Assign material instance
 			UMaterialInstanceDynamic* WaterMaterial = WaterBodyComponent->GetWaterMaterialInstance();
-			RenderData.Material = WaterMaterial;
+			RenderData.Material = WaterMaterial->GetRenderProxy();
 
 			if (RenderData.Material)
 			{
-				if (!IsMaterialUsedWithWater(RenderData.Material))
+				if (!IsMaterialUsedWithWater(WaterMaterial))
 				{
-					RenderData.Material = UMaterial::GetDefaultMaterial(MD_Surface);
+					RenderData.Material = UMaterial::GetDefaultMaterial(MD_Surface)->GetRenderProxy();
+					// Add material so that the component keeps track of all potential materials used
+					UsedMaterials.Add(UMaterial::GetDefaultMaterial(MD_Surface));
 				}
 				else
 				{
 					// Add ocean height as a scalar parameter
 					WaterBodyComponent->SetDynamicParametersOnMID(WaterMaterial);
+					// Add material so that the component keeps track of all potential materials used
+					UsedMaterials.Add(WaterMaterial);
 				}
-
-				// Add material so that the component keeps track of all potential materials used
-				UsedMaterials.Add(RenderData.Material);
 			}
 
 			RenderData.Priority = static_cast<int16>(FMath::Clamp(WaterBodyComponent->GetOverlapMaterialPriority(), MinWaterBodyPriority, MaxWaterBodyPriority));
@@ -518,8 +828,8 @@ void UWaterMeshComponent::RebuildWaterMesh(float InTileSize, const FIntPoint& In
 				UMaterialInstanceDynamic* RiverToLakeMaterial = WaterBodyComponent->GetRiverToLakeTransitionMaterialInstance();
 				if (IsMaterialUsedWithWater(RiverToLakeMaterial))
 				{
-					RenderData.RiverToLakeMaterial = RiverToLakeMaterial;
-					UsedMaterials.Add(RenderData.RiverToLakeMaterial);
+					RenderData.RiverToLakeMaterial = RiverToLakeMaterial->GetRenderProxy();
+					UsedMaterials.Add(RiverToLakeMaterial);
 					// Add ocean height as a scalar parameter
 					WaterBodyComponent->SetDynamicParametersOnMID(RiverToLakeMaterial);
 				}
@@ -527,8 +837,8 @@ void UWaterMeshComponent::RebuildWaterMesh(float InTileSize, const FIntPoint& In
 				UMaterialInstanceDynamic* RiverToOceanMaterial = WaterBodyComponent->GetRiverToOceanTransitionMaterialInstance();
 				if (IsMaterialUsedWithWater(RiverToOceanMaterial))
 				{
-					RenderData.RiverToOceanMaterial = RiverToOceanMaterial;
-					UsedMaterials.Add(RenderData.RiverToOceanMaterial);
+					RenderData.RiverToOceanMaterial = RiverToOceanMaterial->GetRenderProxy();
+					UsedMaterials.Add(RiverToOceanMaterial);
 					// Add ocean height as a scalar parameter
 					WaterBodyComponent->SetDynamicParametersOnMID(RiverToOceanMaterial);
 				}
@@ -793,7 +1103,7 @@ void UWaterMeshComponent::RebuildWaterMesh(float InTileSize, const FIntPoint& In
 		// Far Mesh should stitch to the edge of the water zone
 		const FBox2D FarMeshBounds = WaterZone->GetZoneBounds2D();
 
-		WaterQuadTree.AddFarMesh(FarDistanceMaterial, FarMeshBounds, FarDistanceMeshExtent, FarMeshHeight);
+		WaterQuadTree.AddFarMesh(FarDistanceMaterial->GetRenderProxy(), FarMeshBounds, FarDistanceMeshExtent, FarMeshHeight);
 	}
 
 	// Remove all materials from the MaterialToMID map that aren't currently used
diff --git a/Engine/Plugins/Experimental/Water/Source/Runtime/Private/WaterMeshSceneProxy.cpp b/Engine/Plugins/Experimental/Water/Source/Runtime/Private/WaterMeshSceneProxy.cpp
--- a/Engine/Plugins/Experimental/Water/Source/Runtime/Private/WaterMeshSceneProxy.cpp
+++ b/Engine/Plugins/Experimental/Water/Source/Runtime/Private/WaterMeshSceneProxy.cpp
@@ -165,8 +165,23 @@ FWaterMeshSceneProxy::FWaterMeshSceneProxy(UWaterMeshComponent* Component)
 	: FPrimitiveSceneProxy(Component)
 	, MaterialRelevance(Component->GetWaterMaterialRelevance(GetScene().GetFeatureLevel()))
 {
-	// Cache the tiles and settings
-	WaterQuadTree = Component->GetWaterQuadTree();
+	AWaterZone* WaterZone = CastChecked<AWaterZone>(Component->GetOwner());
+	checkf(WaterZone != nullptr, TEXT("WaterMeshComponent is owned by an actor that is not a WaterZone. This is not supported!"));
+
+	if (Component->ShouldBuildQuadTreeInSceneProxy())
+	{
+		// Get the QuadTreeBuilder and build the quadtree
+		const FWaterQuadTreeBuilder& WaterQuadTreeBuilder = Component->GetWaterQuadTreeBuilder();
+		const FVector2D GridPosition = WaterZone->IsLocalOnlyTessellationEnabled() ?
+			Component->GetDynamicWaterMeshCenter() :
+			FVector2D(FMath::GridSnap<FVector::FReal>(Component->GetComponentLocation().X, WaterQuadTreeBuilder.GetLeafSize()), FMath::GridSnap<FVector::FReal>(Component->GetComponentLocation().Y, WaterQuadTreeBuilder.GetLeafSize()));
+		WaterQuadTreeBuilder.BuildWaterQuadTree(WaterQuadTree, GridPosition);
+	}
+	else
+	{
+		WaterQuadTree = Component->GetWaterQuadTree();
+	}
+
 	// Leaf size * 0.5 equals the tightest possible LOD Scale that doesn't break the morphing. Can be scaled larger
 	LODScale = WaterQuadTree.GetLeafSize() * FMath::Max(Component->GetLODScale(), 0.5f);
 
@@ -253,11 +268,8 @@ FWaterMeshSceneProxy::FWaterMeshSceneProxy(UWaterMeshComponent* Component)
 
 	WaterQuadTree.BuildMaterialIndices();
 
-	if (const AWaterZone* WaterZone = Component->GetOwner<AWaterZone>(); ensureMsgf(WaterZone != nullptr, TEXT("WaterMeshComponent is owned by an actor that is not a WaterZone. This is not supported!")))
-	{
-		const FBox WaterInfoBounds3D = WaterZone->GetDynamicWaterInfoBounds();
-		WaterInfoBounds = FBox2D(FVector2D(WaterInfoBounds3D.Min), FVector2D(WaterInfoBounds3D.Max));
-	}
+	const FBox WaterInfoBounds3D = WaterZone->GetDynamicWaterInfoBounds();
+	WaterInfoBounds = FBox2D(FVector2D(WaterInfoBounds3D.Min), FVector2D(WaterInfoBounds3D.Max));
 
 #if RHI_RAYTRACING
 	RayTracingWaterData.SetNum(DensityCount);
diff --git a/Engine/Plugins/Experimental/Water/Source/Runtime/Private/WaterQuadTree.cpp b/Engine/Plugins/Experimental/Water/Source/Runtime/Private/WaterQuadTree.cpp
--- a/Engine/Plugins/Experimental/Water/Source/Runtime/Private/WaterQuadTree.cpp
+++ b/Engine/Plugins/Experimental/Water/Source/Runtime/Private/WaterQuadTree.cpp
@@ -665,7 +665,7 @@ void FWaterQuadTree::AddLake(const TArray<FVector2D>& InPoly, const FBox& InLake
 	AddLakeRecursive(InPoly, LakeBounds, FVector2D(InLakeBounds.Min.Z, InLakeBounds.Max.Z), true, TreeDepth * 2, InWaterBodyIndex);
 }
 
-void FWaterQuadTree::AddFarMesh(const UMaterialInterface* InFarMeshMaterial, const FBox2D& InInnerRegion, double InFarDistanceMeshExtent, double InFarDistanceMeshHeight)
+void FWaterQuadTree::AddFarMesh(FMaterialRenderProxy* InFarMeshMaterial, const FBox2D& InInnerRegion, double InFarDistanceMeshExtent, double InFarDistanceMeshHeight)
 {
 	// Checking for not being read only here to keep things consistent with the other Add functions. In reality the FarMesh isn't added to the QuadTree itself, so it could technically be done whenever.
 	ensure(!bIsReadOnly);
@@ -715,14 +715,12 @@ void FWaterQuadTree::BuildMaterialIndices()
 	int32 NextIdx = 0;
 	TMap<FMaterialRenderProxy*, int32> MatToIdxMap;
 
-	auto GetMatIdx = [&NextIdx, &MatToIdxMap](const UMaterialInterface* Material)
+	auto GetMatIdx = [&NextIdx, &MatToIdxMap](FMaterialRenderProxy* MaterialRenderProxy)
 	{
-		if (!Material)
+		if (!MaterialRenderProxy)
 		{
 			return (int32)INDEX_NONE;
 		}
-		FMaterialRenderProxy* MaterialRenderProxy = Material->GetRenderProxy();
-		check(MaterialRenderProxy != nullptr);
 		const int32* Found = MatToIdxMap.Find(MaterialRenderProxy);
 		if (!Found)
 		{
diff --git a/Engine/Plugins/Experimental/Water/Source/Runtime/Private/WaterQuadTreeBuilder.cpp b/Engine/Plugins/Experimental/Water/Source/Runtime/Private/WaterQuadTreeBuilder.cpp
new file mode 100644
--- /dev/null
+++ b/Engine/Plugins/Experimental/Water/Source/Runtime/Private/WaterQuadTreeBuilder.cpp
@@ -0,0 +1,211 @@
+// Copyright Epic Games, Inc. All Rights Reserved.
+
+#include "WaterQuadTreeBuilder.h"
+#include "WaterQuadTree.h"
+#include "WaterBodyTypes.h"
+#include "WaterModule.h"
+
+void FWaterQuadTreeBuilder::Init(const FBox2D& InWaterZoneBounds2D, const FIntPoint& InExtentInTiles, float InTileSize, FMaterialRenderProxy* InFarDistanceMaterial, float InFarDistanceMeshExtent, bool bInUseFarMeshWithoutOcean, bool bInIsGPUQuadTree)
+{
+	WaterBodies.Reset();
+	WaterZoneBounds2D = InWaterZoneBounds2D;
+	ExtentInTiles = InExtentInTiles;
+	TileSize = InTileSize;
+	FarDistanceMaterial = InFarDistanceMaterial;
+	FarDistanceMeshExtent = InFarDistanceMeshExtent;
+	bUseFarMeshWithoutOcean = bInUseFarMeshWithoutOcean;
+	bIsGPUQuadTree = bInIsGPUQuadTree;
+}
+
+void FWaterQuadTreeBuilder::AddWaterBody(const FWaterBody& WaterBody)
+{
+	WaterBodies.Add(WaterBody);
+}
+
+bool FWaterQuadTreeBuilder::BuildWaterQuadTree(FWaterQuadTree& WaterQuadTree, const FVector2D& GridPosition) const
+{
+	const FVector2D WorldExtent = FVector2D(TileSize * ExtentInTiles.X, TileSize * ExtentInTiles.Y);
+	const FBox2D WaterWorldBox = FBox2D(-WorldExtent + GridPosition, WorldExtent + GridPosition);
+
+	// If the dynamic bounds is outside the full bounds of the water mesh, we shouldn't regenerate the quadtree
+	if (!(WaterWorldBox.GetArea() > 0.f))
+	{
+		return false;
+	}
+
+	// This resets the tree to an initial state, ready for node insertion
+	WaterQuadTree.InitTree(WaterWorldBox, TileSize, ExtentInTiles, bIsGPUQuadTree);
+
+	// Will be updated with the ocean min bound, to be used to place the far mesh just under the ocean to avoid seams
+	float FarMeshHeight = 0.0f;
+	// Only use a far mesh when there is an ocean in the zone.
+	bool bHasOcean = false;
+	
+	// Min and max user defined priority range. (Input also clamped on OverlapMaterialPriority in AWaterBody)
+	constexpr int32 MinWaterBodyPriority = -8192;
+	constexpr int32 MaxWaterBodyPriority = 8191;
+	constexpr int32 GPUQuadTreeMaxNumPriorities = 8;
+
+	// The GPU quadtree only supports 8 different priority values, so we need to remap priorities into that space.
+	// Fortunately, rivers and non-river water bodies are rendered into their own "priority space", so we don't need to worry about
+	// moving river priorities into their own range.
+	TArray<int16> SortedPriorities;
+	if (bIsGPUQuadTree)
+	{
+		for (const FWaterBody& WaterBody : WaterBodies)
+		{
+			// Don't process water bodies which have their spline outside of this water mesh
+			const FBox WaterBodyBounds = WaterBody.Bounds.GetBox();
+			if (!WaterBodyBounds.IntersectXY(FBox(FVector(WaterWorldBox.Min, 0.0), FVector(WaterWorldBox.Max, 0.0))))
+			{
+				continue;
+			}
+
+			const int16 Priority = static_cast<int16>(FMath::Clamp(WaterBody.OverlapMaterialPriority, MinWaterBodyPriority, MaxWaterBodyPriority));
+			SortedPriorities.AddUnique(Priority);
+		}
+
+		SortedPriorities.Sort();
+
+		if (SortedPriorities.Num() > GPUQuadTreeMaxNumPriorities)
+		{
+			UE_LOG(LogWater, Warning, TEXT("WaterZone has more unique water body priorities (%i) than can be supported with GPU driven water quadtree rendering (%i)!"), SortedPriorities.Num(), GPUQuadTreeMaxNumPriorities);
+		}
+	}
+
+	for (const FWaterBody& WaterBody : WaterBodies)
+	{
+		// Don't process water bodies which have their spline outside of this water mesh
+		const FBox WaterBodyBounds = WaterBody.Bounds.GetBox();
+		if (!WaterBodyBounds.IntersectXY(FBox(FVector(WaterWorldBox.Min, 0.0), FVector(WaterWorldBox.Max, 0.0))))
+		{
+			continue;
+		}
+
+		FWaterBodyRenderData RenderData;
+		RenderData.Material = WaterBody.Material;
+		RenderData.RiverToLakeMaterial = WaterBody.RiverToLakeMaterial;
+		RenderData.RiverToOceanMaterial = WaterBody.RiverToOceanMaterial;
+		RenderData.Priority = static_cast<int16>(FMath::Clamp(WaterBody.OverlapMaterialPriority, MinWaterBodyPriority, MaxWaterBodyPriority));
+		RenderData.WaterBodyIndex = static_cast<int16>(WaterBody.WaterBodyIndex);
+		RenderData.SurfaceBaseHeight = WaterBody.SurfaceBaseHeight;
+		RenderData.MaxWaveHeight = WaterBody.MaxWaveHeight;
+		RenderData.BoundsMinZ = WaterBody.Bounds.GetBox().Min.Z;
+		RenderData.BoundsMaxZ = WaterBody.Bounds.GetBox().Max.Z;
+		RenderData.WaterBodyType = static_cast<int8>(WaterBody.Type);
+#if WITH_WATER_SELECTION_SUPPORT
+		RenderData.HitProxy = WaterBody.HitProxy;
+		RenderData.bWaterBodySelected = WaterBody.bWaterBodySelected;
+#endif // WITH_WATER_SELECTION_SUPPORT
+
+		if (RenderData.RiverToLakeMaterial || RenderData.RiverToOceanMaterial)
+		{
+			// Move rivers up to it's own priority space, so that they always have precedence if they have transitions and that they only compare agains other rivers with transitions
+			RenderData.Priority += (MaxWaterBodyPriority - MinWaterBodyPriority) + 1;
+		}
+
+		const uint32 WaterBodyRenderDataIndex = WaterQuadTree.AddWaterBodyRenderData(RenderData);
+
+		if (bIsGPUQuadTree)
+		{
+			// On the GPU path, we only submit FWaterBodyQuadTreeRasterInfo for the water quadtree to be rasterized on the GPU. In particular, we need the static mesh, a transform and a priority/WaterBodyRenderData index.
+			const int16 ClampedPriority = static_cast<int16>(FMath::Clamp(WaterBody.OverlapMaterialPriority, MinWaterBodyPriority, MaxWaterBodyPriority));
+
+			FWaterBodyQuadTreeRasterInfo RasterInfo;
+			RasterInfo.LocalToWorld = WaterBody.LocalToWorld;
+			RasterInfo.RenderData = WaterBody.StaticMeshRenderData;
+			RasterInfo.WaterBodyRenderDataIndex = WaterBodyRenderDataIndex;
+			RasterInfo.Priority = FMath::Clamp(SortedPriorities.IndexOfByKey(ClampedPriority), 0, GPUQuadTreeMaxNumPriorities - 1);
+			RasterInfo.bIsRiver = WaterBody.Type == EWaterBodyType::River;
+
+			WaterQuadTree.AddWaterBodyRasterInfo(RasterInfo);
+
+			if (WaterBody.Type == EWaterBodyType::Ocean)
+			{
+				// Place far mesh height just below the ocean level
+				FarMeshHeight = RenderData.SurfaceBaseHeight - RenderData.MaxWaveHeight;
+				bHasOcean = true;
+			}
+		}
+		else
+		{
+			switch (WaterBody.Type)
+			{
+			case EWaterBodyType::River:
+			{
+				for (const FBox& Box : WaterBody.RiverBoxes)
+				{
+					WaterQuadTree.AddWaterTilesInsideBounds(Box, WaterBodyRenderDataIndex);
+				}
+				break;
+			}
+			case EWaterBodyType::Lake:
+			{
+				for (const TArray<FVector2D>& Polygon : WaterBody.PolygonBatches)
+				{
+					WaterQuadTree.AddLake(Polygon, WaterBody.PolygonBounds, WaterBodyRenderDataIndex);
+				}
+				break;
+			}
+			case EWaterBodyType::Ocean:
+			{
+				check(WaterBody.PolygonBatches.Num() == 1);
+				WaterQuadTree.AddOcean(WaterBody.PolygonBatches[0], WaterBody.PolygonBounds, WaterBodyRenderDataIndex);
+
+				// Place far mesh height just below the ocean level
+				FarMeshHeight = RenderData.SurfaceBaseHeight - WaterBody.MaxWaveHeight;
+				bHasOcean = true;
+
+				break;
+			}
+			case EWaterBodyType::Transition:
+				// Transitions dont require rendering
+				break;
+			default:
+				ensureMsgf(false, TEXT("This water body type is not implemented and will not produce any water tiles. "));
+			}
+		}
+	}
+
+	// Build the far distance mesh instances if needed
+	if (FarDistanceMaterial && (bHasOcean || bUseFarMeshWithoutOcean) && (FarDistanceMeshExtent > 0.0f))
+	{
+		// Far Mesh should stitch to the edge of the water zone
+		const FBox2D FarMeshBounds = WaterZoneBounds2D;
+		WaterQuadTree.AddFarMesh(FarDistanceMaterial, FarMeshBounds, FarDistanceMeshExtent, FarMeshHeight);
+	}
+
+	WaterQuadTree.Unlock(true);
+
+	return true;
+}
+
+#if WITH_WATER_SELECTION_SUPPORT
+void FWaterQuadTreeBuilder::GatherHitProxies(TArray<TRefCountPtr<HHitProxy>>& OutHitProxies) const
+{
+	for (const FWaterBody& WaterBody : WaterBodies)
+	{
+		OutHitProxies.Add(WaterBody.HitProxy);
+	}
+}
+#endif // WITH_WATER_SELECTION_SUPPORT
+
+bool FWaterQuadTreeBuilder::IsGPUQuadTree() const
+{
+	return bIsGPUQuadTree;
+}
+
+float FWaterQuadTreeBuilder::GetLeafSize() const
+{
+	return TileSize;
+}
+
+int32 FWaterQuadTreeBuilder::GetMaxLeafCount() const
+{
+	return ExtentInTiles.X * ExtentInTiles.Y * 4;
+}
+
+FIntPoint FWaterQuadTreeBuilder::GetResolution() const
+{
+	return ExtentInTiles * 2;
+}
diff --git a/Engine/Plugins/Experimental/Water/Source/Runtime/Public/WaterMeshComponent.h b/Engine/Plugins/Experimental/Water/Source/Runtime/Public/WaterMeshComponent.h
--- a/Engine/Plugins/Experimental/Water/Source/Runtime/Public/WaterMeshComponent.h
+++ b/Engine/Plugins/Experimental/Water/Source/Runtime/Public/WaterMeshComponent.h
@@ -4,6 +4,7 @@
 
 #include "Components/MeshComponent.h"
 #include "WaterQuadTree.h"
+#include "WaterQuadTreeBuilder.h"
 #include "WaterMeshComponent.generated.h"
 
 struct FPSOPrecacheParams;
@@ -49,6 +50,8 @@ public:
 	/** Use this instead of GetMaterialRelevance, since this one will go over all materials from all tiles */
 	FMaterialRelevance GetWaterMaterialRelevance(ERHIFeatureLevel::Type InFeatureLevel) const;
 
+	bool ShouldBuildQuadTreeInSceneProxy() const;
+	const FWaterQuadTreeBuilder& GetWaterQuadTreeBuilder() const { return WaterQuadTreeBuilder; }
 	const FWaterQuadTree& GetWaterQuadTree() const { return WaterQuadTree; }
 
 	const TSet<TObjectPtr<UMaterialInterface>>& GetUsedMaterialsSet() const { return UsedMaterials; }
@@ -106,6 +109,8 @@ private:
 	/** The current center of the dynamic water mesh. Updated by the water view extension whenever the view location crosses the update bounds. */
 	FVector2D DynamicWaterMeshCenter = FVector2D::ZeroVector;
 
+	FWaterQuadTreeBuilder WaterQuadTreeBuilder;
+
 	/** Tiles containing water, stored in a quad tree */
 	FWaterQuadTree WaterQuadTree;
 
diff --git a/Engine/Plugins/Experimental/Water/Source/Runtime/Public/WaterQuadTree.h b/Engine/Plugins/Experimental/Water/Source/Runtime/Public/WaterQuadTree.h
--- a/Engine/Plugins/Experimental/Water/Source/Runtime/Public/WaterQuadTree.h
+++ b/Engine/Plugins/Experimental/Water/Source/Runtime/Public/WaterQuadTree.h
@@ -19,11 +19,11 @@ struct FWaterBodyRenderData
 	bool IsOcean() const { return WaterBodyType == 2; }
 
 	/** The standard material to be used for this water body */
-	UMaterialInterface* Material = nullptr;
+	FMaterialRenderProxy* Material = nullptr;
 
 	/** Transition materials (Only set up for river water bodies) */
-	UMaterialInterface* RiverToLakeMaterial = nullptr;
-	UMaterialInterface* RiverToOceanMaterial = nullptr;
+	FMaterialRenderProxy* RiverToLakeMaterial = nullptr;
+	FMaterialRenderProxy* RiverToOceanMaterial = nullptr;
 
 	/** World Z position of the waterbody, this is where the tiles for this water body will be rendered*/
 	double SurfaceBaseHeight = 0.0;
@@ -156,7 +156,7 @@ struct FWaterQuadTree
 	void AddLake(const TArray<FVector2D>& InPoly, const FBox& InLakeBounds, uint32 InWaterBodyIndex);
 
 	/** Add an automatically generated mesh (8 quads) skirt around InInnerRegion which extends out InFarDistanceMeshExtent, is placed at Z value InFarDistanceMeshHeight and is rendered using InFarMeshMaterial */
-	void AddFarMesh(const UMaterialInterface* InFarMeshMaterial, const FBox2D& InInnerRegion, double InFarDistanceMeshExtent, double InFarDistanceMeshHeight);
+	void AddFarMesh(FMaterialRenderProxy* InFarMeshMaterial, const FBox2D& InInnerRegion, double InFarDistanceMeshExtent, double InFarDistanceMeshHeight);
 
 	/** Assign an index to each material */
 	void BuildMaterialIndices();
@@ -335,7 +335,7 @@ private:
 		TArray<FFarMeshInstanceData> InstanceData;
 
 		/** Material for the Far Distance Mesh, its material render proxy will be cached in WaterMaterials when BuildMaterialIndices is called */
-		const UMaterialInterface* Material = nullptr;
+		FMaterialRenderProxy* Material = nullptr;
 
 		void Clear()
 		{
diff --git a/Engine/Plugins/Experimental/Water/Source/Runtime/Public/WaterQuadTreeBuilder.h b/Engine/Plugins/Experimental/Water/Source/Runtime/Public/WaterQuadTreeBuilder.h
new file mode 100644
--- /dev/null
+++ b/Engine/Plugins/Experimental/Water/Source/Runtime/Public/WaterQuadTreeBuilder.h
@@ -0,0 +1,58 @@
+// Copyright Epic Games, Inc. All Rights Reserved.
+
+#pragma once
+#include "CoreMinimal.h"
+
+class HHitProxy;
+class FStaticMeshRenderData;
+class FMaterialRenderProxy;
+struct FWaterQuadTree;
+enum class EWaterBodyType : uint8;
+
+class FWaterQuadTreeBuilder
+{
+public:
+
+	struct FWaterBody
+	{
+		FMaterialRenderProxy* Material = nullptr;
+		FMaterialRenderProxy* RiverToLakeMaterial = nullptr;
+		FMaterialRenderProxy* RiverToOceanMaterial = nullptr;
+		FStaticMeshRenderData* StaticMeshRenderData = nullptr;
+		FTransform LocalToWorld;
+		FBoxSphereBounds Bounds = FBoxSphereBounds();
+		int32 OverlapMaterialPriority = 0;
+		EWaterBodyType Type = EWaterBodyType(0);
+		int32 WaterBodyIndex = INDEX_NONE;
+		float SurfaceBaseHeight = 0.0f;
+		float MaxWaveHeight = 0.0f;
+		FBox PolygonBounds = FBox();
+		TArray<TArray<FVector2D>> PolygonBatches;
+		TArray<FBox> RiverBoxes;
+#if WITH_WATER_SELECTION_SUPPORT
+		TRefCountPtr<HHitProxy> HitProxy = nullptr;
+		bool bWaterBodySelected = false;
+#endif // WITH_WATER_SELECTION_SUPPORT
+	};
+
+	void Init(const FBox2D& InWaterZoneBounds2D, const FIntPoint& InExtentInTiles, float InTileSize, FMaterialRenderProxy* InFarDistanceMaterial, float InFarDistanceMeshExtent, bool bInUseFarMeshWithoutOcean, bool bInIsGPUQuadTree);
+	void AddWaterBody(const FWaterBody& WaterBody);
+	bool BuildWaterQuadTree(FWaterQuadTree& WaterQuadTree, const FVector2D& GridPosition) const;
+#if WITH_WATER_SELECTION_SUPPORT
+	void GatherHitProxies(TArray<TRefCountPtr<HHitProxy>>& OutHitProxies) const;
+#endif // WITH_WATER_SELECTION_SUPPORT
+	bool IsGPUQuadTree() const;
+	float GetLeafSize() const;
+	int32 GetMaxLeafCount() const;
+	FIntPoint GetResolution() const;
+
+private:
+	TArray<FWaterBody> WaterBodies;
+	FBox2D WaterZoneBounds2D = FBox2D();
+	FIntPoint ExtentInTiles = FIntPoint::ZeroValue;
+	float TileSize = 0.0f;
+	FMaterialRenderProxy* FarDistanceMaterial = nullptr;
+	float FarDistanceMeshExtent = 0.0f;
+	bool bUseFarMeshWithoutOcean = false;
+	bool bIsGPUQuadTree = false;
+};
