Water: Reduce attenuation depth in wave preview editor so waves can be previewed with full amplitude.

diff --git a/Engine/Plugins/Experimental/Water/Source/Editor/Private/WaterWavesEditorViewport.cpp b/Engine/Plugins/Experimental/Water/Source/Editor/Private/WaterWavesEditorViewport.cpp
--- a/Engine/Plugins/Experimental/Water/Source/Editor/Private/WaterWavesEditorViewport.cpp
+++ b/Engine/Plugins/Experimental/Water/Source/Editor/Private/WaterWavesEditorViewport.cpp
@@ -35,6 +35,8 @@ void SWaterWavesEditorViewport::Construct(const FArguments& InArgs)
 	check(WaterBodyComponent);
 	WaterBodyComponent->SetWaterMeshOverride(GetDefault<UWaterEditorSettings>()->WaterBodyCustomDefaults.GetWaterMesh());
 	WaterBodyComponent->SetWaterMaterial(GetDefault<UWaterEditorSettings>()->WaterBodyCustomDefaults.GetWaterMaterial());
+	// Reduce the wave attenuation target depth otherwise we will show attenuated waves in the preview which is not representative of the actual wave parameters:
+	WaterBodyComponent->TargetWaveMaskDepth = 1.f;
 
 
 	UWaterSplineComponent* WaterSpline = CustomWaterBody->GetWaterSpline();
