Water: prevent selection of ABuoyancyActor to avoid it appearing in Select All/Paste.

diff --git a/Engine/Plugins/Experimental/Water/Source/Runtime/Public/BuoyancyManager.h b/Engine/Plugins/Experimental/Water/Source/Runtime/Public/BuoyancyManager.h
--- a/Engine/Plugins/Experimental/Water/Source/Runtime/Public/BuoyancyManager.h
+++ b/Engine/Plugins/Experimental/Water/Source/Runtime/Public/BuoyancyManager.h
@@ -56,6 +56,8 @@ public:
 	virtual bool ShouldImport(FStringView ActorPropString, bool IsMovingLevel) override { return false; }
 	
 	virtual bool CanChangeIsSpatiallyLoadedFlag() const override { return false; }
+
+	virtual bool IsSelectable() const override { return false; }
 #endif // WITH_EDITOR
 private:
 	void InitializeAsyncAux(UBuoyancyComponent* Component);
