[Backout] - CL33457826
Original CL Desc
-----------------------------------------------------------------
Moved water quadtree creation to WaterMeshSceneProxy using a new FWaterQuadTreeBuilder class which can be used to construct quadtrees at different positions using a certain set of fixed parameters. This is the first step towards managing per-view quadtrees in the scene proxy. There should be no difference in behavior with/without this CL; this is merely a change in architecture to enable future improvements. There's a new cvar to toggle between the old and now path: r.Water.WaterMesh.BuildQuadTreeInSceneProxy. In order for that cvar to take effect, the water mesh needs to be rebuild. This can easily be achieved with r.Water.WaterMesh.ForceRebuildMeshPerFrame 1.

#rnx
#tests Local testing on PC in editor (QA-Waterworks, helios_terrain and creating new water bodies in a new empty level; with and without far distance material/mesh) as well as in a cooked instance of QAGame running QA-Waterworks

diff --git a/Engine/Plugins/Experimental/Water/Source/Runtime/Private/WaterMeshComponent.cpp b/Engine/Plugins/Experimental/Water/Source/Runtime/Private/WaterMeshComponent.cpp
--- a/Engine/Plugins/Experimental/Water/Source/Runtime/Private/WaterMeshComponent.cpp
+++ b/Engine/Plugins/Experimental/Water/Source/Runtime/Private/WaterMeshComponent.cpp
@@ -70,13 +70,6 @@ TAutoConsoleVariable<int32> CVarWaterMeshMIDDeduplication(
 	ECVF_RenderThreadSafe
 );
 
-TAutoConsoleVariable<int32> CVarWaterMeshBuildQuadTreeInSceneProxy(
-	TEXT("r.Water.WaterMesh.BuildQuadTreeInSceneProxy"),
-	1,
-	TEXT("Moves water quadtree construction into the water mesh scene proxy."),
-	ECVF_RenderThreadSafe
-);
-
 extern TAutoConsoleVariable<float> CVarWaterSplineResampleMaxDistance;
 
 
@@ -199,12 +192,6 @@ FMaterialRelevance UWaterMeshComponent::GetWaterMaterialRelevance(ERHIFeatureLev
 	return Result;
 }
 
-bool UWaterMeshComponent::ShouldBuildQuadTreeInSceneProxy() const
-{
-	const bool bShouldBuildQuadTreeInSceneProxy = CVarWaterMeshBuildQuadTreeInSceneProxy.GetValueOnGameThread() != 0;
-	return bShouldBuildQuadTreeInSceneProxy;
-}
-
 void UWaterMeshComponent::SetDynamicWaterMeshCenter(const FVector2D& NewCenter)
 {
 	if (!DynamicWaterMeshCenter.Equals(NewCenter))
@@ -242,30 +229,14 @@ FIntPoint UWaterMeshComponent::GetExtentInTiles() const
 
 FBoxSphereBounds UWaterMeshComponent::CalcBounds(const FTransform& LocalToWorld) const
 {
-	FBox NewBounds = FBox();
-	if (ShouldBuildQuadTreeInSceneProxy())
-	{
-		// With the water quadtree building moved into the scene proxy, the scene proxy is free to create quadtree(s) anywhere within the water zone.
-		// To ensure that the scene proxy isn't frustum culled, we need to be conservative and use the water zone bounds instead. The quadtree does
-		// fine grained frustum culling internally anyways, so this shouldn't be a performance issue.
-		AWaterZone* WaterZone = Cast<AWaterZone>(GetOwner());
-		if (WaterZone)
-		{
-			NewBounds = WaterZone->GetZoneBounds();
-		}
-	}
-	else
-	{
-		// Always return valid bounds (tree is initialized with invalid bounds and if nothing is inserted, the tree bounds will stay invalid)
-		NewBounds = WaterQuadTree.GetBounds();
+	// Always return valid bounds (tree is initialized with invalid bounds and if nothing is inserted, the tree bounds will stay invalid)
+	FBox NewBounds = WaterQuadTree.GetBounds();
 
-		if (NewBounds.Min.Z >= NewBounds.Max.Z)
-		{
-			NewBounds.Min.Z = 0.0f;
-			NewBounds.Max.Z = 100.0f;
-		}
+	if (NewBounds.Min.Z >= NewBounds.Max.Z)
+	{
+		NewBounds.Min.Z = 0.0f;
+		NewBounds.Max.Z = 100.0f;
 	}
-	
 	// Add the far distance to the bounds if it's valid
 	if (FarDistanceMaterial)
 	{
@@ -285,291 +256,6 @@ void UWaterMeshComponent::RebuildWaterMesh(float InTileSize, const FIntPoint& In
 
 	AWaterZone* WaterZone = CastChecked<AWaterZone>(GetOwner());
 
-	bool bAnyWaterMeshesNotReady = false;
-	const bool bIsGPUQuadTree = CVarWaterMeshGPUQuadTree.GetValueOnGameThread() != 0;
-	const bool bBuildQuadTreeInSceneProxy = ShouldBuildQuadTreeInSceneProxy();
-
-	// Recreate WaterQuadTreeBuilder
-	if (bBuildQuadTreeInSceneProxy)
-	{
-		FMaterialRenderProxy* FarDistanceMaterialProxy = IsMaterialUsedWithWater(FarDistanceMaterial) ? FarDistanceMaterial->GetRenderProxy() : nullptr;
-		WaterQuadTreeBuilder.Init(WaterZone->GetZoneBounds2D(), InExtentInTiles, InTileSize, FarDistanceMaterialProxy, FarDistanceMeshExtent, bUseFarMeshWithoutOcean, bIsGPUQuadTree);
-
-		UsedMaterials.Empty();
-
-		bool bUsesFarDistanceMaterial = bUseFarMeshWithoutOcean;
-
-		WaterZone->ForEachWaterBodyComponent([&](UWaterBodyComponent* WaterBodyComponent) -> bool
-		{
-			check(WaterBodyComponent);
-			AActor* Actor = WaterBodyComponent->GetOwner();
-			check(Actor);
-
-			// Skip invisible water bodies
-			if (!WaterBodyComponent->ShouldRender() || !WaterBodyComponent->ShouldGenerateWaterMeshTile())
-			{
-				return true;
-			}
-
-			UWaterBodyInfoMeshComponent* WaterBodyInfoMeshComponent = WaterBodyComponent->GetWaterInfoMeshComponent();
-			UStaticMesh* StaticMesh = WaterBodyInfoMeshComponent ? WaterBodyInfoMeshComponent->GetStaticMesh() : nullptr;
-			bAnyWaterMeshesNotReady |= StaticMesh && StaticMesh->IsCompiling();
-			FStaticMeshRenderData* StaticMeshRenderData = StaticMesh ? StaticMesh->GetRenderData() : nullptr;
-			if (!ensure(WaterBodyInfoMeshComponent) || !ensure(StaticMesh) || !StaticMeshRenderData)
-			{
-				return true;
-			}
-
-			bUsesFarDistanceMaterial = bUsesFarDistanceMaterial || (WaterBodyComponent->GetWaterBodyType() == EWaterBodyType::Ocean);
-
-			auto GetMaterialProxy = [&](UMaterialInstanceDynamic* MID, bool bUseFallback) -> FMaterialRenderProxy*
-			{
-				UMaterialInterface* MaterialInterface = nullptr;
-				if (!MID || !IsMaterialUsedWithWater(MID))
-				{
-					MaterialInterface = bUseFallback ? UMaterial::GetDefaultMaterial(MD_Surface) : nullptr;
-				}
-				else
-				{
-					WaterBodyComponent->SetDynamicParametersOnMID(MID);
-					MaterialInterface = MID;
-				}
-				if (MaterialInterface)
-				{
-					UsedMaterials.Add(MaterialInterface);
-					return MaterialInterface->GetRenderProxy();
-				}
-				return nullptr;
-			};
-			
-			const bool bIsRiver = WaterBodyComponent->GetWaterBodyType() == EWaterBodyType::River;
-
-			FWaterQuadTreeBuilder::FWaterBody WaterBody = {};
-			WaterBody.Material = GetMaterialProxy(WaterBodyComponent->GetWaterMaterialInstance(), true);
-			WaterBody.RiverToLakeMaterial = bIsRiver ? GetMaterialProxy(WaterBodyComponent->GetRiverToLakeTransitionMaterialInstance(), false) : nullptr;
-			WaterBody.RiverToOceanMaterial = bIsRiver ? GetMaterialProxy(WaterBodyComponent->GetRiverToOceanTransitionMaterialInstance(), false) : nullptr;
-			WaterBody.StaticMeshRenderData = StaticMeshRenderData;
-			WaterBody.LocalToWorld = WaterBodyComponent->GetComponentTransform();
-			WaterBody.Bounds = WaterBodyComponent->Bounds;
-			WaterBody.OverlapMaterialPriority = WaterBodyComponent->GetOverlapMaterialPriority();
-			WaterBody.Type = WaterBodyComponent->GetWaterBodyType();
-			WaterBody.WaterBodyIndex = WaterBodyComponent->GetWaterBodyIndex();
-			WaterBody.SurfaceBaseHeight = WaterBodyComponent->GetComponentLocation().Z;
-			WaterBody.MaxWaveHeight = WaterBodyComponent->GetMaxWaveHeight();
-#if WITH_WATER_SELECTION_SUPPORT
-			WaterBody.HitProxy = new HActor(/*InActor = */Actor, /*InPrimComponent = */nullptr);
-			WaterBody.bWaterBodySelected = Actor->IsSelected();
-#endif // WITH_WATER_SELECTION_SUPPORT
-
-			if (!bIsGPUQuadTree)
-			{
-				switch (WaterBody.Type)
-				{
-				case EWaterBodyType::River:
-				{
-					TRACE_CPUPROFILER_EVENT_SCOPE(River);
-
-					TArray<FBox, TInlineAllocator<16>> Boxes;
-					TArray<UPrimitiveComponent*> CollisionComponents = WaterBodyComponent->GetCollisionComponents();
-					for (UPrimitiveComponent* Comp : CollisionComponents)
-					{
-						if (UBodySetup* BodySetup = (Comp != nullptr) ? Comp->GetBodySetup() : nullptr)
-						{
-							// Go through all sub shapes on the bodysetup to get a tight fit along water body
-							for (const FKConvexElem& ConvElem : BodySetup->AggGeom.ConvexElems)
-							{
-								TRACE_CPUPROFILER_EVENT_SCOPE(Add);
-
-								FBox SubBox = ConvElem.ElemBox.TransformBy(Comp->GetComponentTransform().ToMatrixWithScale());
-								SubBox.Max.Z += WaterBodyComponent->GetMaxWaveHeight();
-
-								Boxes.Add(SubBox);
-							}
-						}
-						else
-						{
-							// fallback on global AABB: 
-							FVector Center;
-							FVector Extent;
-							Actor->GetActorBounds(false, Center, Extent);
-							FBox Box(FBox::BuildAABB(Center, Extent));
-							Box.Max.Z += WaterBodyComponent->GetMaxWaveHeight();
-							Boxes.Add(Box);
-						}
-					}
-
-					for (const FBox& Box : Boxes)
-					{
-#if !(UE_BUILD_SHIPPING || UE_BUILD_TEST)
-						if (!!CVarWaterMeshShowTileGenerationGeometry.GetValueOnGameThread())
-						{
-							DrawDebugBox(GetWorld(), Box.GetCenter(), Box.GetExtent(), FColor::Red);
-						}
-#endif
-					}
-					WaterBody.RiverBoxes = Boxes;
-					break;
-				}
-				case EWaterBodyType::Lake:
-				{
-					TRACE_CPUPROFILER_EVENT_SCOPE(Lake);
-
-					const UWaterSplineComponent* SplineComp = WaterBodyComponent->GetWaterSpline();
-					const int32 NumOriginalSplinePoints = SplineComp->GetNumberOfSplinePoints();
-					float ConstantZ = SplineComp->GetLocationAtDistanceAlongSpline(0.0f, ESplineCoordinateSpace::World).Z;
-
-					FBox LakeBounds = Actor->GetComponentsBoundingBox(/* bNonColliding = */true);
-					LakeBounds.Max.Z += WaterBodyComponent->GetMaxWaveHeight();
-
-					// Skip lakes with less than 3 spline points
-					if (NumOriginalSplinePoints < 3)
-					{
-						return true;
-					}
-
-					TArray<TArray<FVector2D>> PolygonBatches;
-					// Reuse the convex hulls generated for the physics shape because the work has already been done, but we can fallback to a simple spline evaluation method in case there's no physics:
-					bool bUseFallbackMethod = true;
-
-					TArray<UPrimitiveComponent*> CollisionComponents = WaterBodyComponent->GetCollisionComponents();
-					if (CollisionComponents.Num() > 0)
-					{
-						UPrimitiveComponent* Comp = CollisionComponents[0];
-						if (UBodySetup* BodySetup = (Comp != nullptr) ? Comp->GetBodySetup() : nullptr)
-						{
-							FTransform CompTransform = Comp->GetComponentTransform();
-							// Go through all sub shapes on the bodysetup to get a tight fit along water body
-							for (const FKConvexElem& ConvElem : BodySetup->AggGeom.ConvexElems)
-							{
-								TRACE_CPUPROFILER_EVENT_SCOPE(Add);
-
-								// Vertex data contains the bottom vertices first, then the top ones :
-								int32 TotalNumVertices = ConvElem.VertexData.Num();
-								check(TotalNumVertices % 2 == 0);
-								int32 NumVertices = TotalNumVertices / 2;
-								if (NumVertices > 0)
-								{
-									bUseFallbackMethod = false;
-
-									// Because the physics shape is made of multiple convex hulls, we cannot simply add their vertices to one big list but have to have 1 batch of polygons per convex hull
-									TArray<FVector2D>& Polygon = PolygonBatches.Emplace_GetRef();
-									Polygon.SetNum(NumVertices);
-									for (int32 i = 0; i < NumVertices; ++i)
-									{
-										// Gather the top vertices :
-										Polygon[i] = FVector2D(CompTransform.TransformPosition(ConvElem.VertexData[NumVertices + i]));
-									}
-								}
-							}
-						}
-					}
-
-					if (bUseFallbackMethod)
-					{
-						TArray<FVector> PolyLineVertices;
-						SplineComp->ConvertSplineToPolyLine(ESplineCoordinateSpace::World, FMath::Square(CVarWaterSplineResampleMaxDistance.GetValueOnGameThread()), PolyLineVertices);
-
-						TArray<FVector2D>& Polygon = PolygonBatches.Emplace_GetRef();
-						Polygon.Reserve(PolyLineVertices.Num());
-						Algo::Transform(PolyLineVertices, Polygon, [](const FVector& Vertex) { return FVector2D(Vertex); });
-					}
-
-					for (const TArray<FVector2D>& Polygon : PolygonBatches)
-					{
-#if !(UE_BUILD_SHIPPING || UE_BUILD_TEST)
-						if (!!CVarWaterMeshShowTileGenerationGeometry.GetValueOnGameThread())
-						{
-							float Z = SplineComp->GetLocationAtDistanceAlongSpline(0.0f, ESplineCoordinateSpace::World).Z;
-							int32 NumVertices = Polygon.Num();
-							for (int32 i = 0; i < NumVertices; i++)
-							{
-								const FVector2D& Point0 = Polygon[i];
-								const FVector2D& Point1 = Polygon[(i + 1) % NumVertices];
-								DrawDebugLine(GetWorld(), FVector(Point0.X, Point0.Y, Z), FVector(Point1.X, Point1.Y, Z), FColor::Green);
-							}
-						}
-#endif // !(UE_BUILD_SHIPPING || UE_BUILD_TEST)
-					}
-
-					WaterBody.PolygonBounds = LakeBounds;
-					WaterBody.PolygonBatches = MoveTemp(PolygonBatches);
-
-					break;
-				}
-				case EWaterBodyType::Ocean:
-				{
-					TRACE_CPUPROFILER_EVENT_SCOPE(Ocean);
-
-					const UWaterSplineComponent* SplineComp = WaterBodyComponent->GetWaterSpline();
-					const int32 NumOriginalSplinePoints = SplineComp->GetNumberOfSplinePoints();
-
-					// Skip oceans with less than 3 spline points
-					if (NumOriginalSplinePoints < 3)
-					{
-						return true;
-					}
-
-					TArray<FVector2D> Polygon;
-
-					TArray<FVector> PolyLineVertices;
-					SplineComp->ConvertSplineToPolyLine(ESplineCoordinateSpace::World, FMath::Square(CVarWaterSplineResampleMaxDistance.GetValueOnGameThread()), PolyLineVertices);
-
-					Polygon.Reserve(PolyLineVertices.Num());
-					Algo::Transform(PolyLineVertices, Polygon, [](const FVector& Vertex) { return FVector2D(Vertex); });
-
-					FBox OceanBounds = WaterBodyComponent->Bounds.GetBox();
-					OceanBounds.Max.Z += WaterBodyComponent->GetMaxWaveHeight();
-
-#if !(UE_BUILD_SHIPPING || UE_BUILD_TEST)
-					if (!!CVarWaterMeshShowTileGenerationGeometry.GetValueOnGameThread())
-					{
-						float Z = SplineComp->GetLocationAtDistanceAlongSpline(0.0f, ESplineCoordinateSpace::World).Z;
-						int32 NumVertices = Polygon.Num();
-						for (int32 i = 0; i < NumVertices; i++)
-						{
-							const FVector2D& Point0 = Polygon[i];
-							const FVector2D& Point1 = Polygon[(i + 1) % NumVertices];
-							DrawDebugLine(GetWorld(), FVector(Point0.X, Point0.Y, Z), FVector(Point1.X, Point1.Y, Z), FColor::Blue);
-						}
-
-						DrawDebugBox(GetWorld(), OceanBounds.GetCenter(), OceanBounds.GetExtent(), FColor::Blue);
-					}
-#endif //!(UE_BUILD_SHIPPING || UE_BUILD_TEST)
-
-					WaterBody.PolygonBounds = OceanBounds;
-					WaterBody.PolygonBatches.Add(MoveTemp(Polygon));
-
-					break;
-				}
-				case EWaterBodyType::Transition:
-					// Transitions dont require rendering
-					break;
-				default:
-					ensureMsgf(false, TEXT("This water body type is not implemented and will not produce any water tiles. "));
-				}
-			}
-
-			WaterQuadTreeBuilder.AddWaterBody(MoveTemp(WaterBody));
-
-			return true;
-		});
-
-		if (bUsesFarDistanceMaterial && FarDistanceMaterialProxy && (FarDistanceMeshExtent > 0.0f))
-		{
-			UsedMaterials.Add(FarDistanceMaterial);
-		}
-
-		MarkRenderStateDirty();
-
-		// Force another rebuild next frame if water body meshes are still compiling
-		if (bIsGPUQuadTree && bAnyWaterMeshesNotReady)
-		{
-			bNeedsRebuild = true;
-		}
-
-		return;
-	}
-
 	// Position snapped to the grid
 	const FVector2D GridPosition = WaterZone->IsLocalOnlyTessellationEnabled() ? GetDynamicWaterMeshCenter() : FVector2D(FMath::GridSnap<FVector::FReal>(GetComponentLocation().X, InTileSize), FMath::GridSnap<FVector::FReal>(GetComponentLocation().Y, InTileSize));
 	const FVector2D WorldExtent = FVector2D(InTileSize * InExtentInTiles.X, InTileSize * InExtentInTiles.Y);
@@ -582,6 +268,8 @@ void UWaterMeshComponent::RebuildWaterMesh(float InTileSize, const FIntPoint& In
 		return;
 	}
 
+	const bool bIsGPUQuadTree = CVarWaterMeshGPUQuadTree.GetValueOnGameThread() != 0;
+
 	// This resets the tree to an initial state, ready for node insertion
 	WaterQuadTree.InitTree(WaterWorldBox, InTileSize, InExtentInTiles, bIsGPUQuadTree);
 
@@ -660,6 +348,9 @@ void UWaterMeshComponent::RebuildWaterMesh(float InTileSize, const FIntPoint& In
 		InMID->SetTextureParameterValue(UWaterBodyComponent::WaterVelocityAndHeightName, OwningZone->WaterInfoTexture);
 		return true;
 	};
+	
+
+	bool bAnyWaterMeshesNotReady = false;
 
 	OwningZone->ForEachWaterBodyComponent([this, bIsGPUQuadTree, WaterWorldBox, bIsFlooded, GlobalOceanHeight, OceanFlood, &FarMeshHeight, &bHasOcean, &SortedPriorities, &bAnyWaterMeshesNotReady, &SetDynamicParametersOnSharedMID](UWaterBodyComponent* WaterBodyComponent)
 	{
@@ -758,9 +449,9 @@ void UWaterMeshComponent::RebuildWaterMesh(float InTileSize, const FIntPoint& In
 					}
 				}
 
-				RenderData.Material = OutMaterials[0] ? OutMaterials[0]->GetRenderProxy() : nullptr;
-				RenderData.RiverToLakeMaterial = OutMaterials[1] ? OutMaterials[1]->GetRenderProxy() : nullptr;
-				RenderData.RiverToOceanMaterial = OutMaterials[2] ? OutMaterials[2]->GetRenderProxy() : nullptr;
+				RenderData.Material = OutMaterials[0];
+				RenderData.RiverToLakeMaterial = OutMaterials[1];
+				RenderData.RiverToOceanMaterial = OutMaterials[2];
 			}
 
 			RenderData.Priority = static_cast<int16>(FMath::Clamp(WaterBodyComponent->GetOverlapMaterialPriority(), MinWaterBodyPriority, MaxWaterBodyPriority));
@@ -785,23 +476,22 @@ void UWaterMeshComponent::RebuildWaterMesh(float InTileSize, const FIntPoint& In
 		{
 			// Assign material instance
 			UMaterialInstanceDynamic* WaterMaterial = WaterBodyComponent->GetWaterMaterialInstance();
-			RenderData.Material = WaterMaterial->GetRenderProxy();
+			RenderData.Material = WaterMaterial;
 
 			if (RenderData.Material)
 			{
-				if (!IsMaterialUsedWithWater(WaterMaterial))
+				if (!IsMaterialUsedWithWater(RenderData.Material))
 				{
-					RenderData.Material = UMaterial::GetDefaultMaterial(MD_Surface)->GetRenderProxy();
-					// Add material so that the component keeps track of all potential materials used
-					UsedMaterials.Add(UMaterial::GetDefaultMaterial(MD_Surface));
+					RenderData.Material = UMaterial::GetDefaultMaterial(MD_Surface);
 				}
 				else
 				{
 					// Add ocean height as a scalar parameter
 					WaterBodyComponent->SetDynamicParametersOnMID(WaterMaterial);
-					// Add material so that the component keeps track of all potential materials used
-					UsedMaterials.Add(WaterMaterial);
 				}
+
+				// Add material so that the component keeps track of all potential materials used
+				UsedMaterials.Add(RenderData.Material);
 			}
 
 			RenderData.Priority = static_cast<int16>(FMath::Clamp(WaterBodyComponent->GetOverlapMaterialPriority(), MinWaterBodyPriority, MaxWaterBodyPriority));
@@ -828,8 +518,8 @@ void UWaterMeshComponent::RebuildWaterMesh(float InTileSize, const FIntPoint& In
 				UMaterialInstanceDynamic* RiverToLakeMaterial = WaterBodyComponent->GetRiverToLakeTransitionMaterialInstance();
 				if (IsMaterialUsedWithWater(RiverToLakeMaterial))
 				{
-					RenderData.RiverToLakeMaterial = RiverToLakeMaterial->GetRenderProxy();
-					UsedMaterials.Add(RiverToLakeMaterial);
+					RenderData.RiverToLakeMaterial = RiverToLakeMaterial;
+					UsedMaterials.Add(RenderData.RiverToLakeMaterial);
 					// Add ocean height as a scalar parameter
 					WaterBodyComponent->SetDynamicParametersOnMID(RiverToLakeMaterial);
 				}
@@ -837,8 +527,8 @@ void UWaterMeshComponent::RebuildWaterMesh(float InTileSize, const FIntPoint& In
 				UMaterialInstanceDynamic* RiverToOceanMaterial = WaterBodyComponent->GetRiverToOceanTransitionMaterialInstance();
 				if (IsMaterialUsedWithWater(RiverToOceanMaterial))
 				{
-					RenderData.RiverToOceanMaterial = RiverToOceanMaterial->GetRenderProxy();
-					UsedMaterials.Add(RiverToOceanMaterial);
+					RenderData.RiverToOceanMaterial = RiverToOceanMaterial;
+					UsedMaterials.Add(RenderData.RiverToOceanMaterial);
 					// Add ocean height as a scalar parameter
 					WaterBodyComponent->SetDynamicParametersOnMID(RiverToOceanMaterial);
 				}
@@ -1103,7 +793,7 @@ void UWaterMeshComponent::RebuildWaterMesh(float InTileSize, const FIntPoint& In
 		// Far Mesh should stitch to the edge of the water zone
 		const FBox2D FarMeshBounds = WaterZone->GetZoneBounds2D();
 
-		WaterQuadTree.AddFarMesh(FarDistanceMaterial->GetRenderProxy(), FarMeshBounds, FarDistanceMeshExtent, FarMeshHeight);
+		WaterQuadTree.AddFarMesh(FarDistanceMaterial, FarMeshBounds, FarDistanceMeshExtent, FarMeshHeight);
 	}
 
 	// Remove all materials from the MaterialToMID map that aren't currently used
diff --git a/Engine/Plugins/Experimental/Water/Source/Runtime/Private/WaterMeshSceneProxy.cpp b/Engine/Plugins/Experimental/Water/Source/Runtime/Private/WaterMeshSceneProxy.cpp
--- a/Engine/Plugins/Experimental/Water/Source/Runtime/Private/WaterMeshSceneProxy.cpp
+++ b/Engine/Plugins/Experimental/Water/Source/Runtime/Private/WaterMeshSceneProxy.cpp
@@ -165,23 +165,8 @@ FWaterMeshSceneProxy::FWaterMeshSceneProxy(UWaterMeshComponent* Component)
 	: FPrimitiveSceneProxy(Component)
 	, MaterialRelevance(Component->GetWaterMaterialRelevance(GetScene().GetFeatureLevel()))
 {
-	AWaterZone* WaterZone = CastChecked<AWaterZone>(Component->GetOwner());
-	checkf(WaterZone != nullptr, TEXT("WaterMeshComponent is owned by an actor that is not a WaterZone. This is not supported!"));
-
-	if (Component->ShouldBuildQuadTreeInSceneProxy())
-	{
-		// Get the QuadTreeBuilder and build the quadtree
-		const FWaterQuadTreeBuilder& WaterQuadTreeBuilder = Component->GetWaterQuadTreeBuilder();
-		const FVector2D GridPosition = WaterZone->IsLocalOnlyTessellationEnabled() ?
-			Component->GetDynamicWaterMeshCenter() :
-			FVector2D(FMath::GridSnap<FVector::FReal>(Component->GetComponentLocation().X, WaterQuadTreeBuilder.GetLeafSize()), FMath::GridSnap<FVector::FReal>(Component->GetComponentLocation().Y, WaterQuadTreeBuilder.GetLeafSize()));
-		WaterQuadTreeBuilder.BuildWaterQuadTree(WaterQuadTree, GridPosition);
-	}
-	else
-	{
-		WaterQuadTree = Component->GetWaterQuadTree();
-	}
-
+	// Cache the tiles and settings
+	WaterQuadTree = Component->GetWaterQuadTree();
 	// Leaf size * 0.5 equals the tightest possible LOD Scale that doesn't break the morphing. Can be scaled larger
 	LODScale = WaterQuadTree.GetLeafSize() * FMath::Max(Component->GetLODScale(), 0.5f);
 
@@ -268,8 +253,11 @@ FWaterMeshSceneProxy::FWaterMeshSceneProxy(UWaterMeshComponent* Component)
 
 	WaterQuadTree.BuildMaterialIndices();
 
-	const FBox WaterInfoBounds3D = WaterZone->GetDynamicWaterInfoBounds();
-	WaterInfoBounds = FBox2D(FVector2D(WaterInfoBounds3D.Min), FVector2D(WaterInfoBounds3D.Max));
+	if (const AWaterZone* WaterZone = Component->GetOwner<AWaterZone>(); ensureMsgf(WaterZone != nullptr, TEXT("WaterMeshComponent is owned by an actor that is not a WaterZone. This is not supported!")))
+	{
+		const FBox WaterInfoBounds3D = WaterZone->GetDynamicWaterInfoBounds();
+		WaterInfoBounds = FBox2D(FVector2D(WaterInfoBounds3D.Min), FVector2D(WaterInfoBounds3D.Max));
+	}
 
 #if RHI_RAYTRACING
 	RayTracingWaterData.SetNum(DensityCount);
diff --git a/Engine/Plugins/Experimental/Water/Source/Runtime/Private/WaterQuadTree.cpp b/Engine/Plugins/Experimental/Water/Source/Runtime/Private/WaterQuadTree.cpp
--- a/Engine/Plugins/Experimental/Water/Source/Runtime/Private/WaterQuadTree.cpp
+++ b/Engine/Plugins/Experimental/Water/Source/Runtime/Private/WaterQuadTree.cpp
@@ -665,7 +665,7 @@ void FWaterQuadTree::AddLake(const TArray<FVector2D>& InPoly, const FBox& InLake
 	AddLakeRecursive(InPoly, LakeBounds, FVector2D(InLakeBounds.Min.Z, InLakeBounds.Max.Z), true, TreeDepth * 2, InWaterBodyIndex);
 }
 
-void FWaterQuadTree::AddFarMesh(FMaterialRenderProxy* InFarMeshMaterial, const FBox2D& InInnerRegion, double InFarDistanceMeshExtent, double InFarDistanceMeshHeight)
+void FWaterQuadTree::AddFarMesh(const UMaterialInterface* InFarMeshMaterial, const FBox2D& InInnerRegion, double InFarDistanceMeshExtent, double InFarDistanceMeshHeight)
 {
 	// Checking for not being read only here to keep things consistent with the other Add functions. In reality the FarMesh isn't added to the QuadTree itself, so it could technically be done whenever.
 	ensure(!bIsReadOnly);
@@ -715,12 +715,14 @@ void FWaterQuadTree::BuildMaterialIndices()
 	int32 NextIdx = 0;
 	TMap<FMaterialRenderProxy*, int32> MatToIdxMap;
 
-	auto GetMatIdx = [&NextIdx, &MatToIdxMap](FMaterialRenderProxy* MaterialRenderProxy)
+	auto GetMatIdx = [&NextIdx, &MatToIdxMap](const UMaterialInterface* Material)
 	{
-		if (!MaterialRenderProxy)
+		if (!Material)
 		{
 			return (int32)INDEX_NONE;
 		}
+		FMaterialRenderProxy* MaterialRenderProxy = Material->GetRenderProxy();
+		check(MaterialRenderProxy != nullptr);
 		const int32* Found = MatToIdxMap.Find(MaterialRenderProxy);
 		if (!Found)
 		{
diff --git a/Engine/Plugins/Experimental/Water/Source/Runtime/Private/WaterQuadTreeBuilder.cpp b/Engine/Plugins/Experimental/Water/Source/Runtime/Private/WaterQuadTreeBuilder.cpp
deleted file mode 100644
diff --git a/Engine/Plugins/Experimental/Water/Source/Runtime/Public/WaterMeshComponent.h b/Engine/Plugins/Experimental/Water/Source/Runtime/Public/WaterMeshComponent.h
--- a/Engine/Plugins/Experimental/Water/Source/Runtime/Public/WaterMeshComponent.h
+++ b/Engine/Plugins/Experimental/Water/Source/Runtime/Public/WaterMeshComponent.h
@@ -4,7 +4,6 @@
 
 #include "Components/MeshComponent.h"
 #include "WaterQuadTree.h"
-#include "WaterQuadTreeBuilder.h"
 #include "WaterMeshComponent.generated.h"
 
 struct FPSOPrecacheParams;
@@ -50,8 +49,6 @@ public:
 	/** Use this instead of GetMaterialRelevance, since this one will go over all materials from all tiles */
 	FMaterialRelevance GetWaterMaterialRelevance(ERHIFeatureLevel::Type InFeatureLevel) const;
 
-	bool ShouldBuildQuadTreeInSceneProxy() const;
-	const FWaterQuadTreeBuilder& GetWaterQuadTreeBuilder() const { return WaterQuadTreeBuilder; }
 	const FWaterQuadTree& GetWaterQuadTree() const { return WaterQuadTree; }
 
 	const TSet<TObjectPtr<UMaterialInterface>>& GetUsedMaterialsSet() const { return UsedMaterials; }
@@ -109,8 +106,6 @@ private:
 	/** The current center of the dynamic water mesh. Updated by the water view extension whenever the view location crosses the update bounds. */
 	FVector2D DynamicWaterMeshCenter = FVector2D::ZeroVector;
 
-	FWaterQuadTreeBuilder WaterQuadTreeBuilder;
-
 	/** Tiles containing water, stored in a quad tree */
 	FWaterQuadTree WaterQuadTree;
 
diff --git a/Engine/Plugins/Experimental/Water/Source/Runtime/Public/WaterQuadTree.h b/Engine/Plugins/Experimental/Water/Source/Runtime/Public/WaterQuadTree.h
--- a/Engine/Plugins/Experimental/Water/Source/Runtime/Public/WaterQuadTree.h
+++ b/Engine/Plugins/Experimental/Water/Source/Runtime/Public/WaterQuadTree.h
@@ -19,11 +19,11 @@ struct FWaterBodyRenderData
 	bool IsOcean() const { return WaterBodyType == 2; }
 
 	/** The standard material to be used for this water body */
-	FMaterialRenderProxy* Material = nullptr;
+	UMaterialInterface* Material = nullptr;
 
 	/** Transition materials (Only set up for river water bodies) */
-	FMaterialRenderProxy* RiverToLakeMaterial = nullptr;
-	FMaterialRenderProxy* RiverToOceanMaterial = nullptr;
+	UMaterialInterface* RiverToLakeMaterial = nullptr;
+	UMaterialInterface* RiverToOceanMaterial = nullptr;
 
 	/** World Z position of the waterbody, this is where the tiles for this water body will be rendered*/
 	double SurfaceBaseHeight = 0.0;
@@ -156,7 +156,7 @@ struct FWaterQuadTree
 	void AddLake(const TArray<FVector2D>& InPoly, const FBox& InLakeBounds, uint32 InWaterBodyIndex);
 
 	/** Add an automatically generated mesh (8 quads) skirt around InInnerRegion which extends out InFarDistanceMeshExtent, is placed at Z value InFarDistanceMeshHeight and is rendered using InFarMeshMaterial */
-	void AddFarMesh(FMaterialRenderProxy* InFarMeshMaterial, const FBox2D& InInnerRegion, double InFarDistanceMeshExtent, double InFarDistanceMeshHeight);
+	void AddFarMesh(const UMaterialInterface* InFarMeshMaterial, const FBox2D& InInnerRegion, double InFarDistanceMeshExtent, double InFarDistanceMeshHeight);
 
 	/** Assign an index to each material */
 	void BuildMaterialIndices();
@@ -335,7 +335,7 @@ private:
 		TArray<FFarMeshInstanceData> InstanceData;
 
 		/** Material for the Far Distance Mesh, its material render proxy will be cached in WaterMaterials when BuildMaterialIndices is called */
-		FMaterialRenderProxy* Material = nullptr;
+		const UMaterialInterface* Material = nullptr;
 
 		void Clear()
 		{
diff --git a/Engine/Plugins/Experimental/Water/Source/Runtime/Public/WaterQuadTreeBuilder.h b/Engine/Plugins/Experimental/Water/Source/Runtime/Public/WaterQuadTreeBuilder.h
deleted file mode 100644
