Fix water not rendering in MRQ due to RM1 being incompatible with the way MRQ views are constructed.

diff --git a/Engine/Plugins/Experimental/Water/Source/Runtime/Private/WaterZoneActor.cpp b/Engine/Plugins/Experimental/Water/Source/Runtime/Private/WaterZoneActor.cpp
--- a/Engine/Plugins/Experimental/Water/Source/Runtime/Private/WaterZoneActor.cpp
+++ b/Engine/Plugins/Experimental/Water/Source/Runtime/Private/WaterZoneActor.cpp
@@ -39,13 +39,36 @@ TAutoConsoleVariable<float> CVarWaterFallbackDepth(
 	TEXT("Depth to use for all water when there are no ground actors defined."),
 	ECVF_Default);
 
+void OnSkipWaterInfoTextureRenderWhenWorldRenderingDisabled_Callback(IConsoleVariable*);
+
 // HACK [jonathan.bard] : (details underneath)
 TAutoConsoleVariable<int32> CVarSkipWaterInfoTextureRenderWhenWorldRenderingDisabled(
 	TEXT("r.Water.SkipWaterInfoTextureRenderWhenWorldRenderingDisabled"),
 	1,
 	TEXT("Use this to prevent the water info from rendering when world rendering is disabled."),
+	FConsoleVariableDelegate::CreateStatic(OnSkipWaterInfoTextureRenderWhenWorldRenderingDisabled_Callback),
 	ECVF_Default);
 
+
+// HACK [roey.borsteinas] : This hack piggy-backs off the above hack for disabled world rendering. See details for that first.
+// MRQ creates it's own View/ViewFamily when rendering which breaks compatibility with the WaterInfo RenderMethod 1.
+// RenderMethod 1 tries to enqueue the water info update directly on the View but since the ULocalPlayer view is discarded before rendering, the water info texture never renders.
+void OnSkipWaterInfoTextureRenderWhenWorldRenderingDisabled_Callback(IConsoleVariable*)
+{
+	static int PreviousWaterInfoRenderMethodValue = 1;
+	IConsoleVariable* WaterInfoRenderMethodCVar = IConsoleManager::Get().FindConsoleVariable(TEXT("r.Water.WaterInfo.RenderMethod"));
+	if (CVarSkipWaterInfoTextureRenderWhenWorldRenderingDisabled.GetValueOnAnyThread() == 0)
+	{
+		PreviousWaterInfoRenderMethodValue = WaterInfoRenderMethodCVar->GetInt();
+		WaterInfoRenderMethodCVar->Set(2);
+	}
+	else
+	{
+		WaterInfoRenderMethodCVar->Set(PreviousWaterInfoRenderMethodValue);
+	}
+
+}
+
 AWaterZone::AWaterZone(const FObjectInitializer& Initializer)
 	: Super(Initializer)
 	, RenderTargetResolution(512, 512)
