Fix moudles not having AutoRTFM enabled for them

diff --git a/Engine/Plugins/Animation/ControlRig/Source/ControlRig/ControlRig.Build.cs b/Engine/Plugins/Animation/ControlRig/Source/ControlRig/ControlRig.Build.cs
--- a/Engine/Plugins/Animation/ControlRig/Source/ControlRig/ControlRig.Build.cs
+++ b/Engine/Plugins/Animation/ControlRig/Source/ControlRig/ControlRig.Build.cs
@@ -63,6 +63,8 @@ namespace UnrealBuildTool.Rules
                 PrivateIncludePathModuleNames.Add("ControlRigEditor");
                 DynamicallyLoadedModuleNames.Add("ControlRigEditor");
             }
+
+			bAllowAutoRTFMInstrumentation = true;
         }
     }
 }
diff --git a/Engine/Plugins/Experimental/Water/Source/Runtime/Water.Build.cs b/Engine/Plugins/Experimental/Water/Source/Runtime/Water.Build.cs
--- a/Engine/Plugins/Experimental/Water/Source/Runtime/Water.Build.cs
+++ b/Engine/Plugins/Experimental/Water/Source/Runtime/Water.Build.cs
@@ -68,5 +68,7 @@ public class Water : ModuleRules
 		}
 		// Add a feature define instead of relying on the generic WITH_EDITOR define
 		PublicDefinitions.Add("WITH_WATER_SELECTION_SUPPORT=" + (bWithWaterSelectionSupport ? 1 : 0));
+
+		bAllowAutoRTFMInstrumentation = true;
 	}
-}
\ No newline at end of file
+}
