Water: fix regression caused by quadtree builder refactor. The WaterMeshComponent may not have a valid quadtree if ShouldBuildQuadTreeInSceneProxy is set to true. Added an ensure to avoid these cases slipping through. This is a temporary state where the quad tree can still be built in the component and will be removed soon.

diff --git a/Engine/Plugins/Experimental/Water/Source/Runtime/Private/WaterMeshComponent.cpp b/Engine/Plugins/Experimental/Water/Source/Runtime/Private/WaterMeshComponent.cpp
--- a/Engine/Plugins/Experimental/Water/Source/Runtime/Private/WaterMeshComponent.cpp
+++ b/Engine/Plugins/Experimental/Water/Source/Runtime/Private/WaterMeshComponent.cpp
@@ -205,6 +205,12 @@ bool UWaterMeshComponent::ShouldBuildQuadTreeInSceneProxy() const
 	return bShouldBuildQuadTreeInSceneProxy;
 }
 
+const FWaterQuadTree& UWaterMeshComponent::GetWaterQuadTree() const
+{
+	ensureMsgf(!ShouldBuildQuadTreeInSceneProxy(), TEXT("Retrieving the water quadtree from the WaterMeshComponent is invalid if the quadtree is built in the scene proxy!"));
+	return WaterQuadTree;
+}
+
 void UWaterMeshComponent::SetDynamicWaterMeshCenter(const FVector2D& NewCenter)
 {
 	if (!DynamicWaterMeshCenter.Equals(NewCenter))
diff --git a/Engine/Plugins/Experimental/Water/Source/Runtime/Private/WaterZoneActor.cpp b/Engine/Plugins/Experimental/Water/Source/Runtime/Private/WaterZoneActor.cpp
--- a/Engine/Plugins/Experimental/Water/Source/Runtime/Private/WaterZoneActor.cpp
+++ b/Engine/Plugins/Experimental/Water/Source/Runtime/Private/WaterZoneActor.cpp
@@ -220,22 +220,19 @@ void AWaterZone::DeclareConstructClasses(TArray<FTopLevelAssetPath>& OutConstruc
 
 void AWaterZone::MarkForRebuild(EWaterZoneRebuildFlags Flags, const FBox2D& UpdateRegion, const UObject* DebugRequestingObject)
 {
-	if (EnumHasAnyFlags(Flags, EWaterZoneRebuildFlags::UpdateWaterMesh))
+	const FBox WaterInfoBounds = GetDynamicWaterInfoBounds();
+	const FBox2D WaterInfoBounds2D(FVector2D(WaterInfoBounds.Min), FVector2D(WaterInfoBounds.Max));
+
+	// Suppress updates which occur outside the bounds of the water zone.
+	if ((!UpdateRegion.bIsValid || UpdateRegion.Intersect(WaterInfoBounds2D)))
 	{
-		const FBox2D WaterQuadTreeBounds = WaterMesh->GetWaterQuadTree().GetTileRegion();
-		// Suppress water mesh updates which occur outside the bounds of the water quad tree.
-		if ((!UpdateRegion.bIsValid) || UpdateRegion.Intersect(WaterQuadTreeBounds))
+		if (EnumHasAnyFlags(Flags, EWaterZoneRebuildFlags::UpdateWaterMesh))
 		{
 			UE_LOG(LogWater, Verbose, TEXT("AWaterZone (%s) UpdateWaterMesh in region {%s} (triggered by %s)"), *GetNameSafe(this), *UpdateRegion.ToString(), *GetNameSafe(DebugRequestingObject));
 			WaterMesh->MarkWaterMeshGridDirty();
 			WaterMesh->MarkRenderStateDirty();
 		}
-	}
-	if (EnumHasAnyFlags(Flags, EWaterZoneRebuildFlags::UpdateWaterInfoTexture))
-	{
-		const FBox WaterInfoBounds = GetDynamicWaterInfoBounds();
-		const FBox2D WaterInfoBounds2D(FVector2D(WaterInfoBounds.Min), FVector2D(WaterInfoBounds.Max));
-		if ((!UpdateRegion.bIsValid) || UpdateRegion.Intersect(WaterInfoBounds2D))
+		if (EnumHasAnyFlags(Flags, EWaterZoneRebuildFlags::UpdateWaterInfoTexture))
 		{
 			UE_LOG(LogWater, Verbose, TEXT("AWaterZone (%s) UpdateWaterInfoTexture in region {%s} (triggered by %s)"), *GetNameSafe(this), *UpdateRegion.ToString(), *GetNameSafe(DebugRequestingObject));
 			bNeedsWaterInfoRebuild = true;
diff --git a/Engine/Plugins/Experimental/Water/Source/Runtime/Public/WaterMeshComponent.h b/Engine/Plugins/Experimental/Water/Source/Runtime/Public/WaterMeshComponent.h
--- a/Engine/Plugins/Experimental/Water/Source/Runtime/Public/WaterMeshComponent.h
+++ b/Engine/Plugins/Experimental/Water/Source/Runtime/Public/WaterMeshComponent.h
@@ -52,7 +52,7 @@ public:
 
 	bool ShouldBuildQuadTreeInSceneProxy() const;
 	const FWaterQuadTreeBuilder& GetWaterQuadTreeBuilder() const { return WaterQuadTreeBuilder; }
-	const FWaterQuadTree& GetWaterQuadTree() const { return WaterQuadTree; }
+	const FWaterQuadTree& GetWaterQuadTree() const;
 
 	const TSet<TObjectPtr<UMaterialInterface>>& GetUsedMaterialsSet() const { return UsedMaterials; }
 
