- Add GetWaterSurfaceInfo and deprecate GetWaterDataAtPoint
- The new function is more consistent with the BP call GetWaterSurfaceInfo, and provides more information which is important for snapping FX to the surface

#rnx
#tests Used water snapping to test the outputs worked as expected

diff --git a/Engine/Plugins/Experimental/Water/Source/Runtime/Private/NiagaraDataInterfaceWater.cpp b/Engine/Plugins/Experimental/Water/Source/Runtime/Private/NiagaraDataInterfaceWater.cpp
--- a/Engine/Plugins/Experimental/Water/Source/Runtime/Private/NiagaraDataInterfaceWater.cpp
+++ b/Engine/Plugins/Experimental/Water/Source/Runtime/Private/NiagaraDataInterfaceWater.cpp
@@ -10,12 +10,82 @@
 
 #define LOCTEXT_NAMESPACE "NiagaraDataInterfaceWater"
 
+struct FNDIWater_InstanceData
+{
+	//Cached ptr to component we sample from. 
+	TWeakObjectPtr<UWaterBodyComponent> WaterBodyComponent;
+	FNiagaraLWCConverter LWCConverter;
+};
 
-namespace WaterFunctionNames
+namespace NDIWaterPrivate
 {
 	const FName GetWaterDataAtPointName(TEXT("GetWaterDataAtPoint"));
+	const FName GetWaterSurfaceInfoName(TEXT("GetWaterSurfaceInfo"));
 
 	const FName GetWaveParamLookupTableName(TEXT("GetWaveParamLookupTableOffset"));
+
+	void GetWaterSurfaceInfo(FVectorVMExternalFunctionContext& Context)
+	{
+		QUICK_SCOPE_CYCLE_COUNTER(NiagaraDataInterfaceWater_GetWaterSurfaceInfo);
+
+		// Inputs
+		VectorVM::FUserPtrHandler<FNDIWater_InstanceData> InstData(Context);
+		FNDIInputParam<bool>				InExecuteQuery(Context);
+		FNDIInputParam<FNiagaraPosition>	InQueryPosition(Context);
+		FNDIInputParam<bool>				InIncludeDepth(Context);
+		FNDIInputParam<bool>				InIncludeWaves(Context);
+		FNDIInputParam<bool>				InSimpleWaves(Context);
+
+		// Outputs
+		FNDIOutputParam<FNiagaraPosition>	OutWaterPlanePosition(Context);
+		FNDIOutputParam<FVector3f>			OutWaterPlaneNormal(Context);
+		FNDIOutputParam<FNiagaraPosition>	OutWaterSurfacePosition(Context);
+		FNDIOutputParam<float>				OutWaterDepth(Context);
+		FNDIOutputParam<FVector3f>			OutWaterVelocity(Context);
+		FNDIOutputParam<bool>				OutInExclusionVolume(Context);
+
+		UWaterBodyComponent* Component = InstData->WaterBodyComponent.Get();
+		for (int32 i = 0; i < Context.GetNumInstances(); ++i)
+		{
+			const bool bExecuteQuery = InExecuteQuery.GetAndAdvance() && Component != nullptr;
+			const FVector3f QueryPosition = InQueryPosition.GetAndAdvance();
+			const bool bIncludeDepth = InIncludeDepth.GetAndAdvance();
+			const bool bIncludeWaves = InIncludeWaves.GetAndAdvance();
+			const bool bSimpleWaves = InSimpleWaves.GetAndAdvance();
+
+			if (bExecuteQuery)
+			{
+				const EWaterBodyQueryFlags QueryFlags =
+					EWaterBodyQueryFlags::ComputeLocation |
+					EWaterBodyQueryFlags::ComputeNormal |
+					EWaterBodyQueryFlags::ComputeImmersionDepth |
+					EWaterBodyQueryFlags::ComputeVelocity |
+					(bIncludeDepth ? EWaterBodyQueryFlags::ComputeDepth : EWaterBodyQueryFlags::None) |
+					(bIncludeWaves ? EWaterBodyQueryFlags::IncludeWaves : EWaterBodyQueryFlags::None) |
+					(bIncludeWaves && bSimpleWaves ? EWaterBodyQueryFlags::SimpleWaves : EWaterBodyQueryFlags::None);
+
+				const FWaterBodyQueryResult QueryResult = Component->QueryWaterInfoClosestToWorldLocation(InstData->LWCConverter.ConvertSimulationPositionToWorld(QueryPosition), QueryFlags);
+
+				const FVector3f WaterPlaneLocation = InstData->LWCConverter.ConvertWorldToSimulationPosition(QueryResult.GetWaterPlaneLocation());
+				const FVector3f WaterSurfacePosition = InstData->LWCConverter.ConvertWorldToSimulationPosition(QueryResult.GetWaterSurfaceLocation());
+
+				OutWaterPlanePosition.SetAndAdvance(WaterPlaneLocation);
+				OutWaterPlaneNormal.SetAndAdvance(FVector3f(QueryResult.GetWaterPlaneNormal()));
+				OutWaterSurfacePosition.SetAndAdvance(WaterSurfacePosition);
+				OutWaterDepth.SetAndAdvance(bIncludeDepth ? QueryResult.GetWaterSurfaceDepth() : 0.0f);
+				OutWaterVelocity.SetAndAdvance(FVector3f(QueryResult.GetVelocity()));
+				OutInExclusionVolume.SetAndAdvance(QueryResult.IsInExclusionVolume());
+			}
+			else
+			{
+				OutWaterPlanePosition.SetAndAdvance(FVector3f::ZeroVector);
+				OutWaterPlaneNormal.SetAndAdvance(FVector3f::UpVector);
+				OutWaterSurfacePosition.SetAndAdvance(FVector3f::ZeroVector);
+				OutWaterDepth.SetAndAdvance(0.0f);
+				OutWaterVelocity.SetAndAdvance(FVector3f::ZeroVector);
+			}
+		}
+	}
 }
 
 struct FNiagaraWaterDIFunctionVersion
@@ -45,39 +115,58 @@ void UNiagaraDataInterfaceWater::PostInitProperties()
 void UNiagaraDataInterfaceWater::GetFunctionsInternal(TArray<FNiagaraFunctionSignature>& OutFunctions) const
 {
 	{
-		FNiagaraFunctionSignature Sig;
-		Sig.Name = WaterFunctionNames::GetWaterDataAtPointName;
+		FNiagaraFunctionSignature& Sig = OutFunctions.AddDefaulted_GetRef();
+		Sig.Name = NDIWaterPrivate::GetWaterSurfaceInfoName;
+		Sig.Inputs.Emplace(FNiagaraTypeDefinition(GetClass()), TEXT("Water"));
+		Sig.Inputs.Emplace_GetRef(FNiagaraTypeDefinition::GetBoolDef(), TEXT("ExecuteQuery")).SetValue(true);
+		Sig.Inputs.Emplace(FNiagaraTypeDefinition::GetPositionDef(), TEXT("QueryPosition"));
+		Sig.Inputs.Emplace_GetRef(FNiagaraTypeDefinition::GetBoolDef(), TEXT("IncludeDepth")).SetValue(true);
+		Sig.Inputs.Emplace_GetRef(FNiagaraTypeDefinition::GetBoolDef(), TEXT("IncludeWaves")).SetValue(true);
+		Sig.Inputs.Emplace_GetRef(FNiagaraTypeDefinition::GetBoolDef(), TEXT("SimpleWaves")).SetValue(true);
+
+		Sig.Outputs.Emplace(FNiagaraTypeDefinition::GetPositionDef(), TEXT("WaterPlanePosition"));
+		Sig.Outputs.Emplace(FNiagaraTypeDefinition::GetVec3Def(), TEXT("WaterPlaneNormal"));
+		Sig.Outputs.Emplace(FNiagaraTypeDefinition::GetPositionDef(), TEXT("WaterSurfacePosition"));
+		Sig.Outputs.Emplace(FNiagaraTypeDefinition::GetFloatDef(), TEXT("WaterDepth"));
+		Sig.Outputs.Emplace(FNiagaraTypeDefinition::GetVec3Def(), TEXT("WaterVelocity"));
+		Sig.Outputs.Emplace(FNiagaraTypeDefinition::GetBoolDef(), TEXT("InExclusionVolume"));
+
+		Sig.bMemberFunction = true;
+		Sig.bExperimental = true;
+		Sig.SetDescription(LOCTEXT("DataInterfaceWater_GetWaterSurfaceInfo", "Get the water surface information at the provided world position."));
+		Sig.SetInputDescription(Sig.Inputs[1], LOCTEXT("DataInterfaceWater_GetWaterSurfaceInfo_ExecuteQuery", "When disabled the water query will not run and the results are invalid / defaulted."));
+	}
+
+	{
+		FNiagaraFunctionSignature& Sig = OutFunctions.AddDefaulted_GetRef();
+		Sig.Name = NDIWaterPrivate::GetWaterDataAtPointName;
 
-		Sig.Inputs.Add(FNiagaraVariable(FNiagaraTypeDefinition(GetClass()), TEXT("Water")));
-		Sig.Inputs.Add(FNiagaraVariable(FNiagaraTypeDefinition::GetPositionDef(), TEXT("WorldPosition")));
-		Sig.Inputs.Add(FNiagaraVariable(FNiagaraTypeDefinition::GetFloatDef(), TEXT("Time")));
+		Sig.Inputs.Emplace(FNiagaraTypeDefinition(GetClass()), TEXT("Water"));
+		Sig.Inputs.Emplace(FNiagaraTypeDefinition::GetPositionDef(), TEXT("WorldPosition"));
+		Sig.Inputs.Emplace(FNiagaraTypeDefinition::GetFloatDef(), TEXT("Time"));
 
-		Sig.Outputs.Add(FNiagaraVariable(FNiagaraTypeDefinition::GetFloatDef(), TEXT("WaveHeight")));
-		Sig.Outputs.Add(FNiagaraVariable(FNiagaraTypeDefinition::GetFloatDef(), TEXT("Depth")));
-		Sig.Outputs.Add(FNiagaraVariable(FNiagaraTypeDefinition::GetVec3Def(), TEXT("Velocity")));
-		Sig.Outputs.Add(FNiagaraVariable(FNiagaraTypeDefinition::GetPositionDef(), TEXT("SurfacePosition")));
-		Sig.Outputs.Add(FNiagaraVariable(FNiagaraTypeDefinition::GetVec3Def(), TEXT("SurfaceNormal")));
+		Sig.Outputs.Emplace(FNiagaraTypeDefinition::GetFloatDef(), TEXT("WaveHeight"));
+		Sig.Outputs.Emplace(FNiagaraTypeDefinition::GetFloatDef(), TEXT("Depth"));
+		Sig.Outputs.Emplace(FNiagaraTypeDefinition::GetVec3Def(), TEXT("Velocity"));
+		Sig.Outputs.Emplace(FNiagaraTypeDefinition::GetPositionDef(), TEXT("SurfacePosition"));
+		Sig.Outputs.Emplace(FNiagaraTypeDefinition::GetVec3Def(), TEXT("SurfaceNormal"));
 
 		Sig.bMemberFunction = true;
-		Sig.bRequiresContext = false;
 		Sig.bExperimental = true;
+		Sig.bSoftDeprecatedFunction = true;
 		Sig.SetDescription(LOCTEXT("DataInterfaceWater_GetWaterDataAtPoint", "Get the water data at the provided world position and time"));
-		OutFunctions.Add(Sig);
 	}
 
 	{
-		FNiagaraFunctionSignature Sig;
-		Sig.Name = WaterFunctionNames::GetWaveParamLookupTableName;
+		FNiagaraFunctionSignature& Sig = OutFunctions.AddDefaulted_GetRef();
+		Sig.Name = NDIWaterPrivate::GetWaveParamLookupTableName;
+		Sig.Inputs.Emplace(FNiagaraTypeDefinition(GetClass()), TEXT("Water"));
 
-		Sig.Inputs.Add(FNiagaraVariable(FNiagaraTypeDefinition(GetClass()), TEXT("Water")));
-
-		Sig.Outputs.Add(FNiagaraVariable(FNiagaraTypeDefinition::GetIntDef(), TEXT("Offset")));
+		Sig.Outputs.Emplace(FNiagaraTypeDefinition::GetIntDef(), TEXT("Offset"));
 
 		Sig.bMemberFunction = true;
-		Sig.bRequiresContext = false;
 		Sig.bExperimental = true;
 		Sig.SetDescription(LOCTEXT("DataInterfaceWater_GetWaveParamLookupTableOffset", "Get the lookup table offset into the wave data texture for the data interface's water body"));
-		OutFunctions.Add(Sig);
 	}
 }
 #endif
@@ -85,23 +174,20 @@ void UNiagaraDataInterfaceWater::GetFunctionsInternal(TArray<FNiagaraFunctionSig
 DEFINE_NDI_DIRECT_FUNC_BINDER(UNiagaraDataInterfaceWater, GetWaterDataAtPoint);
 DEFINE_NDI_DIRECT_FUNC_BINDER(UNiagaraDataInterfaceWater, GetWaveParamLookupTableOffset);
 
-struct FNDIWater_InstanceData
-{
-	//Cached ptr to component we sample from. 
-	TWeakObjectPtr<UWaterBodyComponent> WaterBodyComponent;
-	FNiagaraLWCConverter LWCConverter;
-};
-
-void UNiagaraDataInterfaceWater::GetVMExternalFunction(const FVMExternalFunctionBindingInfo& BindingInfo, void* InstanceData, FVMExternalFunction &OutFunc)
+void UNiagaraDataInterfaceWater::GetVMExternalFunction(const FVMExternalFunctionBindingInfo& BindingInfo, void* InstanceData, FVMExternalFunction& OutFunc)
 {
-	if (BindingInfo.Name == WaterFunctionNames::GetWaterDataAtPointName)
+	if (BindingInfo.Name == NDIWaterPrivate::GetWaterSurfaceInfoName)
+	{
+		OutFunc = FVMExternalFunction::CreateStatic(&NDIWaterPrivate::GetWaterSurfaceInfo);
+	}
+	else if (BindingInfo.Name == NDIWaterPrivate::GetWaterDataAtPointName)
 	{
 		if(BindingInfo.GetNumInputs() == 5 && BindingInfo.GetNumOutputs() == 11)
 		{
 			NDI_FUNC_BINDER(UNiagaraDataInterfaceWater, GetWaterDataAtPoint)::Bind(this, OutFunc);
 		}
 	}
-	else if (BindingInfo.Name == WaterFunctionNames::GetWaveParamLookupTableName)
+	else if (BindingInfo.Name == NDIWaterPrivate::GetWaveParamLookupTableName)
 	{
 		if (BindingInfo.GetNumInputs() == 1 && BindingInfo.GetNumOutputs() == 1)
 		{
@@ -175,7 +261,7 @@ bool UNiagaraDataInterfaceWater::UpgradeFunctionCall(FNiagaraFunctionSignature&
 	// upgrade from lwc changes, only parameter types changed there
 	if (FunctionSignature.FunctionVersion < FNiagaraWaterDIFunctionVersion::LargeWorldCoordinates)
 	{
-		if (FunctionSignature.Name == WaterFunctionNames::GetWaterDataAtPointName && ensure(FunctionSignature.Inputs.Num() == 3) && ensure(FunctionSignature.Outputs.Num() == 5))
+		if (FunctionSignature.Name == NDIWaterPrivate::GetWaterDataAtPointName && ensure(FunctionSignature.Inputs.Num() == 3) && ensure(FunctionSignature.Outputs.Num() == 5))
 		{
 			FunctionSignature.Inputs[1].SetType(FNiagaraTypeDefinition::GetPositionDef());
 			FunctionSignature.Outputs[3].SetType(FNiagaraTypeDefinition::GetPositionDef());
