Water: Fix the RenderCaptureNextWaterInfoDraw cvar so that it properly decrements after each capture. Because the cvar will have been set by console,  the attempt to decrement it will be ignored due to the priority of console being higher than code.

diff --git a/Engine/Plugins/Experimental/Water/Source/Runtime/Private/WaterInfoRendering.cpp b/Engine/Plugins/Experimental/Water/Source/Runtime/Private/WaterInfoRendering.cpp
--- a/Engine/Plugins/Experimental/Water/Source/Runtime/Private/WaterInfoRendering.cpp
+++ b/Engine/Plugins/Experimental/Water/Source/Runtime/Private/WaterInfoRendering.cpp
@@ -722,7 +722,7 @@ void UpdateWaterInfoRendering(
 	if (RenderCaptureNextWaterInfoDraws != 0)
 	{
 		RenderCaptureNextWaterInfoDraws = FMath::Max(0, RenderCaptureNextWaterInfoDraws - 1);
-		CVarRenderCaptureNextWaterInfoDraws->Set(RenderCaptureNextWaterInfoDraws);
+		CVarRenderCaptureNextWaterInfoDraws->SetWithCurrentPriority(RenderCaptureNextWaterInfoDraws);
 	}
 
 	if (!IsValid(Context.TextureRenderTarget) || Scene == nullptr)
@@ -1066,7 +1066,7 @@ void UpdateWaterInfoRendering_CustomRenderPass(
 	if (RenderCaptureNextWaterInfoDraws != 0)
 	{
 		RenderCaptureNextWaterInfoDraws = FMath::Max(0, RenderCaptureNextWaterInfoDraws - 1);
-		CVarRenderCaptureNextWaterInfoDraws->Set(RenderCaptureNextWaterInfoDraws);
+		CVarRenderCaptureNextWaterInfoDraws->SetWithCurrentPriority(RenderCaptureNextWaterInfoDraws);
 		bPerformRenderCapture = true;
 	}
 
diff --git a/Engine/Plugins/Experimental/Water/Source/Runtime/Private/WaterZoneActor.cpp b/Engine/Plugins/Experimental/Water/Source/Runtime/Private/WaterZoneActor.cpp
--- a/Engine/Plugins/Experimental/Water/Source/Runtime/Private/WaterZoneActor.cpp
+++ b/Engine/Plugins/Experimental/Water/Source/Runtime/Private/WaterZoneActor.cpp
@@ -60,11 +60,11 @@ void OnSkipWaterInfoTextureRenderWhenWorldRenderingDisabled_Callback(IConsoleVar
 	if (CVarSkipWaterInfoTextureRenderWhenWorldRenderingDisabled.GetValueOnAnyThread() == 0)
 	{
 		PreviousWaterInfoRenderMethodValue = WaterInfoRenderMethodCVar->GetInt();
-		WaterInfoRenderMethodCVar->Set(2);
+		WaterInfoRenderMethodCVar->SetWithCurrentPriority(2);
 	}
 	else
 	{
-		WaterInfoRenderMethodCVar->Set(PreviousWaterInfoRenderMethodValue);
+		WaterInfoRenderMethodCVar->SetWithCurrentPriority(PreviousWaterInfoRenderMethodValue);
 	}
 
 }
diff --git a/Engine/Source/Runtime/Renderer/Private/WaterInfoTextureRendering.cpp b/Engine/Source/Runtime/Renderer/Private/WaterInfoTextureRendering.cpp
--- a/Engine/Source/Runtime/Renderer/Private/WaterInfoTextureRendering.cpp
+++ b/Engine/Source/Runtime/Renderer/Private/WaterInfoTextureRendering.cpp
@@ -789,7 +789,7 @@ void RenderWaterInfoTexture(
 	if (RenderCaptureNextWaterInfoDraws != 0)
 	{
 		RenderCaptureNextWaterInfoDraws = FMath::Max(0, RenderCaptureNextWaterInfoDraws - 1);
-		CVarRenderCaptureNextWaterInfoDraws->Set(RenderCaptureNextWaterInfoDraws);
+		CVarRenderCaptureNextWaterInfoDraws->SetWithCurrentPriority(RenderCaptureNextWaterInfoDraws);
 	}
 	
 
