Converted WaterInfoTexture into a URenderTargetTexture2DArray to support multiple views and split screen in the future. Also changed the WaterBodyData MF and related material functions to use an index provided by the view uniform buffer to index into the correct texture array layer. The WaterInfoTexture is currently created with just a single layer, but this should now be extensible to multiple layers/views.
ShallowWater and the BP_FluidSim_01 blueprint were fixed up by simply using the new Texture2DArray WaterInfoTexture instead. ShallowWater internals were changed to accept a Texture2DArray instead.

#rnx

#changelist validated
#virtualized

diff --git a/Engine/Plugins/Experimental/Water/Shaders/Private/WaterDataFunctions.ush b/Engine/Plugins/Experimental/Water/Shaders/Private/WaterDataFunctions.ush
--- a/Engine/Plugins/Experimental/Water/Shaders/Private/WaterDataFunctions.ush
+++ b/Engine/Plugins/Experimental/Water/Shaders/Private/WaterDataFunctions.ush
@@ -54,6 +54,11 @@ int GetWaveDataOffset()
 	return (int)View.WaterData[0].y;
 }
 
+int GetWaterInfoTextureViewIndex()
+{
+	return View.WaterInfoTextureViewIndex;
+}
+
 /** Function to decode the raw data from the water indirection buffer */
 FWaterBodyData DecodeWaterBodyData(float4 InDataToDecode0, float4 InDataToDecode1)
 {
diff --git a/Engine/Plugins/Experimental/Water/Source/Runtime/Private/WaterBodyComponent.cpp b/Engine/Plugins/Experimental/Water/Source/Runtime/Private/WaterBodyComponent.cpp
--- a/Engine/Plugins/Experimental/Water/Source/Runtime/Private/WaterBodyComponent.cpp
+++ b/Engine/Plugins/Experimental/Water/Source/Runtime/Private/WaterBodyComponent.cpp
@@ -20,6 +20,7 @@
 #include "WaterRuntimeSettings.h"
 #include "WaterUtils.h"
 #include "Engine/TextureRenderTarget2D.h"
+#include "Engine/TextureRenderTarget2DArray.h"
 #include "Engine/Texture2D.h"
 #include "WaterMeshComponent.h"
 #include "WaterVersion.h"
@@ -1750,7 +1751,7 @@ bool UWaterBodyComponent::SetDynamicParametersOnMID(UMaterialInstanceDynamic* In
 	if (const AWaterZone* WaterZone = GetWaterZone())
 	{
 		InMID->SetScalarParameterValue(WaterZoneIndexParamName, WaterZone->GetWaterZoneIndex());
-		InMID->SetTextureParameterValue(WaterVelocityAndHeightName, WaterZone->WaterInfoTexture);
+		InMID->SetTextureParameterValue(WaterVelocityAndHeightName, WaterZone->WaterInfoTextureArray);
 	}
 
 	return true;
diff --git a/Engine/Plugins/Experimental/Water/Source/Runtime/Private/WaterInfoRendering.cpp b/Engine/Plugins/Experimental/Water/Source/Runtime/Private/WaterInfoRendering.cpp
--- a/Engine/Plugins/Experimental/Water/Source/Runtime/Private/WaterInfoRendering.cpp
+++ b/Engine/Plugins/Experimental/Water/Source/Runtime/Private/WaterInfoRendering.cpp
@@ -9,7 +9,7 @@
 #include "Modules/ModuleManager.h"
 #include "RenderCaptureInterface.h"
 #include "RHIStaticStates.h"
-#include "Engine/TextureRenderTarget2D.h"
+#include "Engine/TextureRenderTarget2DArray.h"
 #include "SceneCaptureRendering.h"
 #include "PostProcess/DrawRectangle.h"
 #include "Math/OrthoMatrix.h"
@@ -990,7 +990,9 @@ public:
 		FinalizeWaterInfo(GraphBuilder, *Views[0]->Family, *Views[0], MergeTargetTexture, FinalizedTexture, Params);
 
 		FRDGTextureRef WaterInfoTexture = RegisterExternalTexture(GraphBuilder, WaterInfoRenderTarget->GetRenderTargetTexture(), TEXT("WaterInfoTexture"));
-		AddCopyTexturePass(GraphBuilder, FinalizedTexture, WaterInfoTexture);
+		FRHICopyTextureInfo CopyInfo;
+		CopyInfo.DestSliceIndex = RenderTargetArraySlice;
+		AddCopyTexturePass(GraphBuilder, FinalizedTexture, WaterInfoTexture, CopyInfo);
 		GraphBuilder.UseExternalAccessMode(WaterInfoTexture, ERHIAccess::SRVMask);
 	}
 
@@ -998,6 +1000,7 @@ public:
 	FWaterInfoRenderingDepthPass* DepthPass = nullptr;
 	FWaterInfoRenderingColorPass* ColorPass = nullptr;
 	FUpdateWaterInfoParams Params;
+	int32 RenderTargetArraySlice = 0;
 };
 
 const FName& GetWaterInfoDilationPassName() { return FWaterInfoRenderingDilationPass::GetTypeNameStatic(); }
diff --git a/Engine/Plugins/Experimental/Water/Source/Runtime/Private/WaterInfoRendering.h b/Engine/Plugins/Experimental/Water/Source/Runtime/Private/WaterInfoRendering.h
--- a/Engine/Plugins/Experimental/Water/Source/Runtime/Private/WaterInfoRendering.h
+++ b/Engine/Plugins/Experimental/Water/Source/Runtime/Private/WaterInfoRendering.h
@@ -7,7 +7,7 @@
 class AWaterZone;
 class UWaterBodyComponent;
 class FSceneInterface;
-class UTextureRenderTarget2D;
+class UTextureRenderTarget2DArray;
 class UPrimitiveComponent;
 class FSceneView;
 class FSceneViewFamily;
@@ -18,7 +18,7 @@ namespace UE::WaterInfo
 struct FRenderingContext
 {
 	AWaterZone* ZoneToRender = nullptr;
-	UTextureRenderTarget2D* TextureRenderTarget;
+	UTextureRenderTarget2DArray* TextureRenderTarget;
 	TArray<UWaterBodyComponent*> WaterBodies;
 	TArray<TWeakObjectPtr<UPrimitiveComponent>> GroundPrimitiveComponents;
 	float CaptureZ;
diff --git a/Engine/Plugins/Experimental/Water/Source/Runtime/Private/WaterMeshComponent.cpp b/Engine/Plugins/Experimental/Water/Source/Runtime/Private/WaterMeshComponent.cpp
--- a/Engine/Plugins/Experimental/Water/Source/Runtime/Private/WaterMeshComponent.cpp
+++ b/Engine/Plugins/Experimental/Water/Source/Runtime/Private/WaterMeshComponent.cpp
@@ -7,7 +7,7 @@
 #include "Materials/Material.h"
 #include "PhysicsEngine/BodySetup.h"
 #include "Engine/StaticMesh.h"
-#include "Engine/TextureRenderTarget2D.h"
+#include "Engine/TextureRenderTarget2DArray.h"
 #include "WaterBodyComponent.h"
 #include "WaterMeshSceneProxy.h"
 #include "WaterModule.h"
@@ -663,7 +663,7 @@ void UWaterMeshComponent::RebuildWaterMesh(float InTileSize, const FIntPoint& In
 		InMID->SetScalarParameterValue(UWaterBodyComponent::WaterBodyIndexParamName, -1);
 		InMID->SetScalarParameterValue(UWaterBodyComponent::GlobalOceanHeightName, WaterSubsystem->GetOceanTotalHeight());
 		InMID->SetScalarParameterValue(UWaterBodyComponent::WaterZoneIndexParamName, OwningZone->GetWaterZoneIndex());
-		InMID->SetTextureParameterValue(UWaterBodyComponent::WaterVelocityAndHeightName, OwningZone->WaterInfoTexture);
+		InMID->SetTextureParameterValue(UWaterBodyComponent::WaterVelocityAndHeightName, OwningZone->WaterInfoTextureArray);
 		return true;
 	};
 
diff --git a/Engine/Plugins/Experimental/Water/Source/Runtime/Private/WaterUtils.cpp b/Engine/Plugins/Experimental/Water/Source/Runtime/Private/WaterUtils.cpp
--- a/Engine/Plugins/Experimental/Water/Source/Runtime/Private/WaterUtils.cpp
+++ b/Engine/Plugins/Experimental/Water/Source/Runtime/Private/WaterUtils.cpp
@@ -3,6 +3,7 @@
 #include "WaterUtils.h"
 #include "Materials/MaterialInstanceDynamic.h"
 #include "Engine/TextureRenderTarget2D.h"
+#include "Engine/TextureRenderTarget2DArray.h"
 #include "UObject/Package.h"
 #include "RenderingThread.h"
 
@@ -89,6 +90,46 @@ UTextureRenderTarget2D* FWaterUtils::GetOrCreateTransientRenderTarget2D(UTexture
 	return NewRenderTarget2D;
 }
 
+WATER_API UTextureRenderTarget2DArray* FWaterUtils::GetOrCreateTransientRenderTarget2DArray(UTextureRenderTarget2DArray* InRenderTarget, FName InRenderTargetName, const FIntPoint& InSize, int32 InSlices, ETextureRenderTargetFormat InFormat, const FLinearColor& InClearColor, bool bInAutoGenerateMipMaps)
+{
+	EPixelFormat PixelFormat = GetPixelFormatFromRenderTargetFormat(InFormat);
+	if ((InSize.X <= 0)
+		|| (InSize.Y <= 0)
+		|| (InSlices <= 0)
+		|| (PixelFormat == EPixelFormat::PF_Unknown))
+	{
+		return nullptr;
+	}
+
+	if (IsValid(InRenderTarget))
+	{
+		if ((InRenderTarget->SizeX == InSize.X)
+			&& (InRenderTarget->SizeY == InSize.Y)
+			&& (InRenderTarget->Slices == InSlices)
+			&& (InRenderTarget->GetFormat() == PixelFormat) // Watch out : GetFormat() returns a EPixelFormat (non-class enum), so we can't compare with a ETextureRenderTargetFormat
+			&& (InRenderTarget->ClearColor == InClearColor))
+		{
+			return InRenderTarget;
+		}
+	}
+
+	UTextureRenderTarget2DArray* NewRenderTarget2DArray = NewObject<UTextureRenderTarget2DArray>(GetTransientPackage(), MakeUniqueObjectName(GetTransientPackage(), UTextureRenderTarget2DArray::StaticClass(), InRenderTargetName));
+	check(NewRenderTarget2DArray);
+	NewRenderTarget2DArray->ClearColor = InClearColor;
+	NewRenderTarget2DArray->Init(InSize.X, InSize.Y, InSlices, PixelFormat);
+	NewRenderTarget2DArray->UpdateResourceImmediate(true);
+
+	// Flush RHI thread after creating texture render target to make sure that RHIUpdateTextureReference is executed before doing any rendering with it
+	// This makes sure that Value->TextureReference.TextureReferenceRHI->GetReferencedTexture() is valid so that FUniformExpressionSet::FillUniformBuffer properly uses the texture for rendering, instead of using a fallback texture
+	ENQUEUE_RENDER_COMMAND(FlushRHIThreadToUpdateTextureRenderTargetReference)(
+	[](FRHICommandListImmediate& RHICmdList)
+	{
+		RHICmdList.ImmediateFlush(EImmediateFlushType::FlushRHIThread);
+	});
+
+	return NewRenderTarget2DArray;
+}
+
 FGuid FWaterUtils::StringToGuid(const FString& InStr)
 {
 	// Compute a 128-bit hash based on the string and use that as a GUID :
diff --git a/Engine/Plugins/Experimental/Water/Source/Runtime/Private/WaterViewExtension.cpp b/Engine/Plugins/Experimental/Water/Source/Runtime/Private/WaterViewExtension.cpp
--- a/Engine/Plugins/Experimental/Water/Source/Runtime/Private/WaterViewExtension.cpp
+++ b/Engine/Plugins/Experimental/Water/Source/Runtime/Private/WaterViewExtension.cpp
@@ -422,6 +422,8 @@ void FWaterViewExtension::PreRenderView_RenderThread(FRDGBuilder& GraphBuilder,
 		InView.WaterDataBuffer = WaterGPUData->AuxDataSRV;
 		InView.WaterIndirectionBuffer = WaterGPUData->WaterBodyDataSRV;
 	}
+	// TODO: Look up an actually unique index for the given view.
+	InView.WaterInfoTextureViewIndex = 0;
 }
 
 void FWaterViewExtension::PreRenderBasePass_RenderThread(FRDGBuilder& GraphBuilder, bool bDepthBufferIsPopulated)
diff --git a/Engine/Plugins/Experimental/Water/Source/Runtime/Private/WaterZoneActor.cpp b/Engine/Plugins/Experimental/Water/Source/Runtime/Private/WaterZoneActor.cpp
--- a/Engine/Plugins/Experimental/Water/Source/Runtime/Private/WaterZoneActor.cpp
+++ b/Engine/Plugins/Experimental/Water/Source/Runtime/Private/WaterZoneActor.cpp
@@ -642,13 +642,13 @@ bool AWaterZone::UpdateWaterInfoTexture()
 		}
 
 		const ETextureRenderTargetFormat Format = bHalfPrecisionTexture ? ETextureRenderTargetFormat::RTF_RGBA16f : RTF_RGBA32f;
-		UTextureRenderTarget2D* OldTexture = WaterInfoTexture;
-		WaterInfoTexture = FWaterUtils::GetOrCreateTransientRenderTarget2D(OldTexture, TEXT("WaterInfoTexture"), RenderTargetResolution, Format);
+		UTextureRenderTarget2DArray* OldTexture = WaterInfoTextureArray;
+		WaterInfoTextureArray = FWaterUtils::GetOrCreateTransientRenderTarget2DArray(OldTexture, TEXT("WaterInfoTexture"), RenderTargetResolution, 1, Format);
 
 		// The water info texture is different, we need to bind the newly created texture to all registered water bodies
-		if (WaterInfoTexture != OldTexture)
+		if (WaterInfoTextureArray != OldTexture)
 		{
-			OnWaterInfoTextureCreated.Broadcast(WaterInfoTexture);
+			OnWaterInfoTextureArrayCreated.Broadcast(WaterInfoTextureArray);
 
 			ForEachWaterBodyComponent([](UWaterBodyComponent* WaterBodyComponent)
 			{
@@ -662,7 +662,7 @@ bool AWaterZone::UpdateWaterInfoTexture()
 		Context.WaterBodies = WaterBodiesToRender;
 		Context.GroundPrimitiveComponents = MoveTemp(GroundPrimitiveComponents);
 		Context.CaptureZ = FMath::Max(WaterZMax, GroundZMax) + CaptureZOffset;
-		Context.TextureRenderTarget = WaterInfoTexture;
+		Context.TextureRenderTarget = WaterInfoTextureArray;
 
 		if (FWaterViewExtension* WaterViewExtension = UWaterSubsystem::GetWaterViewExtension(World))
 		{
diff --git a/Engine/Plugins/Experimental/Water/Source/Runtime/Public/WaterUtils.h b/Engine/Plugins/Experimental/Water/Source/Runtime/Public/WaterUtils.h
--- a/Engine/Plugins/Experimental/Water/Source/Runtime/Public/WaterUtils.h
+++ b/Engine/Plugins/Experimental/Water/Source/Runtime/Public/WaterUtils.h
@@ -9,6 +9,7 @@ enum ETextureRenderTargetFormat : int;
 class UMaterialInstanceDynamic;
 class UMaterialInterface;
 class UTextureRenderTarget2D;
+class UTextureRenderTarget2DArray;
 
 struct FWaterUtils
 {
@@ -38,6 +39,21 @@ struct FWaterUtils
 	static WATER_API UTextureRenderTarget2D* GetOrCreateTransientRenderTarget2D(UTextureRenderTarget2D* InRenderTarget, FName InRenderTargetName, const FIntPoint& InSize, ETextureRenderTargetFormat InFormat, 
 		const FLinearColor& InClearColor = FLinearColor::Black, bool bInAutoGenerateMipMaps = false);
 
+	/**
+	 *  Creates a transient render target of the proper size/format/... or returns the existing one if it's compatible.
+		 InRenderTarget is read-only so it's up to the caller to decide to update or not the RT. Usually, the calling code will be : RT = FWaterUtils::GetOrCreateTransientRenderTarget2D(RT, ...)
+	 * @param InRenderTarget is the current render target currently used. It will be the one returned if it's already compatible with the size/format/...
+	 * @param InRenderTargetName is the UObject's name in case we need to create a MID
+	 * @param InSize is the render target's size
+	 * @param InSlices is the number of slices of the render target
+	 * @param InFormat is the render target's format
+	 * @param InClearColor is the render target's default clear color
+	 * @param bInAutoGenerateMipMaps is for generating mipmaps automatically after rendering to the render target
+	 * @return a compatible transient render target if InMaterialInterface is size/format/... is valid, nullptr otherwise
+	 */
+	static WATER_API UTextureRenderTarget2DArray* GetOrCreateTransientRenderTarget2DArray(UTextureRenderTarget2DArray* InRenderTarget, FName InRenderTargetName, const FIntPoint& InSize, int32 InSlices, ETextureRenderTargetFormat InFormat,
+		const FLinearColor& InClearColor = FLinearColor::Black, bool bInAutoGenerateMipMaps = false);
+
 	static FGuid StringToGuid(const FString& InStr);
 
 	static WATER_API bool IsWaterEnabled(bool bIsRenderThread);
diff --git a/Engine/Plugins/Experimental/Water/Source/Runtime/Public/WaterZoneActor.h b/Engine/Plugins/Experimental/Water/Source/Runtime/Public/WaterZoneActor.h
--- a/Engine/Plugins/Experimental/Water/Source/Runtime/Public/WaterZoneActor.h
+++ b/Engine/Plugins/Experimental/Water/Source/Runtime/Public/WaterZoneActor.h
@@ -8,6 +8,7 @@
 class UBillboardComponent;
 class UTexture2D;
 class UTextureRenderTarget2D;
+class UTextureRenderTarget2DArray;
 namespace EEndPlayReason { enum Type : int; }
 
 class UWaterMeshComponent;
@@ -26,6 +27,7 @@ ENUM_CLASS_FLAGS(EWaterZoneRebuildFlags);
 
 
 DECLARE_DYNAMIC_MULTICAST_DELEGATE_OneParam(FOnWaterInfoTextureCreated, const UTextureRenderTarget2D*, WaterInfoTexture);
+DECLARE_DYNAMIC_MULTICAST_DELEGATE_OneParam(FOnWaterInfoTextureArrayCreated, const UTextureRenderTarget2DArray*, WaterInfoTextureArray);
 
 UCLASS(Blueprintable, HideCategories=(Physics, Replication, Input, Collision))
 class WATER_API AWaterZone : public AActor
@@ -85,9 +87,9 @@ public:
 	int32 GetWaterZoneIndex() const { return WaterZoneIndex; }
 
 	UPROPERTY(Transient, DuplicateTransient, VisibleAnywhere, BlueprintReadOnly, Category = Water)
-	TObjectPtr<UTextureRenderTarget2D> WaterInfoTexture;
+	TObjectPtr<UTextureRenderTarget2DArray> WaterInfoTextureArray;
 
-	FOnWaterInfoTextureCreated& GetOnWaterInfoTextureCreated() { return OnWaterInfoTextureCreated; }
+	FOnWaterInfoTextureArrayCreated& GetOnWaterInfoTextureArrayCreated() { return OnWaterInfoTextureArrayCreated; }
 	
 #if WITH_EDITOR
 	virtual TUniquePtr<class FWorldPartitionActorDesc> CreateClassActorDesc() const override;
@@ -97,6 +99,18 @@ public:
 	UFUNCTION(BlueprintCallable, Category=Rendering)
 	void SetFarMeshMaterial(UMaterialInterface* InFarMaterial);
 
+
+#if WITH_EDITORONLY_DATA
+	UE_DEPRECATED(5.5, "WaterInfoTexture is deprecated, please use WaterInfoTextureArray instead.")
+	UPROPERTY(Transient, meta = (DeprecatedProperty, DeprecationMessage = "Use WaterInfoTextureArray instead."))
+	TObjectPtr<UTextureRenderTarget2D> WaterInfoTexture_DEPRECATED;
+#endif
+
+#if WITH_EDITOR 
+	UE_DEPRECATED(5.5, "GetOnWaterInfoTextureCreated is deprecated, please use GetOnWaterInfoTextureArrayCreated instead.")
+	FOnWaterInfoTextureCreated& GetOnWaterInfoTextureCreated() { return OnWaterInfoTextureCreated_DEPRECATED; }
+#endif
+
 private:
 
 	/**
@@ -201,9 +215,9 @@ private:
 	/** Unique Id for accessing zone data (Location, extent, ,...) in GPU buffers */
 	UPROPERTY(Transient, DuplicateTransient, NonTransactional, VisibleAnywhere, Category = Water)
 	int32 WaterZoneIndex = INDEX_NONE;
-
-	UPROPERTY(BlueprintAssignable, Category=Water)
-	FOnWaterInfoTextureCreated OnWaterInfoTextureCreated;
+	
+	UPROPERTY(BlueprintAssignable, Category = Water)
+	FOnWaterInfoTextureArrayCreated OnWaterInfoTextureArrayCreated;
 
 #if WITH_EDITORONLY_DATA
 	/** A manipulatable box for visualizing/editing the water zone bounds */
@@ -221,6 +235,9 @@ private:
 	UPROPERTY(Transient)
 	TObjectPtr<UTexture2D> WaterVelocityTexture_DEPRECATED;
 
+	UPROPERTY()
+	FOnWaterInfoTextureCreated OnWaterInfoTextureCreated_DEPRECATED;
+
 	UPROPERTY()
 	FVector TessellatedWaterMeshExtent_DEPRECATED;
 
diff --git a/Engine/Plugins/Experimental/WaterAdvanced/Source/WaterAdvanced/Private/ShallowWaterSubsystem.cpp b/Engine/Plugins/Experimental/WaterAdvanced/Source/WaterAdvanced/Private/ShallowWaterSubsystem.cpp
--- a/Engine/Plugins/Experimental/WaterAdvanced/Source/WaterAdvanced/Private/ShallowWaterSubsystem.cpp
+++ b/Engine/Plugins/Experimental/WaterAdvanced/Source/WaterAdvanced/Private/ShallowWaterSubsystem.cpp
@@ -12,6 +12,7 @@
 #include "WaterSubsystem.h"
 #include "Components/SkeletalMeshComponent.h"
 #include "Engine/TextureRenderTarget2D.h"
+#include "Engine/TextureRenderTarget2DArray.h"
 #include "Kismet/GameplayStatics.h"
 #include "Kismet/KismetMaterialLibrary.h"
 #include "DrawDebugHelpers.h"
@@ -1123,16 +1124,16 @@ void UShallowWaterSubsystem::TryGetOrWaitForWaterInfoTextureFromWaterBodies(TSet
 	    {
 	    	if (AWaterZone* WaterZone = WaterBodyComp->GetWaterZone())
 	    	{
-	    		const TObjectPtr<UTextureRenderTarget2D> NewWaterInfoTexture = WaterZone->WaterInfoTexture;
+	    		const TObjectPtr<UTextureRenderTarget2DArray> NewWaterInfoTexture = WaterZone->WaterInfoTextureArray;
 
 	    		if (NewWaterInfoTexture == nullptr)
 	    		{
-	    			WaterZone->GetOnWaterInfoTextureCreated().RemoveDynamic(this, &UShallowWaterSubsystem::OnWaterInfoTextureCreated);
-	    			WaterZone->GetOnWaterInfoTextureCreated().AddDynamic(this, &UShallowWaterSubsystem::OnWaterInfoTextureCreated);
+	    			WaterZone->GetOnWaterInfoTextureArrayCreated().RemoveDynamic(this, &UShallowWaterSubsystem::OnWaterInfoTextureArrayCreated);
+	    			WaterZone->GetOnWaterInfoTextureArrayCreated().AddDynamic(this, &UShallowWaterSubsystem::OnWaterInfoTextureArrayCreated);
 	    		}
 	    		else
 	    		{
-	    			OnWaterInfoTextureCreated(NewWaterInfoTexture);
+	    			OnWaterInfoTextureArrayCreated(NewWaterInfoTexture);
 	    		}
 
 	    		// Currently there can only be one unique WaterInfoTexture
@@ -1142,7 +1143,7 @@ void UShallowWaterSubsystem::TryGetOrWaitForWaterInfoTextureFromWaterBodies(TSet
     }
 }
 
-void UShallowWaterSubsystem::OnWaterInfoTextureCreated(const UTextureRenderTarget2D* InWaterInfoTexture)
+void UShallowWaterSubsystem::OnWaterInfoTextureArrayCreated(const UTextureRenderTarget2DArray* InWaterInfoTexture)
 {
 	if(InWaterInfoTexture == nullptr)
 	{
@@ -1152,7 +1153,7 @@ void UShallowWaterSubsystem::OnWaterInfoTextureCreated(const UTextureRenderTarge
 	WaterInfoTexture = InWaterInfoTexture;
 	if (ShallowWaterNiagaraSimulation)
 	{
-		ShallowWaterNiagaraSimulation->SetVariableTexture(FName("WaterInfoTexture"), Cast<UTexture>(const_cast<UTextureRenderTarget2D*>(WaterInfoTexture.Get())));
+		ShallowWaterNiagaraSimulation->SetVariableTexture(FName("WaterInfoTexture"), Cast<UTexture>(const_cast<UTextureRenderTarget2DArray*>(WaterInfoTexture.Get())));
 	}
 }
 
diff --git a/Engine/Plugins/Experimental/WaterAdvanced/Source/WaterAdvanced/Public/ShallowWaterSubsystem.h b/Engine/Plugins/Experimental/WaterAdvanced/Source/WaterAdvanced/Public/ShallowWaterSubsystem.h
--- a/Engine/Plugins/Experimental/WaterAdvanced/Source/WaterAdvanced/Public/ShallowWaterSubsystem.h
+++ b/Engine/Plugins/Experimental/WaterAdvanced/Source/WaterAdvanced/Public/ShallowWaterSubsystem.h
@@ -6,6 +6,7 @@
 #include "GameplayTagContainer.h"
 #include "ShallowWaterSettings.h"
 #include "Engine/TextureRenderTarget2D.h"
+#include "Engine/TextureRenderTarget2DArray.h"
 #include "Components/SceneComponent.h"
 #include "Components/SkeletalMeshComponent.h"
 #include "Subsystems/WorldSubsystem.h"
@@ -236,10 +237,10 @@ protected:
 	 */
 	void TryGetOrWaitForWaterInfoTextureFromWaterBodies(TSet<AWaterBody*> CurrentWaterBodies);
 	UFUNCTION()
-	void OnWaterInfoTextureCreated(const UTextureRenderTarget2D* WaterInfoTexture);
+	void OnWaterInfoTextureArrayCreated(const UTextureRenderTarget2DArray* WaterInfoTexture);
 
 	UPROPERTY(BlueprintReadOnly, VisibleAnywhere, Category="Shallow Water")
-	TObjectPtr<const UTextureRenderTarget2D> WaterInfoTexture;
+	TObjectPtr<const UTextureRenderTarget2DArray> WaterInfoTexture;
 	UPROPERTY(BlueprintReadOnly, VisibleAnywhere, Category="Shallow Water")
 	TObjectPtr<UTextureRenderTarget2D> NormalRT;
 
diff --git a/Engine/Source/Runtime/Engine/Public/SceneView.h b/Engine/Source/Runtime/Engine/Public/SceneView.h
--- a/Engine/Source/Runtime/Engine/Public/SceneView.h
+++ b/Engine/Source/Runtime/Engine/Public/SceneView.h
@@ -1120,6 +1120,7 @@ BEGIN_GLOBAL_SHADER_PARAMETER_STRUCT_WITH_CONSTRUCTOR(FViewUniformShaderParamete
 	// Water
 	SHADER_PARAMETER_SRV(Buffer<float4>, WaterIndirection)
 	SHADER_PARAMETER_SRV(Buffer<float4>, WaterData)
+	SHADER_PARAMETER(int32, WaterInfoTextureViewIndex)
 	// Rect light atlas
 	SHADER_PARAMETER(FVector4f, RectLightAtlasSizeAndInvSize)
 	SHADER_PARAMETER(float, RectLightAtlasMaxMipLevel)
@@ -1580,6 +1581,8 @@ public:
 	/** Water rendering related data */
 	FShaderResourceViewRHIRef WaterIndirectionBuffer;
 	FShaderResourceViewRHIRef WaterDataBuffer;
+	/** Index of the water info texture slice corresponding to this view. */
+	int32 WaterInfoTextureViewIndex = INDEX_NONE;
 
 	struct FWaterInfoTextureRenderingParams
 	{
@@ -1595,6 +1598,7 @@ public:
 		float GroundZMin = 0.0f;
 		float CaptureZ = 0.0f;
 		int32 VelocityBlurRadius = 0;
+		int32 RenderTargetArrayLayer = 0;
 	};
 	TArray<FWaterInfoTextureRenderingParams> WaterInfoTextureRenderingParams;
 
diff --git a/Engine/Source/Runtime/Renderer/Private/SceneRendering.cpp b/Engine/Source/Runtime/Renderer/Private/SceneRendering.cpp
--- a/Engine/Source/Runtime/Renderer/Private/SceneRendering.cpp
+++ b/Engine/Source/Runtime/Renderer/Private/SceneRendering.cpp
@@ -2138,6 +2138,7 @@ void FViewInfo::SetupUniformBufferParameters(
 		ViewUniformShaderParameters.WaterIndirection = GWhiteVertexBufferWithSRV->ShaderResourceViewRHI;
 		ViewUniformShaderParameters.WaterData = GWhiteVertexBufferWithSRV->ShaderResourceViewRHI;
 	}
+	ViewUniformShaderParameters.WaterInfoTextureViewIndex = WaterInfoTextureViewIndex;
 
 	if (LandscapePerComponentDataBuffer.IsValid() && LandscapeIndirectionBuffer.IsValid())
 	{
diff --git a/Engine/Source/Runtime/Renderer/Private/WaterInfoTextureRendering.cpp b/Engine/Source/Runtime/Renderer/Private/WaterInfoTextureRendering.cpp
--- a/Engine/Source/Runtime/Renderer/Private/WaterInfoTextureRendering.cpp
+++ b/Engine/Source/Runtime/Renderer/Private/WaterInfoTextureRendering.cpp
@@ -1002,7 +1002,7 @@ void RenderWaterInfoTexture(
 
 			FWaterInfoTextureBlurPS::FParameters* PassParameters = GraphBuilder.AllocParameters<FWaterInfoTextureBlurPS::FParameters>();
 			PassParameters->View = WaterView.GetShaderParameters();
-			PassParameters->RenderTargets[0] = FRenderTargetBinding(OutputTexture, ERenderTargetLoadAction::ENoAction);
+			PassParameters->RenderTargets[0] = FRenderTargetBinding(OutputTexture, ERenderTargetLoadAction::ENoAction, 0, RenderingParams.RenderTargetArrayLayer);
 			PassParameters->SceneTextures = GetSceneTextureShaderParameters(WaterView);
 			PassParameters->WaterInfoTexture = MergedTexture;
 			PassParameters->WaterZMin = RenderingParams.WaterHeightExtents.X;
