UWaterSubsystem::FindWaterZone optimization: go through ForEachIntersectingActorDescInstance instead of doing the ForEachActorDescInstance/Intersect combo.

diff --git a/Engine/Plugins/Experimental/Water/Source/Runtime/Private/WaterSubsystem.cpp b/Engine/Plugins/Experimental/Water/Source/Runtime/Private/WaterSubsystem.cpp
--- a/Engine/Plugins/Experimental/Water/Source/Runtime/Private/WaterSubsystem.cpp
+++ b/Engine/Plugins/Experimental/Water/Source/Runtime/Private/WaterSubsystem.cpp
@@ -623,17 +623,11 @@ TSoftObjectPtr<AWaterZone> UWaterSubsystem::FindWaterZone(const UWorld* World, c
 	{
 		if (UWorldPartition* WorldPartition = World->GetWorldPartition())
 		{
-			FWorldPartitionHelpers::ForEachActorDescInstance<AWaterZone>(WorldPartition, [Bounds, &ViableZones](const FWorldPartitionActorDescInstance* ActorDescInstance)
+			const FBox Bounds3D(FVector(Bounds.Min.X, Bounds.Min.Y, -HALF_WORLD_MAX), FVector(Bounds.Max.X, Bounds.Max.Y, HALF_WORLD_MAX));
+			FWorldPartitionHelpers::ForEachIntersectingActorDescInstance<AWaterZone>(WorldPartition, Bounds3D, [&Bounds, &ViableZones](const FWorldPartitionActorDescInstance* ActorDescInstance)
 			{
 				FWaterZoneActorDesc* WaterZoneActorDesc = (FWaterZoneActorDesc*)ActorDescInstance->GetActorDesc();
-				const FBox WaterZoneBounds = ActorDescInstance->GetEditorBounds();
-				const FBox2D WaterZoneBounds2D(FVector2D(WaterZoneBounds.Min), FVector2D(WaterZoneBounds.Max));
-
-				if (Bounds.Intersect(WaterZoneBounds2D))
-				{
-					ViableZones.Emplace(ActorDescInstance->GetActorSoftPath(), WaterZoneActorDesc->GetOverlapPriority());
-				}
-
+				ViableZones.Emplace(ActorDescInstance->GetActorSoftPath(), WaterZoneActorDesc->GetOverlapPriority());
 				return true;
 			});
 		}
