SceneCapture: Fixed engine show flags not updating at runtime (from non-editor changes).

We now automatically call the required update method inside a setter.

diff --git a/Engine/Plugins/Experimental/Water/Source/Editor/Private/WaterBrushManager.cpp b/Engine/Plugins/Experimental/Water/Source/Editor/Private/WaterBrushManager.cpp
--- a/Engine/Plugins/Experimental/Water/Source/Editor/Private/WaterBrushManager.cpp
+++ b/Engine/Plugins/Experimental/Water/Source/Editor/Private/WaterBrushManager.cpp
@@ -51,12 +51,15 @@ AWaterBrushManager::AWaterBrushManager(const FObjectInitializer& ObjectInitializ
 	SceneCaptureComponent2D->bCaptureOnMovement = false;
 	SceneCaptureComponent2D->SetRelativeRotation(FRotator(-90.0f, 0.0f, -90.0f));
 	SceneCaptureComponent2D->SetRelativeScale3D(FVector(0.01f, 0.01f, 0.01f));
+
+	TArray<FEngineShowFlagsSetting> ShowFlagSettings;
 	// HACK [jonathan.bard] : Nanite doesn't support USceneCaptureComponent's ShowOnlyComponents ATM so just disable Nanite during captures : 
-	SceneCaptureComponent2D->ShowFlagSettings.Add(FEngineShowFlagsSetting { TEXT("NaniteMeshes"), false } );
+	ShowFlagSettings.Add(FEngineShowFlagsSetting { TEXT("NaniteMeshes"), false } );
 	// These also need to be disabled to get a clean capture of just the water info material output
-	SceneCaptureComponent2D->ShowFlagSettings.Add(FEngineShowFlagsSetting { TEXT("Atmosphere"), false } );
-	SceneCaptureComponent2D->ShowFlagSettings.Add(FEngineShowFlagsSetting { TEXT("Bloom"), false } );
-	SceneCaptureComponent2D->ShowFlagSettings.Add(FEngineShowFlagsSetting { TEXT("Lighting"), false } );
+	ShowFlagSettings.Add(FEngineShowFlagsSetting { TEXT("Atmosphere"), false } );
+	ShowFlagSettings.Add(FEngineShowFlagsSetting { TEXT("Bloom"), false } );
+	ShowFlagSettings.Add(FEngineShowFlagsSetting { TEXT("Lighting"), false } );
+	SceneCaptureComponent2D->SetShowFlagSettings(ShowFlagSettings);
 
 	PrimaryActorTick.TickGroup = ETickingGroup::TG_PrePhysics;
 	bIsEditorOnlyActor = false;
diff --git a/Engine/Plugins/VirtualProduction/CameraCalibration/Source/CameraCalibrationEditor/Private/AssetEditor/CameraCalibrationStepsController.cpp b/Engine/Plugins/VirtualProduction/CameraCalibration/Source/CameraCalibrationEditor/Private/AssetEditor/CameraCalibrationStepsController.cpp
--- a/Engine/Plugins/VirtualProduction/CameraCalibration/Source/CameraCalibrationEditor/Private/AssetEditor/CameraCalibrationStepsController.cpp
+++ b/Engine/Plugins/VirtualProduction/CameraCalibration/Source/CameraCalibrationEditor/Private/AssetEditor/CameraCalibrationStepsController.cpp
@@ -718,15 +718,7 @@ void FCameraCalibrationStepsController::CreateComp()
 			NewFlagSetting.ShowFlagName = FEngineShowFlags::FindNameByIndex(FEngineShowFlags::EShowFlag::SF_Fog);
 			NewFlagSetting.Enabled = false;
 
-			CaptureComponent->ShowFlagSettings.Add(NewFlagSetting);
-
-			if (FProperty* ShowFlagSettingsProperty = CaptureComponent->GetClass()->FindPropertyByName(
-				GET_MEMBER_NAME_CHECKED(USceneCaptureComponent2D, ShowFlagSettings)))
-			{
-				// This PostEditChange will ensure that ShowFlags is updated.
-				FPropertyChangedEvent PropertyChangedEvent(ShowFlagSettingsProperty);
-				CaptureComponent->PostEditChangeProperty(PropertyChangedEvent);
-			}
+			CaptureComponent->SetShowFlagSettings({ NewFlagSetting });
 		}
 	}
 
diff --git a/Engine/Source/Runtime/Engine/Classes/Components/SceneCaptureComponent.h b/Engine/Source/Runtime/Engine/Classes/Components/SceneCaptureComponent.h
--- a/Engine/Source/Runtime/Engine/Classes/Components/SceneCaptureComponent.h
+++ b/Engine/Source/Runtime/Engine/Classes/Components/SceneCaptureComponent.h
@@ -141,13 +141,22 @@ class USceneCaptureComponent : public USceneComponent
 	bool bUseRayTracingIfEnabled;
 
 	/** ShowFlags for the SceneCapture's ViewFamily, to control rendering settings for this view. Hidden but accessible through details customization */
-	UPROPERTY(EditAnywhere, interp, Category=SceneCapture)
-	TArray<struct FEngineShowFlagsSetting> ShowFlagSettings;
+	UE_DEPRECATED_FORGAME(5.5, "Public access to this property is deprecated, and it will become private in a future release. Please use SetShowFlagSettings and GetShowFlagSettings instead.")
+	UPROPERTY(EditAnywhere, interp, Blueprintgetter = GetShowFlagSettings, BlueprintSetter = SetShowFlagSettings, Category=SceneCapture)
+	TArray<FEngineShowFlagsSetting> ShowFlagSettings;
 
 	// TODO: Make this a UStruct to set directly?
 	/** Settings stored here read from the strings and int values in the ShowFlagSettings array */
 	FEngineShowFlags ShowFlags;
 
+	/** Get the show flag settings. */
+	UFUNCTION(BlueprintGetter)
+	ENGINE_API const TArray<FEngineShowFlagsSetting>& GetShowFlagSettings() const;
+
+	/** Set the show flag settings. */
+	UFUNCTION(BlueprintSetter)
+	ENGINE_API void SetShowFlagSettings(const TArray<FEngineShowFlagsSetting>& InShowFlagSettings);
+
 public:
 	/** Name of the profiling event. */
 	UPROPERTY(EditAnywhere, interp, Category = SceneCapture)
diff --git a/Engine/Source/Runtime/Engine/Private/Components/SceneCaptureComponent.cpp b/Engine/Source/Runtime/Engine/Private/Components/SceneCaptureComponent.cpp
--- a/Engine/Source/Runtime/Engine/Private/Components/SceneCaptureComponent.cpp
+++ b/Engine/Source/Runtime/Engine/Private/Components/SceneCaptureComponent.cpp
@@ -221,6 +221,18 @@ void USceneCaptureComponent::PostLoad()
 	}
 }
 
+const TArray<FEngineShowFlagsSetting>& USceneCaptureComponent::GetShowFlagSettings() const
+{
+	return ShowFlagSettings;
+}
+
+void USceneCaptureComponent::SetShowFlagSettings(const TArray<FEngineShowFlagsSetting>& InShowFlagSettings)
+{
+	ShowFlagSettings = InShowFlagSettings;
+
+	UpdateShowFlags();
+}
+
 void USceneCaptureComponent::BeginDestroy()
 {
 	Super::BeginDestroy();
