* Fixed UWaterBodyComponent::UpdateComponentVisibility updating the water info texture whenever it was called.
* Fixed "Update Water Mesh During Interactive Changes" (Water Editor) setting that has been ineffective for a while now

diff --git a/Engine/Plugins/Experimental/Water/Source/Editor/Private/WaterBrushManager.cpp b/Engine/Plugins/Experimental/Water/Source/Editor/Private/WaterBrushManager.cpp
--- a/Engine/Plugins/Experimental/Water/Source/Editor/Private/WaterBrushManager.cpp
+++ b/Engine/Plugins/Experimental/Water/Source/Editor/Private/WaterBrushManager.cpp
@@ -1278,12 +1278,6 @@ void AWaterBrushManager::RenderBrushActorContext(FBrushRenderContext& BrushRende
 	}
 	++BrushRenderContext.RTIndex;
 
-	if (WaterBody != nullptr)
-	{
-		// rebuilding the water mesh is expensive and not necessary : 
-		WaterBody->GetWaterBodyComponent()->UpdateComponentVisibility(/* bAllowWaterZoneRebuild = */false);
-	}
-
 	ApplyToCompositeWaterBodyTexture(BrushRenderContext, BrushActorRenderContext);
 }
 
diff --git a/Engine/Plugins/Experimental/Water/Source/Editor/Private/WaterEditorSubsystem.cpp b/Engine/Plugins/Experimental/Water/Source/Editor/Private/WaterEditorSubsystem.cpp
--- a/Engine/Plugins/Experimental/Water/Source/Editor/Private/WaterEditorSubsystem.cpp
+++ b/Engine/Plugins/Experimental/Water/Source/Editor/Private/WaterEditorSubsystem.cpp
@@ -228,5 +228,10 @@ void UWaterEditorSubsystem::DirtyAllModifiedPackages()
 	ClearModifiedPackages();
 }
 
+bool UWaterEditorSubsystem::GetShouldUpdateWaterMeshDuringInteractiveChanges() const
+{
+	return GetDefault<UWaterEditorSettings>()->GetShouldUpdateWaterMeshDuringInteractiveChanges();
+}
+
 #undef LOCTEXT_NAMESPACE
 
diff --git a/Engine/Plugins/Experimental/Water/Source/Editor/Public/WaterEditorSettings.h b/Engine/Plugins/Experimental/Water/Source/Editor/Public/WaterEditorSettings.h
--- a/Engine/Plugins/Experimental/Water/Source/Editor/Public/WaterEditorSettings.h
+++ b/Engine/Plugins/Experimental/Water/Source/Editor/Public/WaterEditorSettings.h
@@ -232,7 +232,12 @@ public:
 	UMaterialInterface* GetDefaultRenderRiverSplineDepthsMaterial() const;
 	FSoftObjectPath GetDefaultRenderRiverSplineDepthsMaterialPath() const { return DefaultRenderRiverSplineDepthsMaterial.ToSoftObjectPath(); }
 
+	bool GetShouldUpdateLandscapeDuringInteractiveChanges() const { return bUpdateLandscapeDuringInteractiveChanges; }
+	bool GetShouldUpdateWaterMeshDuringInteractiveChanges() const { return bUpdateWaterMeshDuringInteractiveChanges; }
+
+	UE_DEPRECATED(5.5, "Use GetShouldUpdateLandscapeDuringInteractiveChanges");
 	bool GetUpdateLandscapeDuringInteractiveChanges() const { return bUpdateLandscapeDuringInteractiveChanges; }
+	UE_DEPRECATED(5.5, "Use GetShouldUpdateWaterMeshDuringInteractiveChanges");
 	bool GetUpdateWaterMeshDuringInteractiveChanges() const { return bUpdateWaterMeshDuringInteractiveChanges; }
 	
 public:
diff --git a/Engine/Plugins/Experimental/Water/Source/Editor/Public/WaterEditorSubsystem.h b/Engine/Plugins/Experimental/Water/Source/Editor/Public/WaterEditorSubsystem.h
--- a/Engine/Plugins/Experimental/Water/Source/Editor/Public/WaterEditorSubsystem.h
+++ b/Engine/Plugins/Experimental/Water/Source/Editor/Public/WaterEditorSubsystem.h
@@ -36,6 +36,8 @@ public:
 	virtual void ForEachModifiedPackage(const TFunctionRef<bool(UPackage*)>& Func) const override;
 	virtual void ClearModifiedPackages() override;
 	virtual void DirtyAllModifiedPackages() override;
+
+	virtual bool GetShouldUpdateWaterMeshDuringInteractiveChanges() const override;
 	//~ End IWaterEditorServices interface
 
 private:
diff --git a/Engine/Plugins/Experimental/Water/Source/Runtime/Private/WaterBodyActor.cpp b/Engine/Plugins/Experimental/Water/Source/Runtime/Private/WaterBodyActor.cpp
--- a/Engine/Plugins/Experimental/Water/Source/Runtime/Private/WaterBodyActor.cpp
+++ b/Engine/Plugins/Experimental/Water/Source/Runtime/Private/WaterBodyActor.cpp
@@ -721,7 +721,7 @@ void AWaterBody::SetActorHiddenInGame(bool bNewHidden)
 
 		if (WaterBodyComponent)
 		{
-			WaterBodyComponent->UpdateComponentVisibility(/* bAllowWaterMeshRebuild = */true);
+			WaterBodyComponent->UpdateVisibility();
 		}
 	}
 }
@@ -736,7 +736,7 @@ void AWaterBody::SetIsTemporarilyHiddenInEditor(bool bIsHidden)
 
 		if (WaterBodyComponent)
 		{
-			WaterBodyComponent->UpdateComponentVisibility(/* bAllowWaterMeshRebuild = */true);
+			WaterBodyComponent->UpdateVisibility();
 		}
 	}
 }
@@ -749,7 +749,7 @@ bool AWaterBody::SetIsHiddenEdLayer(bool bIsHiddenEdLayer)
 	{
 		if (WaterBodyComponent)
 		{
-			WaterBodyComponent->UpdateComponentVisibility(/* bAllowWaterMeshRebuild = */true);
+			WaterBodyComponent->UpdateVisibility();
 		}
 		return true;
 	}
diff --git a/Engine/Plugins/Experimental/Water/Source/Runtime/Private/WaterBodyComponent.cpp b/Engine/Plugins/Experimental/Water/Source/Runtime/Private/WaterBodyComponent.cpp
--- a/Engine/Plugins/Experimental/Water/Source/Runtime/Private/WaterBodyComponent.cpp
+++ b/Engine/Plugins/Experimental/Water/Source/Runtime/Private/WaterBodyComponent.cpp
@@ -80,6 +80,25 @@ const FName UWaterBodyComponent::WaterVelocityAndHeightName(TEXT("WaterVelocityA
 const FName UWaterBodyComponent::GlobalOceanHeightName(TEXT("GlobalOceanHeight"));
 const FName UWaterBodyComponent::MaxFlowVelocityParamName(TEXT("MaxFlowVelocity"));
 
+namespace UE::Water
+{
+	static bool ShouldUpdateWaterMeshForPropertyChange(const FPropertyChangedEvent& InPropertyChangedEvent)
+	{
+#if WITH_EDITOR
+		if (InPropertyChangedEvent.ChangeType == EPropertyChangeType::Interactive)
+		{
+			const IWaterModuleInterface& WaterModule = FModuleManager::GetModuleChecked<IWaterModuleInterface>("Water");
+			if (const IWaterEditorServices* WaterEditorServices = WaterModule.GetWaterEditorServices())
+			{
+				return WaterEditorServices->GetShouldUpdateWaterMeshDuringInteractiveChanges();
+			}
+		}
+#endif // WITH_EDITOR
+
+		return true;
+	}
+}
+
 UWaterBodyComponent::UWaterBodyComponent(const FObjectInitializer& ObjectInitializer)
 	: Super(ObjectInitializer)
 {
@@ -119,14 +138,14 @@ void UWaterBodyComponent::OnVisibilityChanged()
 {
 	Super::OnVisibilityChanged();
 
-	UpdateComponentVisibility(/* bAllowWaterZoneRebuild = */true);
+	UpdateVisibility();
 }
 
 void UWaterBodyComponent::OnHiddenInGameChanged()
 {
 	Super::OnHiddenInGameChanged();
 
-	UpdateComponentVisibility(/* bAllowWaterZoneRebuild = */true);
+	UpdateVisibility();
 }
 
 bool UWaterBodyComponent::IsFlatSurface() const
@@ -1034,41 +1053,41 @@ ALandscapeProxy* UWaterBodyComponent::FindLandscape() const
 	return Landscape.Get();
 }
 
+// Deprecated
 void UWaterBodyComponent::UpdateComponentVisibility(bool bAllowWaterZoneRebuild)
+{
+	UpdateVisibility();
+}
+
+void UWaterBodyComponent::UpdateVisibility()
 {
 	if (UWorld* World = GetWorld())
 	{
 	 	const bool bIsWaterRenderingEnabled = FWaterUtils::IsWaterEnabled(/*bIsRenderThread = */false);
 	 
 		bool bIsRenderedByWaterMesh = ShouldGenerateWaterMeshTile();
-		bool bLocalVisible = bIsWaterRenderingEnabled && !bIsRenderedByWaterMesh && GetVisibleFlag();
-		bool bLocalHiddenInGame = !bIsWaterRenderingEnabled || bIsRenderedByWaterMesh || bHiddenInGame;
-
-	 	for (UPrimitiveComponent* Component : GetStandardRenderableComponents())
-	 	{
-	 		Component->SetVisibility(bLocalVisible);
-	 		Component->SetHiddenInGame(bLocalHiddenInGame);
-	 	}
 
-		if (bAllowWaterZoneRebuild)
+		// Handle the standard renderable components (i.e. when the water body is not rendered by the water mesh) : 
 		{
-			if (AWaterZone* WaterZone = GetWaterZone())
-			{
-				// If the component is being or can be rendered by the water mesh or renders into the water info texture, rebuild it in case its visibility has changed : 
+			bool bLocalVisible = bIsWaterRenderingEnabled && !bIsRenderedByWaterMesh && GetVisibleFlag();
+			bool bLocalHiddenInGame = !bIsWaterRenderingEnabled || bIsRenderedByWaterMesh || bHiddenInGame;
 
-				EWaterZoneRebuildFlags RebuildFlags = EWaterZoneRebuildFlags::None;
-				if (AffectsWaterMesh())
-				{
-					RebuildFlags |= EWaterZoneRebuildFlags::UpdateWaterMesh;
-				}
+			for (UPrimitiveComponent* Component : GetStandardRenderableComponents())
+			{
+				Component->SetVisibility(bLocalVisible);
+				Component->SetHiddenInGame(bLocalHiddenInGame);
+			}
+		}
 
-				if (AffectsWaterInfo())
-				{
-					RebuildFlags |= EWaterZoneRebuildFlags::UpdateWaterInfoTexture;
-				}
+		// Keep track of the visibility state when the water body is being rendered by the water mesh, to avoid calling MarkOwningWaterZoneForRebuild every time UpdateVisibility is called :
+		const bool bPreviousIsRenderedByWaterMeshAndVisible = bIsRenderedByWaterMeshAndVisible;
+		bIsRenderedByWaterMeshAndVisible = bIsRenderedByWaterMesh && ShouldRender();
+		const bool bHasVisibilityChanged = (bIsRenderedByWaterMeshAndVisible != bPreviousIsRenderedByWaterMeshAndVisible);
 
-				MarkOwningWaterZoneForRebuild(RebuildFlags);
-			}
+		if (bHasVisibilityChanged)
+		{
+			// If the component is being or can be rendered by the water mesh or renders into the water info texture, rebuild it in case its visibility has changed : 
+			MarkOwningWaterZoneForRebuild(EWaterZoneRebuildFlags::All);
 		}
 	}
 }
@@ -1443,6 +1462,7 @@ void UWaterBodyComponent::UpdateAll(const FOnWaterBodyChangedParams& InParams)
 	const bool bUserTriggered = InParams.bUserTriggered;
 	bool bShapeOrPositionChanged = InParams.bShapeOrPositionChanged;
 	
+	// TODO [jonathan.bard] : Replace GIsEditor by World->IsEditorWorld
 	if (GIsEditor || IsBodyDynamic())
 	{
 		TRACE_CPUPROFILER_EVENT_SCOPE(UWaterBodyComponent::UpdateAll);
@@ -1481,15 +1501,24 @@ void UWaterBodyComponent::UpdateAll(const FOnWaterBodyChangedParams& InParams)
 
 		if (bShapeOrPositionChanged)
 		{
+			EWaterZoneRebuildFlags RebuildFlags = EWaterZoneRebuildFlags::All;
+			if (!UE::Water::ShouldUpdateWaterMeshForPropertyChange(InParams.PropertyChangedEvent))
+			{
+				EnumRemoveFlags(RebuildFlags, EWaterZoneRebuildFlags::UpdateWaterMesh);
+			}
+
+			MarkOwningWaterZoneForRebuild(RebuildFlags);
+
 			FNavigationSystem::UpdateActorAndComponentData(*WaterBodyOwner);
 		}
 
-		UpdateComponentVisibility(/* bAllowWaterZoneRebuild = */true);
-
 #if WITH_EDITOR
 		UpdateWaterSpriteComponent();
 #endif
 	}
+
+	// Always update the visibility, since it might change based on CVars and there's a transient boolean (bIsComponent
+	UpdateVisibility();
 }
 
 void UWaterBodyComponent::OnPostRegisterAllComponents()
@@ -1556,7 +1585,7 @@ void UWaterBodyComponent::OnWaterBodyChanged(const FOnWaterBodyChangedParams& In
 	}
 
 #if WITH_EDITOR
-	if (InParams.PropertyChangedEvent.ChangeType != EPropertyChangeType::Interactive)
+	if (UE::Water::ShouldUpdateWaterMeshForPropertyChange(InParams.PropertyChangedEvent))
 	{
 		UpdateWaterBodyRenderData();
 	}
diff --git a/Engine/Plugins/Experimental/Water/Source/Runtime/Private/WaterSubsystem.cpp b/Engine/Plugins/Experimental/Water/Source/Runtime/Private/WaterSubsystem.cpp
--- a/Engine/Plugins/Experimental/Water/Source/Runtime/Private/WaterSubsystem.cpp
+++ b/Engine/Plugins/Experimental/Water/Source/Runtime/Private/WaterSubsystem.cpp
@@ -734,7 +734,7 @@ void UWaterSubsystem::NotifyWaterVisibilityChangedInternal(IConsoleVariable* CVa
 	// Water body visibility depends on various CVars. All need to update the visibility in water body components : 
 	GetWaterBodyManagerInternal().ForEachWaterBodyComponent([](UWaterBodyComponent* WaterBodyComponent)
 	{
-		WaterBodyComponent->UpdateComponentVisibility(/* bAllowWaterZoneRebuild = */true);
+		WaterBodyComponent->UpdateVisibility();
 		return true;
 	});
 }
diff --git a/Engine/Plugins/Experimental/Water/Source/Runtime/Public/WaterBodyComponent.h b/Engine/Plugins/Experimental/Water/Source/Runtime/Public/WaterBodyComponent.h
--- a/Engine/Plugins/Experimental/Water/Source/Runtime/Public/WaterBodyComponent.h
+++ b/Engine/Plugins/Experimental/Water/Source/Runtime/Public/WaterBodyComponent.h
@@ -365,11 +365,12 @@ public:
 	/** Returns true if the location is within one of this water body's exclusion volumes */
 	bool IsWorldLocationInExclusionVolume(const FVector& InWorldLocation) const;
 
-	/** Updates the bVisible/bHiddenInGame flags on the component and eventually the child renderable components (e.g. custom water body)
-	 - bAllowWaterZoneRebuild : if true, the function will request a rebuild of the water zone (expensive), which is necessary to take into account visibility changes
-	*/
+	UE_DEPRECATED(5.5, "Use UpdateComponentVisibility")
 	virtual void UpdateComponentVisibility(bool bAllowWaterZoneRebuild);
 
+	/** Updates the bVisible/bHiddenInGame flags on the component and eventually the child renderable components (e.g. custom water body) */
+	void UpdateVisibility();
+
 	/** Creates/Destroys/Updates necessary MIDS */
 	virtual void UpdateMaterialInstances();
 
@@ -701,4 +702,9 @@ protected:
 	UPROPERTY()
 	bool bOverrideWaterMesh_DEPRECATED;
 #endif // WITH_EDITORONLY_DATA
+
+private:
+	/** Boolean to keep track of whether the water body is visible in the water mesh. This avoids unnecessary calls to rebuild the water mesh whenever UpdateComponentVisibility is called */
+	bool bIsRenderedByWaterMeshAndVisible = false;
+
 };
diff --git a/Engine/Plugins/Experimental/Water/Source/Runtime/Public/WaterEditorServices.h b/Engine/Plugins/Experimental/Water/Source/Runtime/Public/WaterEditorServices.h
--- a/Engine/Plugins/Experimental/Water/Source/Runtime/Public/WaterEditorServices.h
+++ b/Engine/Plugins/Experimental/Water/Source/Runtime/Public/WaterEditorServices.h
@@ -28,6 +28,8 @@ public:
 	virtual void ForEachModifiedPackage(const TFunctionRef<bool(UPackage*)>& Func) const = 0;
 	virtual void ClearModifiedPackages() = 0;
 	virtual void DirtyAllModifiedPackages() = 0;
+
+	virtual bool GetShouldUpdateWaterMeshDuringInteractiveChanges() const = 0;
 };
 #endif // WITH_EDITORONLY_DATA
 
diff --git a/Engine/Plugins/Experimental/Water/Source/Runtime/Public/WaterZoneActor.h b/Engine/Plugins/Experimental/Water/Source/Runtime/Public/WaterZoneActor.h
--- a/Engine/Plugins/Experimental/Water/Source/Runtime/Public/WaterZoneActor.h
+++ b/Engine/Plugins/Experimental/Water/Source/Runtime/Public/WaterZoneActor.h
@@ -19,8 +19,8 @@ class UWaterBodyComponent;
 enum class EWaterZoneRebuildFlags
 {
 	None = 0,
-	UpdateWaterInfoTexture = (1 << 1),
-	UpdateWaterMesh = (1 << 2),
+	UpdateWaterInfoTexture = (1 << 0),
+	UpdateWaterMesh = (1 << 1),
 	All = (~0),
 };
 ENUM_CLASS_FLAGS(EWaterZoneRebuildFlags);
