Build fix

diff --git a/Engine/Plugins/Experimental/Water/Source/Editor/Public/WaterEditorSettings.h b/Engine/Plugins/Experimental/Water/Source/Editor/Public/WaterEditorSettings.h
--- a/Engine/Plugins/Experimental/Water/Source/Editor/Public/WaterEditorSettings.h
+++ b/Engine/Plugins/Experimental/Water/Source/Editor/Public/WaterEditorSettings.h
@@ -235,9 +235,9 @@ public:
 	bool GetShouldUpdateLandscapeDuringInteractiveChanges() const { return bUpdateLandscapeDuringInteractiveChanges; }
 	bool GetShouldUpdateWaterMeshDuringInteractiveChanges() const { return bUpdateWaterMeshDuringInteractiveChanges; }
 
-	UE_DEPRECATED(5.5, "Use GetShouldUpdateLandscapeDuringInteractiveChanges");
+	UE_DEPRECATED(5.5, "Use GetShouldUpdateLandscapeDuringInteractiveChanges")
 	bool GetUpdateLandscapeDuringInteractiveChanges() const { return bUpdateLandscapeDuringInteractiveChanges; }
-	UE_DEPRECATED(5.5, "Use GetShouldUpdateWaterMeshDuringInteractiveChanges");
+	UE_DEPRECATED(5.5, "Use GetShouldUpdateWaterMeshDuringInteractiveChanges")
 	bool GetUpdateWaterMeshDuringInteractiveChanges() const { return bUpdateWaterMeshDuringInteractiveChanges; }
 	
 public:
