Build fix

diff --git a/Engine/Plugins/Experimental/Water/Source/Editor/Private/WaterLandscapeBrush.cpp b/Engine/Plugins/Experimental/Water/Source/Editor/Private/WaterLandscapeBrush.cpp
--- a/Engine/Plugins/Experimental/Water/Source/Editor/Private/WaterLandscapeBrush.cpp
+++ b/Engine/Plugins/Experimental/Water/Source/Editor/Private/WaterLandscapeBrush.cpp
@@ -197,7 +197,7 @@ void AWaterLandscapeBrush::OnWaterBrushActorChanged(const IWaterBrushActorInterf
 	const UWaterEditorSettings* WaterEditorSettings = GetDefault<UWaterEditorSettings>();
 	check(WaterEditorSettings != nullptr);
 
-	const bool bAllowLandscapeUpdate = (InParams.PropertyChangedEvent.ChangeType != EPropertyChangeType::Interactive) || WaterEditorSettings->GetUpdateLandscapeDuringInteractiveChanges();
+	const bool bAllowLandscapeUpdate = (InParams.PropertyChangedEvent.ChangeType != EPropertyChangeType::Interactive) || WaterEditorSettings->GetShouldUpdateLandscapeDuringInteractiveChanges();
 	if (bForceUpdateBrush || (bAffectsLandscape && bAllowLandscapeUpdate))
 	{
 		RequestLandscapeUpdate(/* bInUserTriggered = */ InParams.bUserTriggered);
