Fixed Water brush manager not properly reacting to landscape resolution changes

diff --git a/Engine/Plugins/Experimental/Water/Source/Editor/Private/WaterBrushManager.cpp b/Engine/Plugins/Experimental/Water/Source/Editor/Private/WaterBrushManager.cpp
--- a/Engine/Plugins/Experimental/Water/Source/Editor/Private/WaterBrushManager.cpp
+++ b/Engine/Plugins/Experimental/Water/Source/Editor/Private/WaterBrushManager.cpp
@@ -406,10 +406,43 @@ void AWaterBrushManager::Initialize_Native(FTransform const& InLandscapeTransfor
 {
 	UE_LOG(LogWaterEditor, Verbose, TEXT("Updated Landscape Transform"));
 
-	LandscapeQuads = InLandscapeSize;
-	LandscapeRTRes = InLandscapeRenderTargetSize;
+	bool bNeedsFullUpdate = false;
+	if (LandscapeQuads != InLandscapeSize)
+	{
+		LandscapeQuads = InLandscapeSize;
+		bNeedsFullUpdate = true;
+	}
+	if (LandscapeRTRes != InLandscapeRenderTargetSize)
+	{
+		LandscapeRTRes = InLandscapeRenderTargetSize;
+		bNeedsFullUpdate = true;
+	}
+	if (!InLandscapeTransform.Equals(LandscapeTransform))
+	{
+		LandscapeTransform = InLandscapeTransform;
+		bNeedsFullUpdate = true;
+	}
+
+	if (bNeedsFullUpdate)
+	{
+		check(SceneCaptureComponent2D != nullptr);
+
+		FVector Scale = LandscapeTransform.GetScale3D();
+		WorldSize.Set(Scale.X * (float)LandscapeQuads.X, Scale.Y * (float)LandscapeQuads.Y, 0.512f);
 
-	UpdateTransform(InLandscapeTransform);
+		const FVector Temp(Scale.X * (float)LandscapeRTRes.X, Scale.Y * (float)LandscapeRTRes.Y, 0.512f);
+		SceneCaptureComponent2D->OrthoWidth = FMath::Max(Temp.X, Temp.Y);
+
+		FVector LocationVector(Temp - Scale);
+		LocationVector *= 0.5f;
+		LocationVector = LandscapeTransform.GetRotation().RotateVector(LocationVector);
+		LocationVector += LandscapeTransform.GetLocation();
+		LocationVector.Z = 50000.0f;
+		SceneCaptureComponent2D->SetWorldLocation(LocationVector);
+
+		// If the transform or resolution changes, the distance fields need to be recomputed entirely : 
+		bKillCache = true;
+	}
 }
 
 void AWaterBrushManager::CaptureMeshDepth(const TArrayView<UStaticMeshComponent*>& MeshComponents)
@@ -464,31 +497,6 @@ void AWaterBrushManager::GetRenderDependencies(TSet<UObject*>& OutDependencies)
 	AddDependencyIfValid(BlurEdgesMaterial, OutDependencies);
 }
 
-void AWaterBrushManager::UpdateTransform(const FTransform& Transform)
-{
-	if (!Transform.Equals(LandscapeTransform))
-	{
-		LandscapeTransform = Transform;
-		check(SceneCaptureComponent2D != nullptr);
-
-		FVector Scale = LandscapeTransform.GetScale3D();
-		WorldSize.Set(Scale.X * (float)LandscapeQuads.X, Scale.Y * (float)LandscapeQuads.Y, 0.512f);
-
-		const FVector Temp(Scale.X * (float)LandscapeRTRes.X, Scale.Y * (float)LandscapeRTRes.Y, 0.512f);
-		SceneCaptureComponent2D->OrthoWidth = FMath::Max(Temp.X, Temp.Y);
-
-		FVector LocationVector(Temp - Scale);
-		LocationVector *= 0.5f;
-		LocationVector = LandscapeTransform.GetRotation().RotateVector(LocationVector);
-		LocationVector += LandscapeTransform.GetLocation();
-		LocationVector.Z = 50000.0f;
-		SceneCaptureComponent2D->SetWorldLocation(LocationVector);
-
-		// The landscape transform has changed, let's re-draw everything (no need to request a landscape update because we're in the middle of one) :
-		bKillCache = true;
-	}
-}
-
 bool AWaterBrushManager::SetupRiverSplineRenderMIDs(const FBrushActorRenderContext& BrushActorRenderContext, bool bRestoreMIDs, TArray<UMaterialInterface*>& InOutMIDs)
 {
 	AWaterBody* WaterBody = BrushActorRenderContext.GetActorAs<AWaterBody>();
diff --git a/Engine/Plugins/Experimental/Water/Source/Editor/Private/WaterBrushManager.h b/Engine/Plugins/Experimental/Water/Source/Editor/Private/WaterBrushManager.h
--- a/Engine/Plugins/Experimental/Water/Source/Editor/Private/WaterBrushManager.h
+++ b/Engine/Plugins/Experimental/Water/Source/Editor/Private/WaterBrushManager.h
@@ -270,7 +270,6 @@ private:
 	virtual void PostEditChangeProperty(FPropertyChangedEvent& PropertyChangedEvent) override;
 	virtual bool AllocateRTs();
 	virtual void SetMPCParams();
-	virtual void UpdateTransform(const FTransform& Transform);
 	virtual bool SetupRiverSplineRenderMIDs(const FBrushActorRenderContext& BrushActorRenderContext, bool bRestoreMIDs, TArray<UMaterialInterface*>& InOutMIDs);
 	virtual void CaptureMeshDepth(const TArrayView<UStaticMeshComponent*>& MeshComponents);
 	virtual void CacheBrushDistanceField(const FBrushActorRenderContext& BrushActorRenderContext);
