- Water DI changes
- Add ability to find the closest body automatically
- Add IsValid function
- Fix a missing output when in an exlusion volume

#rnx
#tests Placing things on water bodies, enabling / disabling the closest flag, etc.

diff --git a/Engine/Plugins/Experimental/Water/Source/Runtime/Private/NiagaraDataInterfaceWater.cpp b/Engine/Plugins/Experimental/Water/Source/Runtime/Private/NiagaraDataInterfaceWater.cpp
--- a/Engine/Plugins/Experimental/Water/Source/Runtime/Private/NiagaraDataInterfaceWater.cpp
+++ b/Engine/Plugins/Experimental/Water/Source/Runtime/Private/NiagaraDataInterfaceWater.cpp
@@ -3,9 +3,12 @@
 #include "NiagaraDataInterfaceWater.h"
 
 #include "NiagaraSystemInstance.h"
+#include "WaterBodyActor.h"
 #include "WaterBodyComponent.h"
 #include "WaterModule.h"
 
+#include "UObject/UObjectIterator.h"
+
 #include UE_INLINE_GENERATED_CPP_BY_NAME(NiagaraDataInterfaceWater)
 
 #define LOCTEXT_NAMESPACE "NiagaraDataInterfaceWater"
@@ -19,12 +22,25 @@ struct FNDIWater_InstanceData
 
 namespace NDIWaterPrivate
 {
+	const FName IsValidName(TEXT("IsValid"));
 	const FName GetWaterDataAtPointName(TEXT("GetWaterDataAtPoint"));
 	const FName GetWaterSurfaceInfoName(TEXT("GetWaterSurfaceInfo"));
 
 	const FName GetWaveParamLookupTableName(TEXT("GetWaveParamLookupTableOffset"));
 
-	void GetWaterSurfaceInfo(FVectorVMExternalFunctionContext& Context)
+	void VMIsValid(FVectorVMExternalFunctionContext& Context)
+	{
+		VectorVM::FUserPtrHandler<FNDIWater_InstanceData> InstData(Context);
+		FNDIOutputParam<bool> OutIsValid(Context);
+
+		const bool bIsValid = InstData->WaterBodyComponent.IsValid();
+		for (int32 i = 0; i < Context.GetNumInstances(); ++i)
+		{
+			OutIsValid.SetAndAdvance(bIsValid);
+		}
+	}
+
+	void VMGetWaterSurfaceInfo(FVectorVMExternalFunctionContext& Context)
 	{
 		QUICK_SCOPE_CYCLE_COUNTER(NiagaraDataInterfaceWater_GetWaterSurfaceInfo);
 
@@ -83,6 +99,7 @@ namespace NDIWaterPrivate
 				OutWaterSurfacePosition.SetAndAdvance(FVector3f::ZeroVector);
 				OutWaterDepth.SetAndAdvance(0.0f);
 				OutWaterVelocity.SetAndAdvance(FVector3f::ZeroVector);
+				OutInExclusionVolume.SetAndAdvance(false);
 			}
 		}
 	}
@@ -114,6 +131,17 @@ void UNiagaraDataInterfaceWater::PostInitProperties()
 #if WITH_EDITORONLY_DATA
 void UNiagaraDataInterfaceWater::GetFunctionsInternal(TArray<FNiagaraFunctionSignature>& OutFunctions) const
 {
+	{
+		FNiagaraFunctionSignature & Sig = OutFunctions.AddDefaulted_GetRef();
+		Sig.Name = NDIWaterPrivate::IsValidName;
+		Sig.Inputs.Emplace(FNiagaraTypeDefinition(GetClass()), TEXT("Water"));
+		Sig.Outputs.Emplace(FNiagaraTypeDefinition::GetBoolDef(), TEXT("IsValid"));
+
+		Sig.bMemberFunction = true;
+		Sig.bExperimental = true;
+		Sig.SetDescription(LOCTEXT("DataInterfaceWater_IsValidDesc", "Returns true if we are reading from a valid water component, false if not."));
+	}
+
 	{
 		FNiagaraFunctionSignature& Sig = OutFunctions.AddDefaulted_GetRef();
 		Sig.Name = NDIWaterPrivate::GetWaterSurfaceInfoName;
@@ -176,9 +204,13 @@ DEFINE_NDI_DIRECT_FUNC_BINDER(UNiagaraDataInterfaceWater, GetWaveParamLookupTabl
 
 void UNiagaraDataInterfaceWater::GetVMExternalFunction(const FVMExternalFunctionBindingInfo& BindingInfo, void* InstanceData, FVMExternalFunction& OutFunc)
 {
-	if (BindingInfo.Name == NDIWaterPrivate::GetWaterSurfaceInfoName)
+	if (BindingInfo.Name == NDIWaterPrivate::IsValidName)
+	{
+		OutFunc = FVMExternalFunction::CreateStatic(&NDIWaterPrivate::VMIsValid);
+	}
+	else if (BindingInfo.Name == NDIWaterPrivate::GetWaterSurfaceInfoName)
 	{
-		OutFunc = FVMExternalFunction::CreateStatic(&NDIWaterPrivate::GetWaterSurfaceInfo);
+		OutFunc = FVMExternalFunction::CreateStatic(&NDIWaterPrivate::VMGetWaterSurfaceInfo);
 	}
 	else if (BindingInfo.Name == NDIWaterPrivate::GetWaterDataAtPointName)
 	{
@@ -203,7 +235,9 @@ bool UNiagaraDataInterfaceWater::Equals(const UNiagaraDataInterface* Other) cons
 		return false;
 	}
 	const UNiagaraDataInterfaceWater* OtherTyped = CastChecked<const UNiagaraDataInterfaceWater>(Other);
-	return OtherTyped->SourceBodyComponent == SourceBodyComponent;
+	return
+		OtherTyped->bFindWaterBodyOnSpawn == bFindWaterBodyOnSpawn &&
+		OtherTyped->SourceBodyComponent == SourceBodyComponent;
 }
 
 bool UNiagaraDataInterfaceWater::CopyToInternal(UNiagaraDataInterface* Destination) const
@@ -214,6 +248,7 @@ bool UNiagaraDataInterfaceWater::CopyToInternal(UNiagaraDataInterface* Destinati
 	}
 
 	UNiagaraDataInterfaceWater* OtherTyped = CastChecked<UNiagaraDataInterfaceWater>(Destination);
+	OtherTyped->bFindWaterBodyOnSpawn = bFindWaterBodyOnSpawn;
 	OtherTyped->SourceBodyComponent = SourceBodyComponent;
 
 	return true;
@@ -231,6 +266,32 @@ bool UNiagaraDataInterfaceWater::InitPerInstanceData(void* PerInstanceData, FNia
 	InstData->WaterBodyComponent = SourceBodyComponent;
 	InstData->LWCConverter = SystemInstance->GetLWCConverter();
 
+	if (SourceBodyComponent == nullptr && bFindWaterBodyOnSpawn)
+	{
+		const FVector QueryLocation = SystemInstance->GetWorldTransform().GetTranslation();
+		const EWaterBodyQueryFlags QueryFlags = (EWaterBodyQueryFlags::ComputeLocation | EWaterBodyQueryFlags::IncludeWaves);
+		TOptional<float> ClosestWaterSq;
+
+		for (TObjectIterator<AWaterBody> It; It; ++It)
+		{
+			AWaterBody* WaterBody = *It;
+			UWaterBodyComponent* WaterComponent = WaterBody->GetWaterBodyComponent();
+			const FWaterBodyQueryResult CurrentQueryResult = WaterComponent->QueryWaterInfoClosestToWorldLocation(QueryLocation, QueryFlags);
+			if (CurrentQueryResult.IsInExclusionVolume())
+			{
+				continue;
+			}
+
+			const FVector WaterLocation = CurrentQueryResult.GetWaterPlaneLocation();
+			const float DistanceSq = (WaterLocation - QueryLocation).SquaredLength();
+			if (!ClosestWaterSq.IsSet() || DistanceSq < ClosestWaterSq.GetValue())
+			{
+				ClosestWaterSq = DistanceSq;
+				InstData->WaterBodyComponent = WaterComponent;
+			}
+		}
+	}
+
 	return true;
 }
 
@@ -245,7 +306,7 @@ bool UNiagaraDataInterfaceWater::PerInstanceTick(void* PerInstanceData, FNiagara
 	check(SystemInstance);
 	FNDIWater_InstanceData* InstData = static_cast<FNDIWater_InstanceData*>(PerInstanceData);
 
-	if (InstData->WaterBodyComponent != SourceBodyComponent)
+	if (!bFindWaterBodyOnSpawn || (SourceBodyComponent != nullptr))
 	{
 		InstData->WaterBodyComponent = SourceBodyComponent;
 	}
diff --git a/Engine/Plugins/Experimental/Water/Source/Runtime/Public/NiagaraDataInterfaceWater.h b/Engine/Plugins/Experimental/Water/Source/Runtime/Public/NiagaraDataInterfaceWater.h
--- a/Engine/Plugins/Experimental/Water/Source/Runtime/Public/NiagaraDataInterfaceWater.h
+++ b/Engine/Plugins/Experimental/Water/Source/Runtime/Public/NiagaraDataInterfaceWater.h
@@ -47,6 +47,9 @@ protected:
 #endif
 
 private:
+	UPROPERTY(EditAnywhere, Category = "Water")
+	bool bFindWaterBodyOnSpawn = true;
+
 	UPROPERTY(EditAnywhere, Category = "Water") 
 	TObjectPtr<UWaterBodyComponent> SourceBodyComponent;
 };
