Fixed river mesh components remaining visible after a water landscape brush render.

diff --git a/Engine/Plugins/Experimental/Water/Source/Editor/Private/WaterBrushManager.cpp b/Engine/Plugins/Experimental/Water/Source/Editor/Private/WaterBrushManager.cpp
--- a/Engine/Plugins/Experimental/Water/Source/Editor/Private/WaterBrushManager.cpp
+++ b/Engine/Plugins/Experimental/Water/Source/Editor/Private/WaterBrushManager.cpp
@@ -449,14 +449,33 @@ void AWaterBrushManager::CaptureMeshDepth(const TArrayView<UStaticMeshComponent*
 {
 	SceneCaptureComponent2D->ClearShowOnlyComponents();
 	SceneCaptureComponent2D->ShowOnlyActors.Empty();
-	for (UStaticMeshComponent* PrimitiveComponent : MeshComponents)
+
+	TArray<bool> PreviousVisibilities;
+	TArray<bool> PreviousHiddenInGame;
+	PreviousVisibilities.SetNum(MeshComponents.Num());
+	PreviousHiddenInGame.SetNum(MeshComponents.Num());
+
+	for (int32 PrimitiveComponentIndex = 0; PrimitiveComponentIndex < MeshComponents.Num(); ++PrimitiveComponentIndex)
 	{
+		UStaticMeshComponent* PrimitiveComponent = MeshComponents[PrimitiveComponentIndex];
+		PreviousVisibilities[PrimitiveComponentIndex] = PrimitiveComponent->GetVisibleFlag();
+		PreviousHiddenInGame[PrimitiveComponentIndex] = PrimitiveComponent->bHiddenInGame;
+
 		PrimitiveComponent->SetVisibility(true);
 		PrimitiveComponent->SetHiddenInGame(false);
+
 		SceneCaptureComponent2D->ShowOnlyComponent(PrimitiveComponent);
 	}
+
 	SceneCaptureComponent2D->CaptureScene();
 
+	for (int32 PrimitiveComponentIndex = 0; PrimitiveComponentIndex < MeshComponents.Num(); ++PrimitiveComponentIndex)
+	{
+		UStaticMeshComponent* PrimitiveComponent = MeshComponents[PrimitiveComponentIndex];
+		PrimitiveComponent->SetVisibility(PreviousVisibilities[PrimitiveComponentIndex]);
+		PrimitiveComponent->SetHiddenInGame(PreviousHiddenInGame[PrimitiveComponentIndex]);
+	}
+
 	// Avoid keeping references to Captured components
 	SceneCaptureComponent2D->ClearShowOnlyComponents();
 }
