- Change to TActorIterator when searching for the water body
- Move the search into update as init might not have the parameter set
- Handle water bodies with no component
- Default search to disabled

#rnx

diff --git a/Engine/Plugins/Experimental/Water/Source/Runtime/Private/NiagaraDataInterfaceWater.cpp b/Engine/Plugins/Experimental/Water/Source/Runtime/Private/NiagaraDataInterfaceWater.cpp
--- a/Engine/Plugins/Experimental/Water/Source/Runtime/Private/NiagaraDataInterfaceWater.cpp
+++ b/Engine/Plugins/Experimental/Water/Source/Runtime/Private/NiagaraDataInterfaceWater.cpp
@@ -7,7 +7,7 @@
 #include "WaterBodyComponent.h"
 #include "WaterModule.h"
 
-#include "UObject/UObjectIterator.h"
+#include "EngineUtils.h"
 
 #include UE_INLINE_GENERATED_CPP_BY_NAME(NiagaraDataInterfaceWater)
 
@@ -15,9 +15,9 @@
 
 struct FNDIWater_InstanceData
 {
-	//Cached ptr to component we sample from. 
-	TWeakObjectPtr<UWaterBodyComponent> WaterBodyComponent;
-	FNiagaraLWCConverter LWCConverter;
+	bool								bFindClosestBody = false;
+	TWeakObjectPtr<UWaterBodyComponent>	WaterBodyComponent;
+	FNiagaraLWCConverter				LWCConverter;
 };
 
 namespace NDIWaterPrivate
@@ -28,6 +28,43 @@ namespace NDIWaterPrivate
 
 	const FName GetWaveParamLookupTableName(TEXT("GetWaveParamLookupTableOffset"));
 
+	UWaterBodyComponent* FindClosestWaterBody(UWorld* World, const FVector QueryLocation)
+	{
+		if (World == nullptr)
+		{
+			return nullptr;
+		}
+
+		const EWaterBodyQueryFlags QueryFlags = (EWaterBodyQueryFlags::ComputeLocation | EWaterBodyQueryFlags::IncludeWaves);
+		UWaterBodyComponent* ClosestComponent = nullptr;
+		double ClosestWaterSq = 0.0f;
+
+		for (TActorIterator<AWaterBody> It(World); It; ++It)
+		{
+			AWaterBody* WaterBody = *It;
+			UWaterBodyComponent* WaterComponent = WaterBody ? WaterBody->GetWaterBodyComponent() : nullptr;
+			if (!WaterComponent)
+			{
+				continue;
+			}
+
+			const FWaterBodyQueryResult CurrentQueryResult = WaterComponent->QueryWaterInfoClosestToWorldLocation(QueryLocation, QueryFlags);
+			if (CurrentQueryResult.IsInExclusionVolume())
+			{
+				continue;
+			}
+
+			const FVector WaterLocation = CurrentQueryResult.GetWaterPlaneLocation();
+			const double DistanceSq = (WaterLocation - QueryLocation).SquaredLength();
+			if (!ClosestComponent || DistanceSq < ClosestWaterSq)
+			{
+				ClosestWaterSq = DistanceSq;
+				ClosestComponent = WaterComponent;
+			}
+		}
+		return ClosestComponent;
+	}
+
 	void VMIsValid(FVectorVMExternalFunctionContext& Context)
 	{
 		VectorVM::FUserPtrHandler<FNDIWater_InstanceData> InstData(Context);
@@ -262,35 +299,9 @@ int32 UNiagaraDataInterfaceWater::PerInstanceDataSize() const
 bool UNiagaraDataInterfaceWater::InitPerInstanceData(void* PerInstanceData, FNiagaraSystemInstance* SystemInstance)
 {
 	FNDIWater_InstanceData* InstData = new (PerInstanceData) FNDIWater_InstanceData();
-
-	InstData->WaterBodyComponent = SourceBodyComponent;
-	InstData->LWCConverter = SystemInstance->GetLWCConverter();
-
-	if (SourceBodyComponent == nullptr && bFindWaterBodyOnSpawn)
-	{
-		const FVector QueryLocation = SystemInstance->GetWorldTransform().GetTranslation();
-		const EWaterBodyQueryFlags QueryFlags = (EWaterBodyQueryFlags::ComputeLocation | EWaterBodyQueryFlags::IncludeWaves);
-		TOptional<float> ClosestWaterSq;
-
-		for (TObjectIterator<AWaterBody> It; It; ++It)
-		{
-			AWaterBody* WaterBody = *It;
-			UWaterBodyComponent* WaterComponent = WaterBody->GetWaterBodyComponent();
-			const FWaterBodyQueryResult CurrentQueryResult = WaterComponent->QueryWaterInfoClosestToWorldLocation(QueryLocation, QueryFlags);
-			if (CurrentQueryResult.IsInExclusionVolume())
-			{
-				continue;
-			}
-
-			const FVector WaterLocation = CurrentQueryResult.GetWaterPlaneLocation();
-			const float DistanceSq = (WaterLocation - QueryLocation).SquaredLength();
-			if (!ClosestWaterSq.IsSet() || DistanceSq < ClosestWaterSq.GetValue())
-			{
-				ClosestWaterSq = DistanceSq;
-				InstData->WaterBodyComponent = WaterComponent;
-			}
-		}
-	}
+	InstData->bFindClosestBody		= SourceBodyComponent == nullptr && bFindWaterBodyOnSpawn;
+	InstData->WaterBodyComponent	= SourceBodyComponent;
+	InstData->LWCConverter			= SystemInstance->GetLWCConverter();
 
 	return true;
 }
@@ -306,10 +317,18 @@ bool UNiagaraDataInterfaceWater::PerInstanceTick(void* PerInstanceData, FNiagara
 	check(SystemInstance);
 	FNDIWater_InstanceData* InstData = static_cast<FNDIWater_InstanceData*>(PerInstanceData);
 
-	if (!bFindWaterBodyOnSpawn || (SourceBodyComponent != nullptr))
+	// If the search for closest was enabled, perform the search
+	// Note: we do this here rather than in Init as the system might be auto activate and the user parameter not set until after the spawn
+	if (SourceBodyComponent == nullptr && InstData->bFindClosestBody)
+	{
+		const FVector QueryLocation = SystemInstance->GetWorldTransform().GetTranslation();
+		InstData->WaterBodyComponent = NDIWaterPrivate::FindClosestWaterBody(SystemInstance->GetWorld(), QueryLocation);
+	}
+	else
 	{
 		InstData->WaterBodyComponent = SourceBodyComponent;
 	}
+	InstData->bFindClosestBody = false;
 
 	return false;
 }
diff --git a/Engine/Plugins/Experimental/Water/Source/Runtime/Public/NiagaraDataInterfaceWater.h b/Engine/Plugins/Experimental/Water/Source/Runtime/Public/NiagaraDataInterfaceWater.h
--- a/Engine/Plugins/Experimental/Water/Source/Runtime/Public/NiagaraDataInterfaceWater.h
+++ b/Engine/Plugins/Experimental/Water/Source/Runtime/Public/NiagaraDataInterfaceWater.h
@@ -48,7 +48,7 @@ protected:
 
 private:
 	UPROPERTY(EditAnywhere, Category = "Water")
-	bool bFindWaterBodyOnSpawn = true;
+	bool bFindWaterBodyOnSpawn = false;
 
 	UPROPERTY(EditAnywhere, Category = "Water") 
 	TObjectPtr<UWaterBodyComponent> SourceBodyComponent;
