Fixed rivers rendering in the water brush : the scene capture component that's used internally was rendering Fog somehow

diff --git a/Engine/Plugins/Experimental/Water/Source/Editor/Private/WaterBrushManager.cpp b/Engine/Plugins/Experimental/Water/Source/Editor/Private/WaterBrushManager.cpp
--- a/Engine/Plugins/Experimental/Water/Source/Editor/Private/WaterBrushManager.cpp
+++ b/Engine/Plugins/Experimental/Water/Source/Editor/Private/WaterBrushManager.cpp
@@ -59,6 +59,7 @@ AWaterBrushManager::AWaterBrushManager(const FObjectInitializer& ObjectInitializ
 	ShowFlagSettings.Add(FEngineShowFlagsSetting { TEXT("Atmosphere"), false } );
 	ShowFlagSettings.Add(FEngineShowFlagsSetting { TEXT("Bloom"), false } );
 	ShowFlagSettings.Add(FEngineShowFlagsSetting { TEXT("Lighting"), false } );
+	ShowFlagSettings.Add(FEngineShowFlagsSetting { TEXT("Fog"), false } );
 	SceneCaptureComponent2D->SetShowFlagSettings(ShowFlagSettings);
 
 	PrimaryActorTick.TickGroup = ETickingGroup::TG_PrePhysics;
diff --git a/Engine/Plugins/Experimental/Water/Source/Editor/Private/WaterBrushManager.h b/Engine/Plugins/Experimental/Water/Source/Editor/Private/WaterBrushManager.h
--- a/Engine/Plugins/Experimental/Water/Source/Editor/Private/WaterBrushManager.h
+++ b/Engine/Plugins/Experimental/Water/Source/Editor/Private/WaterBrushManager.h
@@ -30,10 +30,10 @@ class AWaterBrushManager : public AWaterLandscapeBrush
 public:
 	GENERATED_BODY()
 	
- 	UPROPERTY()
+ 	UPROPERTY(Transient)
 	TObjectPtr<USceneCaptureComponent2D> SceneCaptureComponent2D = nullptr;
 
-	UPROPERTY()
+	UPROPERTY(Transient)
 	TObjectPtr<UJumpFloodComponent2D> JumpFloodComponent2D = nullptr;
 
 	// RTs
@@ -270,6 +270,7 @@ private:
 	virtual void PostEditChangeProperty(FPropertyChangedEvent& PropertyChangedEvent) override;
 	virtual bool AllocateRTs();
 	virtual void SetMPCParams();
+
 	virtual bool SetupRiverSplineRenderMIDs(const FBrushActorRenderContext& BrushActorRenderContext, bool bRestoreMIDs, TArray<UMaterialInterface*>& InOutMIDs);
 	virtual void CaptureMeshDepth(const TArrayView<UStaticMeshComponent*>& MeshComponents);
 	virtual void CacheBrushDistanceField(const FBrushActorRenderContext& BrushActorRenderContext);
