[Water]: Expose the water body exclusion volume list and island list so it can be manually edited per instance.

diff --git a/Engine/Plugins/Experimental/Water/Source/Runtime/Public/WaterBodyComponent.h b/Engine/Plugins/Experimental/Water/Source/Runtime/Public/WaterBodyComponent.h
--- a/Engine/Plugins/Experimental/Water/Source/Runtime/Public/WaterBodyComponent.h
+++ b/Engine/Plugins/Experimental/Water/Source/Runtime/Public/WaterBodyComponent.h
@@ -652,10 +652,10 @@ protected:
 	TObjectPtr<UMaterialInstanceDynamic> WaterInfoMID;
 
 	/** Islands in this water body*/
-	UPROPERTY(Category = Water, VisibleAnywhere, AdvancedDisplay)
+	UPROPERTY(Category = Water, EditAnywhere, AdvancedDisplay)
 	TArray<TSoftObjectPtr<AWaterBodyIsland>> WaterBodyIslands;
 
-	UPROPERTY(Category = Water, VisibleAnywhere, AdvancedDisplay)
+	UPROPERTY(Category = Water, EditAnywhere, AdvancedDisplay)
 	TArray<TSoftObjectPtr<AWaterBodyExclusionVolume>> WaterBodyExclusionVolumes;
 
 	UPROPERTY(Transient)
