[Water]: Fixed crash due to null water body static mesh components when running pie with server in separate process. Water Info meshes are nulled out during server loading which caused a crash due to the assumption that the meshes are valid in the water body mesh builder.

diff --git a/Engine/Plugins/Experimental/Water/Source/Runtime/Private/WaterBodyComponent.cpp b/Engine/Plugins/Experimental/Water/Source/Runtime/Private/WaterBodyComponent.cpp
--- a/Engine/Plugins/Experimental/Water/Source/Runtime/Private/WaterBodyComponent.cpp
+++ b/Engine/Plugins/Experimental/Water/Source/Runtime/Private/WaterBodyComponent.cpp
@@ -1944,16 +1944,20 @@ void UWaterBodyComponent::SetWaterBodyStaticMeshEnabled(bool bEnabled)
 
 void UWaterBodyComponent::UpdateWaterBodyRenderData()
 {
-	// Avoid updating any mesh data if we are in a PIE world if dynamic data changes aren't allowed.
-	if (const UWorld* World = GetWorld())
+	// Don't need to update render data if we can't ever render
+	if (FApp::CanEverRender())
 	{
-		if (!World->HasBegunPlay() || AreDynamicDataChangesAllowed())
+		if (const UWorld* World = GetWorld())
 		{
-			UpdateWaterInfoMeshComponents();
+			// Avoid updating any mesh data if we are in a PIE world if dynamic data changes aren't allowed.
+			if (!World->HasBegunPlay() || AreDynamicDataChangesAllowed())
+			{
+				UpdateWaterInfoMeshComponents();
 
-			UpdateWaterBodyStaticMeshComponents();
+				UpdateWaterBodyStaticMeshComponents();
 
-			OnWaterBodyRenderDataUpdated();
+				OnWaterBodyRenderDataUpdated();
+			}
 		}
 	}
 }
diff --git a/Engine/Plugins/Experimental/Water/Source/Runtime/Private/WaterBodyMeshBuilder.cpp b/Engine/Plugins/Experimental/Water/Source/Runtime/Private/WaterBodyMeshBuilder.cpp
--- a/Engine/Plugins/Experimental/Water/Source/Runtime/Private/WaterBodyMeshBuilder.cpp
+++ b/Engine/Plugins/Experimental/Water/Source/Runtime/Private/WaterBodyMeshBuilder.cpp
@@ -37,6 +37,11 @@ void FWaterBodyMeshBuilder::BuildWaterInfoMeshes(UWaterBodyComponent* WaterBodyC
 	TRACE_CPUPROFILER_EVENT_SCOPE(FWaterBodyMeshBuilder::BuildWaterInfoMesh);
 	
 	using namespace UE::Geometry;
+	
+	if (!ensureMsgf(WaterInfoMeshComponent != nullptr && WaterInfoDilatedMeshComponent != nullptr, TEXT("BuildWaterInfoMeshes called on a water body with null water info mesh component %s"), *WaterBodyComponent->GetOwner()->GetActorNameOrLabel()))
+	{
+		return;
+	}
 
 	FDynamicMesh3 WaterInfoMesh(EMeshComponents::None);
 	FDynamicMesh3 WaterInfoDilatedMesh(EMeshComponents::None);
