[Water]
Removed Transactional flag from water body static meshes. Serializing static meshes into the transaction buffer causes the async static mesh build to sync with the main thread.

diff --git a/Engine/Plugins/Experimental/Water/Source/Runtime/Private/WaterBodyMeshBuilder.cpp b/Engine/Plugins/Experimental/Water/Source/Runtime/Private/WaterBodyMeshBuilder.cpp
--- a/Engine/Plugins/Experimental/Water/Source/Runtime/Private/WaterBodyMeshBuilder.cpp
+++ b/Engine/Plugins/Experimental/Water/Source/Runtime/Private/WaterBodyMeshBuilder.cpp
@@ -597,7 +597,7 @@ void FWaterBodyMeshBuilder::UpdateStaticMesh(UStaticMesh* WaterMesh, const FMesh
 
 UStaticMesh* FWaterBodyMeshBuilder::CreateUStaticMesh(UObject* Outer, FName MeshName) const
 {
-	UStaticMesh* StaticMesh = NewObject<UStaticMesh>(Outer, MeshName, RF_Transactional | RF_TextExportTransient | RF_NonPIEDuplicateTransient);
+	UStaticMesh* StaticMesh = NewObject<UStaticMesh>(Outer, MeshName, RF_TextExportTransient | RF_NonPIEDuplicateTransient);
 
 	// Disable navigation
 	StaticMesh->MarkAsNotHavingNavigationData();
