[Water]
Speculative fix for crash on map change in FWaterViewExtension::SetupView. Changed water zone and water body references in the view extension to be weak pointers so we can detect GC'd entries.

diff --git a/Engine/Plugins/Experimental/Water/Source/Runtime/Private/WaterInfoRendering.cpp b/Engine/Plugins/Experimental/Water/Source/Runtime/Private/WaterInfoRendering.cpp
--- a/Engine/Plugins/Experimental/Water/Source/Runtime/Private/WaterInfoRendering.cpp
+++ b/Engine/Plugins/Experimental/Water/Source/Runtime/Private/WaterInfoRendering.cpp
@@ -774,19 +774,17 @@ void UpdateWaterInfoRendering(
 	{
 		ComponentsToRenderInColorPass.Reserve(Context.WaterBodies.Num());
 		ComponentsToRenderInDilationPass.Reserve(Context.WaterBodies.Num());
-		for (const UWaterBodyComponent* WaterBodyToRender : Context.WaterBodies)
+		for (const TWeakObjectPtr<UWaterBodyComponent> WaterBodyToRenderPtr : Context.WaterBodies)
 		{
-			if (!IsValid(WaterBodyToRender))
+			if (UWaterBodyComponent* WaterBodyToRender = WaterBodyToRenderPtr.Get())
 			{
-				continue;
-			}
-
-			// Perform our own simple culling based on the known Capture bounds:
-			const FBox WaterBodyBounds = WaterBodyToRender->Bounds.GetBox();
-			if (CaptureBounds.Intersect(FBox2D(FVector2D(WaterBodyBounds.Min), FVector2D(WaterBodyBounds.Max))))
-			{
-				ComponentsToRenderInColorPass.Add(WaterBodyToRender->GetWaterInfoMeshComponent()->GetPrimitiveSceneId());
-				ComponentsToRenderInDilationPass.Add(WaterBodyToRender->GetDilatedWaterInfoMeshComponent()->GetPrimitiveSceneId());
+				// Perform our own simple culling based on the known Capture bounds:
+				const FBox WaterBodyBounds = WaterBodyToRender->Bounds.GetBox();
+				if (CaptureBounds.Intersect(FBox2D(FVector2D(WaterBodyBounds.Min), FVector2D(WaterBodyBounds.Max))))
+				{
+					ComponentsToRenderInColorPass.Add(WaterBodyToRender->GetWaterInfoMeshComponent()->GetPrimitiveSceneId());
+					ComponentsToRenderInDilationPass.Add(WaterBodyToRender->GetDilatedWaterInfoMeshComponent()->GetPrimitiveSceneId());
+				}
 			}
 		}
 	}
@@ -824,75 +822,76 @@ void UpdateWaterInfoRendering(
 		});
 }
 
-void UpdateWaterInfoRendering2(FSceneView& InView, const TMap<AWaterZone*, UE::WaterInfo::FRenderingContext>& WaterInfoContexts)
+void UpdateWaterInfoRendering2(FSceneView& InView, const TWeakObjectPtrKeyMap<AWaterZone, UE::WaterInfo::FRenderingContext> WaterInfoContexts)
 {
 	TRACE_CPUPROFILER_EVENT_SCOPE(WaterInfo::UpdateWaterInfoRendering2);
 
 	InView.WaterInfoTextureRenderingParams.Reset();
-	for (const TPair<AWaterZone*, UE::WaterInfo::FRenderingContext>& Pair : WaterInfoContexts)
+	for (const TPair<TWeakObjectPtr<AWaterZone>, UE::WaterInfo::FRenderingContext>& Pair : WaterInfoContexts)
 	{
-		const UE::WaterInfo::FRenderingContext& Context(Pair.Value);
-
-		if (!IsValid(Context.TextureRenderTarget))
+		if (Pair.Key.IsValid())
 		{
-			continue;
-		}
-		const FVector ZoneExtent = Context.ZoneToRender->GetDynamicWaterInfoExtent();
-
-		FVector ViewLocation = Context.ZoneToRender->GetDynamicWaterInfoCenter();
-		ViewLocation.Z = Context.CaptureZ;
-
-		const FBox2D CaptureBounds(FVector2D(ViewLocation - ZoneExtent), FVector2D(ViewLocation + ZoneExtent));
-
-		// Zone rendering always happens facing towards negative z.
-		const FVector LookAt = ViewLocation - FVector(0.f, 0.f, 1.f);
-
-		FSceneView::FWaterInfoTextureRenderingParams RenderingParams;
-		RenderingParams.RenderTarget = Context.TextureRenderTarget->GameThread_GetRenderTargetResource();
-		RenderingParams.ViewLocation = ViewLocation;
-		RenderingParams.ViewRotationMatrix = FLookAtMatrix(ViewLocation, LookAt, FVector(0.f, -1.f, 0.f));
-		RenderingParams.ViewRotationMatrix = RenderingParams.ViewRotationMatrix.RemoveTranslation();
-		RenderingParams.ViewRotationMatrix.RemoveScaling();
-		RenderingParams.ProjectionMatrix = BuildOrthoMatrix(ZoneExtent.X, ZoneExtent.Y);
-		RenderingParams.CaptureZ = ViewLocation.Z;
-		RenderingParams.WaterHeightExtents = Context.ZoneToRender->GetWaterHeightExtents();
-		RenderingParams.GroundZMin = Context.ZoneToRender->GetGroundZMin();
-		RenderingParams.VelocityBlurRadius = Context.ZoneToRender->GetVelocityBlurRadius();
-		RenderingParams.WaterZoneExtents = ZoneExtent;
-
-		if (Context.GroundPrimitiveComponents.Num() > 0)
-		{
-			RenderingParams.TerrainComponentIds.Reserve(Context.GroundPrimitiveComponents.Num());
-			for (TWeakObjectPtr<UPrimitiveComponent> GroundPrimComp : Context.GroundPrimitiveComponents)
+			const UE::WaterInfo::FRenderingContext& Context(Pair.Value);
+
+			if (!IsValid(Context.TextureRenderTarget))
 			{
-				if (GroundPrimComp.IsValid())
-				{
-					RenderingParams.TerrainComponentIds.Add(GroundPrimComp.Get()->GetPrimitiveSceneId());
-				}
+				continue;
 			}
-		}
-		if (Context.WaterBodies.Num() > 0)
-		{
-			RenderingParams.WaterBodyComponentIds.Reserve(Context.WaterBodies.Num());
-			RenderingParams.DilatedWaterBodyComponentIds.Reserve(Context.WaterBodies.Num());
-			for (const UWaterBodyComponent* WaterBodyToRender : Context.WaterBodies)
+			const FVector ZoneExtent = Context.ZoneToRender->GetDynamicWaterInfoExtent();
+
+			FVector ViewLocation = Context.ZoneToRender->GetDynamicWaterInfoCenter();
+			ViewLocation.Z = Context.CaptureZ;
+
+			const FBox2D CaptureBounds(FVector2D(ViewLocation - ZoneExtent), FVector2D(ViewLocation + ZoneExtent));
+
+			// Zone rendering always happens facing towards negative z.
+			const FVector LookAt = ViewLocation - FVector(0.f, 0.f, 1.f);
+
+			FSceneView::FWaterInfoTextureRenderingParams RenderingParams;
+			RenderingParams.RenderTarget = Context.TextureRenderTarget->GameThread_GetRenderTargetResource();
+			RenderingParams.ViewLocation = ViewLocation;
+			RenderingParams.ViewRotationMatrix = FLookAtMatrix(ViewLocation, LookAt, FVector(0.f, -1.f, 0.f));
+			RenderingParams.ViewRotationMatrix = RenderingParams.ViewRotationMatrix.RemoveTranslation();
+			RenderingParams.ViewRotationMatrix.RemoveScaling();
+			RenderingParams.ProjectionMatrix = BuildOrthoMatrix(ZoneExtent.X, ZoneExtent.Y);
+			RenderingParams.CaptureZ = ViewLocation.Z;
+			RenderingParams.WaterHeightExtents = Context.ZoneToRender->GetWaterHeightExtents();
+			RenderingParams.GroundZMin = Context.ZoneToRender->GetGroundZMin();
+			RenderingParams.VelocityBlurRadius = Context.ZoneToRender->GetVelocityBlurRadius();
+			RenderingParams.WaterZoneExtents = ZoneExtent;
+
+			if (Context.GroundPrimitiveComponents.Num() > 0)
 			{
-				if (!IsValid(WaterBodyToRender))
+				RenderingParams.TerrainComponentIds.Reserve(Context.GroundPrimitiveComponents.Num());
+				for (TWeakObjectPtr<UPrimitiveComponent> GroundPrimComp : Context.GroundPrimitiveComponents)
 				{
-					continue;
+					if (GroundPrimComp.IsValid())
+					{
+						RenderingParams.TerrainComponentIds.Add(GroundPrimComp.Get()->GetPrimitiveSceneId());
+					}
 				}
-
-				// Perform our own simple culling based on the known Capture bounds:
-				const FBox WaterBodyBounds = WaterBodyToRender->Bounds.GetBox();
-				if (CaptureBounds.Intersect(FBox2D(FVector2D(WaterBodyBounds.Min), FVector2D(WaterBodyBounds.Max))))
+			}
+			if (Context.WaterBodies.Num() > 0)
+			{
+				RenderingParams.WaterBodyComponentIds.Reserve(Context.WaterBodies.Num());
+				RenderingParams.DilatedWaterBodyComponentIds.Reserve(Context.WaterBodies.Num());
+				for (const TWeakObjectPtr<UWaterBodyComponent> WaterBodyToRenderPtr : Context.WaterBodies)
 				{
-					RenderingParams.WaterBodyComponentIds.Add(WaterBodyToRender->GetWaterInfoMeshComponent()->GetPrimitiveSceneId());
-					RenderingParams.DilatedWaterBodyComponentIds.Add(WaterBodyToRender->GetDilatedWaterInfoMeshComponent()->GetPrimitiveSceneId());
+					if (UWaterBodyComponent* WaterBodyToRender = WaterBodyToRenderPtr.Get())
+					{
+						// Perform our own simple culling based on the known Capture bounds:
+						const FBox WaterBodyBounds = WaterBodyToRender->Bounds.GetBox();
+						if (CaptureBounds.Intersect(FBox2D(FVector2D(WaterBodyBounds.Min), FVector2D(WaterBodyBounds.Max))))
+						{
+							RenderingParams.WaterBodyComponentIds.Add(WaterBodyToRender->GetWaterInfoMeshComponent()->GetPrimitiveSceneId());
+							RenderingParams.DilatedWaterBodyComponentIds.Add(WaterBodyToRender->GetDilatedWaterInfoMeshComponent()->GetPrimitiveSceneId());
+						}
+					}
 				}
 			}
-		}
 
-		InView.WaterInfoTextureRenderingParams.Add(MoveTemp(RenderingParams));
+			InView.WaterInfoTextureRenderingParams.Add(MoveTemp(RenderingParams));
+		}
 	}
 }
 
@@ -1124,19 +1123,17 @@ void UpdateWaterInfoRendering_CustomRenderPass(
 	{
 		ComponentsToRenderInColorPass.Reserve(Context.WaterBodies.Num());
 		ComponentsToRenderInDilationPass.Reserve(Context.WaterBodies.Num());
-		for (const UWaterBodyComponent* WaterBodyToRender : Context.WaterBodies)
+		for (const TWeakObjectPtr<UWaterBodyComponent> WaterBodyToRenderPtr : Context.WaterBodies)
 		{
-			if (!IsValid(WaterBodyToRender))
-			{
-				continue;
-			}
-
-			// Perform our own simple culling based on the known Capture bounds:
-			const FBox WaterBodyBounds = WaterBodyToRender->Bounds.GetBox();
-			if (CaptureBounds.Intersect(FBox2D(FVector2D(WaterBodyBounds.Min), FVector2D(WaterBodyBounds.Max))))
+			if (UWaterBodyComponent* WaterBodyToRender = WaterBodyToRenderPtr.Get())
 			{
-				ComponentsToRenderInColorPass.Add(WaterBodyToRender->GetWaterInfoMeshComponent()->GetPrimitiveSceneId());
-				ComponentsToRenderInDilationPass.Add(WaterBodyToRender->GetDilatedWaterInfoMeshComponent()->GetPrimitiveSceneId());
+				// Perform our own simple culling based on the known Capture bounds:
+				const FBox WaterBodyBounds = WaterBodyToRender->Bounds.GetBox();
+				if (CaptureBounds.Intersect(FBox2D(FVector2D(WaterBodyBounds.Min), FVector2D(WaterBodyBounds.Max))))
+				{
+					ComponentsToRenderInColorPass.Add(WaterBodyToRender->GetWaterInfoMeshComponent()->GetPrimitiveSceneId());
+					ComponentsToRenderInDilationPass.Add(WaterBodyToRender->GetDilatedWaterInfoMeshComponent()->GetPrimitiveSceneId());
+				}
 			}
 		}
 	}
diff --git a/Engine/Plugins/Experimental/Water/Source/Runtime/Private/WaterInfoRendering.h b/Engine/Plugins/Experimental/Water/Source/Runtime/Private/WaterInfoRendering.h
--- a/Engine/Plugins/Experimental/Water/Source/Runtime/Private/WaterInfoRendering.h
+++ b/Engine/Plugins/Experimental/Water/Source/Runtime/Private/WaterInfoRendering.h
@@ -12,6 +12,10 @@ class UPrimitiveComponent;
 class FSceneView;
 class FSceneViewFamily;
 
+
+template <typename KeyType, typename ValueType>
+using TWeakObjectPtrKeyMap = TMap<TWeakObjectPtr<KeyType>, ValueType, FDefaultSetAllocator, TWeakObjectPtrMapKeyFuncs<TWeakObjectPtr<KeyType>, ValueType>>;
+
 namespace UE::WaterInfo
 {
 
@@ -19,7 +23,7 @@ struct FRenderingContext
 {
 	AWaterZone* ZoneToRender = nullptr;
 	UTextureRenderTarget2DArray* TextureRenderTarget;
-	TArray<UWaterBodyComponent*> WaterBodies;
+	TArray<TWeakObjectPtr<UWaterBodyComponent>> WaterBodies;
 	TArray<TWeakObjectPtr<UPrimitiveComponent>> GroundPrimitiveComponents;
 	float CaptureZ;
 };
@@ -30,7 +34,7 @@ void UpdateWaterInfoRendering(
 
 void UpdateWaterInfoRendering2(
 	FSceneView& InView, 
-	const TMap<AWaterZone*, UE::WaterInfo::FRenderingContext>& WaterInfoContexts);
+	const TWeakObjectPtrKeyMap<AWaterZone, UE::WaterInfo::FRenderingContext> WaterInfoContexts);
 
 void UpdateWaterInfoRendering_CustomRenderPass(
 	FSceneInterface* Scene,
diff --git a/Engine/Plugins/Experimental/Water/Source/Runtime/Private/WaterViewExtension.cpp b/Engine/Plugins/Experimental/Water/Source/Runtime/Private/WaterViewExtension.cpp
--- a/Engine/Plugins/Experimental/Water/Source/Runtime/Private/WaterViewExtension.cpp
+++ b/Engine/Plugins/Experimental/Water/Source/Runtime/Private/WaterViewExtension.cpp
@@ -316,64 +316,68 @@ void FWaterViewExtension::SetupView(FSceneViewFamily& InViewFamily, FSceneView&
 	FSceneInterface* Scene = WorldPtr.Get()->Scene;
 	check(Scene != nullptr);
 
-	for (const TPair<AWaterZone*, UE::WaterInfo::FRenderingContext>& Pair : WaterInfoContextsToRender)
+	for (const TPair<TWeakObjectPtr<AWaterZone>, UE::WaterInfo::FRenderingContext>& Pair : WaterInfoContextsToRender)
 	{
-		AWaterZone* WaterZone = Pair.Key;
-		check(WaterZone != nullptr);
-		FWaterZoneInfo* WaterZoneInfo = WaterZoneInfos.Find(WaterZone);
-		if (ensureMsgf(WaterZoneInfo != nullptr, TEXT("We are trying to rendering a water info texture for a water zone that is not registered!")))
+		TWeakObjectPtr<AWaterZone> WaterZonePtr = Pair.Key;
+		if (AWaterZone* WaterZone = WaterZonePtr.Get())
 		{
-			if (WaterZone->IsLocalOnlyTessellationEnabled())
+			check(WaterZone != nullptr);
+
+			FWaterZoneInfo* WaterZoneInfo = WaterZoneInfos.Find(WaterZone);
+			if (ensureMsgf(WaterZoneInfo != nullptr, TEXT("We are trying to rendering a water info texture for a water zone that is not registered!")))
 			{
-				UWaterMeshComponent* WaterMesh= WaterZone->GetWaterMeshComponent();
-				check(WaterMesh);
+				if (WaterZone->IsLocalOnlyTessellationEnabled())
+				{
+					UWaterMeshComponent* WaterMesh= WaterZone->GetWaterMeshComponent();
+					check(WaterMesh);
 
-				const double TileSize = WaterMesh->GetTileSize();
+					const double TileSize = WaterMesh->GetTileSize();
 
-				const FVector WaterMeshCenter(ViewLocation.GridSnap(TileSize));
+					const FVector WaterMeshCenter(ViewLocation.GridSnap(TileSize));
 
-				const FVector2D WaterInfoHalfExtent(WaterZone->GetDynamicWaterInfoExtent() / 2.0);
-				const FBox2D WaterInfoBounds(FVector2D(WaterMeshCenter) - WaterInfoHalfExtent, FVector2D(WaterMeshCenter) + WaterInfoHalfExtent);
+					const FVector2D WaterInfoHalfExtent(WaterZone->GetDynamicWaterInfoExtent() / 2.0);
+					const FBox2D WaterInfoBounds(FVector2D(WaterMeshCenter) - WaterInfoHalfExtent, FVector2D(WaterMeshCenter) + WaterInfoHalfExtent);
 
-				WaterZone->SetLocalTessellationCenter(WaterMeshCenter);
+					WaterZone->SetLocalTessellationCenter(WaterMeshCenter);
 
-				// Trigger the next update when the camera is <UpdateMargin> units away from the border of the current window.
-				const FVector2D UpdateMargin(CVarLocalTessellationUpdateMargin.GetValueOnGameThread());
+					// Trigger the next update when the camera is <UpdateMargin> units away from the border of the current window.
+					const FVector2D UpdateMargin(CVarLocalTessellationUpdateMargin.GetValueOnGameThread());
 
-				// Keep a minimum of <1., 1.> bounds to avoid updating every frame if the update margin is larger than the zone.
-				const FVector2D UpdateExtents = FVector2D::Max(FVector2D(1., 1.), WaterInfoHalfExtent - UpdateMargin);
-				WaterZoneInfo->UpdateBounds.Emplace(FVector2D(WaterMeshCenter) - UpdateExtents, FVector2D(WaterMeshCenter) + UpdateExtents);
+					// Keep a minimum of <1., 1.> bounds to avoid updating every frame if the update margin is larger than the zone.
+					const FVector2D UpdateExtents = FVector2D::Max(FVector2D(1., 1.), WaterInfoHalfExtent - UpdateMargin);
+					WaterZoneInfo->UpdateBounds.Emplace(FVector2D(WaterMeshCenter) - UpdateExtents, FVector2D(WaterMeshCenter) + UpdateExtents);
 
-				const FVector2D WaterQuadTreeHalfExtent = WaterMesh->GetExtentInTiles() * TileSize;
-				const FVector2D WaterQuadTreeCenter = WaterMesh->GetDynamicWaterMeshCenter();
-				const FBox2D WaterQuadTreeBounds(WaterQuadTreeCenter - WaterQuadTreeHalfExtent, WaterQuadTreeCenter + WaterQuadTreeHalfExtent);
+					const FVector2D WaterQuadTreeHalfExtent = WaterMesh->GetExtentInTiles() * TileSize;
+					const FVector2D WaterQuadTreeCenter = WaterMesh->GetDynamicWaterMeshCenter();
+					const FBox2D WaterQuadTreeBounds(WaterQuadTreeCenter - WaterQuadTreeHalfExtent, WaterQuadTreeCenter + WaterQuadTreeHalfExtent);
 
-				// If the new water info bounds would be outside the water mesh, recenter and rebuild the water mesh
-				if (!WaterQuadTreeBounds.IsInside(WaterInfoBounds.ExpandBy(WaterInfoBounds.GetExtent())))
+					// If the new water info bounds would be outside the water mesh, recenter and rebuild the water mesh
+					if (!WaterQuadTreeBounds.IsInside(WaterInfoBounds.ExpandBy(WaterInfoBounds.GetExtent())))
+					{
+						WaterMesh->SetDynamicWaterMeshCenter(FVector2D(WaterMeshCenter));
+					}
+
+					// Mark GPU data dirty since we have a new WaterArea parameter and need to push this to water bodies.
+					MarkGPUDataDirty();
+				}
+				else
 				{
-					WaterMesh->SetDynamicWaterMeshCenter(FVector2D(WaterMeshCenter));
+					WaterZoneInfo->UpdateBounds.Reset();
 				}
+			}
 
-				// Mark GPU data dirty since we have a new WaterArea parameter and need to push this to water bodies.
-				MarkGPUDataDirty();
+			const UE::WaterInfo::FRenderingContext& Context(Pair.Value);
+			// Old method of rendering the water info texture; uses scene captures
+			if (WaterInfoRenderMethod == 0)
+			{
+				UE::WaterInfo::UpdateWaterInfoRendering(Scene, Context);
 			}
-			else
+			// Render the water info texture using custom render pass method
+			else if (WaterInfoRenderMethod == 2)
 			{
-				WaterZoneInfo->UpdateBounds.Reset();
+				UE::WaterInfo::UpdateWaterInfoRendering_CustomRenderPass(Scene, InViewFamily, Context);
 			}
 		}
-
-		const UE::WaterInfo::FRenderingContext& Context(Pair.Value);
-		// Old method of rendering the water info texture; uses scene captures
-		if (WaterInfoRenderMethod == 0)
-		{
-			UE::WaterInfo::UpdateWaterInfoRendering(Scene, Context);
-		}
-		// Render the water info texture using custom render pass method
-		else if (WaterInfoRenderMethod == 2)
-		{
-			UE::WaterInfo::UpdateWaterInfoRendering_CustomRenderPass(Scene, InViewFamily, Context);
-		}
 	}
 
 	// New method of rendering the water info texture; rendering is done in a separate pass when rendering the main view
diff --git a/Engine/Plugins/Experimental/Water/Source/Runtime/Private/WaterZoneActor.cpp b/Engine/Plugins/Experimental/Water/Source/Runtime/Private/WaterZoneActor.cpp
--- a/Engine/Plugins/Experimental/Water/Source/Runtime/Private/WaterZoneActor.cpp
+++ b/Engine/Plugins/Experimental/Water/Source/Runtime/Private/WaterZoneActor.cpp
@@ -482,7 +482,7 @@ bool AWaterZone::UpdateWaterInfoTexture()
 		// If they do not, we must submit compile jobs for them and wait until they are finished before re-rendering.
 		TArray<UMaterialInterface*> UsedMaterials;
 
-		TArray<UWaterBodyComponent*> WaterBodiesToRender;
+		TArray<TWeakObjectPtr<UWaterBodyComponent>> WaterBodiesToRender;
 		ForEachWaterBodyComponent([World, &WaterBodiesToRender, &WaterZMax, &WaterZMin, &UsedMaterials](UWaterBodyComponent* WaterBodyComponent)
 		{
 			// skip components which don't affect the water info texture
@@ -520,16 +520,19 @@ bool AWaterZone::UpdateWaterInfoTexture()
 		{
 			bool bHaveAllPSOsBeenCached = true;
 
-			for (UWaterBodyComponent* WaterBodyComponent : WaterBodiesToRender)
+			for (const TWeakObjectPtr<UWaterBodyComponent>& WaterBodyComponentPtr : WaterBodiesToRender)
 			{
-				// CheckPSOPrecachingAndBoostPriority returns true if PSOs are still precaching.
-				if (UWaterBodyInfoMeshComponent* WaterInfoMeshComponent = WaterBodyComponent->GetWaterInfoMeshComponent())
+				if (UWaterBodyComponent* WaterBodyComponent = WaterBodyComponentPtr.Get())
 				{
-					bHaveAllPSOsBeenCached &= !WaterInfoMeshComponent->CheckPSOPrecachingAndBoostPriority();
-				}
-				if (UWaterBodyInfoMeshComponent* WaterInfoMeshComponent = WaterBodyComponent->GetDilatedWaterInfoMeshComponent())
-				{
-					bHaveAllPSOsBeenCached &= !WaterInfoMeshComponent->CheckPSOPrecachingAndBoostPriority();
+					// CheckPSOPrecachingAndBoostPriority returns true if PSOs are still precaching.
+					if (UWaterBodyInfoMeshComponent* WaterInfoMeshComponent = WaterBodyComponent->GetWaterInfoMeshComponent())
+					{
+						bHaveAllPSOsBeenCached &= !WaterInfoMeshComponent->CheckPSOPrecachingAndBoostPriority();
+					}
+					if (UWaterBodyInfoMeshComponent* WaterInfoMeshComponent = WaterBodyComponent->GetDilatedWaterInfoMeshComponent())
+					{
+						bHaveAllPSOsBeenCached &= !WaterInfoMeshComponent->CheckPSOPrecachingAndBoostPriority();
+					}
 				}
 			}
 
diff --git a/Engine/Plugins/Experimental/Water/Source/Runtime/Public/WaterViewExtension.h b/Engine/Plugins/Experimental/Water/Source/Runtime/Public/WaterViewExtension.h
--- a/Engine/Plugins/Experimental/Water/Source/Runtime/Public/WaterViewExtension.h
+++ b/Engine/Plugins/Experimental/Water/Source/Runtime/Public/WaterViewExtension.h
@@ -36,7 +36,7 @@ public:
 
 private:
 	/** Queued Water Info rendering contexts to submit for rendering on the next SetupView call */
-	TMap<AWaterZone*, UE::WaterInfo::FRenderingContext> WaterInfoContextsToRender;
+	TWeakObjectPtrKeyMap<AWaterZone, UE::WaterInfo::FRenderingContext> WaterInfoContextsToRender;
 
 	struct FWaterZoneInfo
 	{
@@ -46,7 +46,7 @@ private:
 		 */
 		TOptional<FBox2D> UpdateBounds = FBox2D(ForceInit);
 	};
-	TMap<AWaterZone*, FWaterZoneInfo> WaterZoneInfos;
+	TWeakObjectPtrKeyMap<AWaterZone, FWaterZoneInfo> WaterZoneInfos;
 
 	struct FWaterGPUResources
 	{
