[Water]
Made cached landscape properties transient on the water brush manager. This fixes a bug where the scene capture component (which is transient) would not be initialized when the brush is initialized since the cached properties would not be different from the landscape properties and therefore the brush would not update the capture component position/ortho width/etc.

diff --git a/Engine/Plugins/Experimental/Water/Source/Editor/Private/WaterBrushManager.h b/Engine/Plugins/Experimental/Water/Source/Editor/Private/WaterBrushManager.h
--- a/Engine/Plugins/Experimental/Water/Source/Editor/Private/WaterBrushManager.h
+++ b/Engine/Plugins/Experimental/Water/Source/Editor/Private/WaterBrushManager.h
@@ -142,16 +142,16 @@ public:
 	UPROPERTY(VisibleAnywhere, BlueprintReadWrite, Transient, meta=(Category="Debug"))
 	TMap<TObjectPtr<UCurveFloat>,FWaterBodyBrushCache> BrushCurveRTCache;
 	
-	UPROPERTY(VisibleAnywhere, BlueprintReadOnly, meta=(Category="Debug"))
+	UPROPERTY(VisibleAnywhere, BlueprintReadOnly, Transient, meta=(Category="Debug"))
 	FVector WorldSize;
 	
-	UPROPERTY(VisibleAnywhere, BlueprintReadOnly, meta=(Category="Debug"))
+	UPROPERTY(VisibleAnywhere, BlueprintReadOnly, Transient, meta=(Category="Debug"))
 	FIntPoint LandscapeRTRes;
 	
-	UPROPERTY(VisibleAnywhere, BlueprintReadOnly, meta=(Category="Debug"))
+	UPROPERTY(VisibleAnywhere, BlueprintReadOnly, Transient, meta=(Category="Debug"))
 	FIntPoint LandscapeQuads;
 	
-	UPROPERTY(VisibleAnywhere, BlueprintReadOnly, meta=(Category="Debug"))
+	UPROPERTY(VisibleAnywhere, BlueprintReadOnly, Transient, meta=(Category="Debug"))
 	FTransform LandscapeTransform;
 	
 	UPROPERTY(EditAnywhere, BlueprintReadWrite, meta=(Category="Debug"))
