[Viewport Toolbar] Match the entries of the "View Modes" menu for the Static Mesh, Material and Water Editor

diff --git a/Engine/Plugins/Experimental/Water/Source/Editor/Private/WaterWavesEditorViewport.cpp b/Engine/Plugins/Experimental/Water/Source/Editor/Private/WaterWavesEditorViewport.cpp
--- a/Engine/Plugins/Experimental/Water/Source/Editor/Private/WaterWavesEditorViewport.cpp
+++ b/Engine/Plugins/Experimental/Water/Source/Editor/Private/WaterWavesEditorViewport.cpp
@@ -87,6 +87,7 @@ TSharedPtr<SWidget> SWaterWavesEditorViewport::MakeViewportToolbar()
 			}
 		);
 }
+
 TSharedPtr<SWidget> SWaterWavesEditorViewport::BuildViewportToolbar()
 {
 	// Register the viewport toolbar if another viewport hasn't already (it's shared).
@@ -194,28 +195,8 @@ TSharedPtr<SWidget> SWaterWavesEditorViewport::BuildViewportToolbar()
 
 		// Add the UnrealEd viewport toolbar context.
 		{
-			UUnrealEdViewportToolbarContext* const ContextObject = NewObject<UUnrealEdViewportToolbarContext>();
-			ContextObject->Viewport = SharedThis(this);
-
-			// Hook up our toolbar's filter for supported view modes.
-			ContextObject->IsViewModeSupported = UE::UnrealEd::IsViewModeSupportedDelegate::CreateLambda(
-				[](EViewModeIndex ViewModeIndex) -> bool
-				{
-					// This code is taken from SViewportToolBar::IsViewModeSupported
-					// SSCSEditorViewportToolBar does not override it, so we just take it as-is
-					// TODO: maybe create a private function for it, or move IsViewModeSupported to SEditorViewport
-
-					switch (ViewModeIndex)
-					{
-					case VMI_PrimitiveDistanceAccuracy:
-					case VMI_MaterialTextureScaleAccuracy:
-					case VMI_RequiredTextureResolution:
-						return false;
-					default:
-						return true;
-					}
-				}
-			);
+			UUnrealEdViewportToolbarContext* const ContextObject =
+				UE::UnrealEd::CreateViewportToolbarDefaultContext(SharedThis(this));
 
 			ViewportToolbarContext.AddObject(ContextObject);
 		}
diff --git a/Engine/Source/Editor/MaterialEditor/Private/SMaterialEditorViewport.cpp b/Engine/Source/Editor/MaterialEditor/Private/SMaterialEditorViewport.cpp
--- a/Engine/Source/Editor/MaterialEditor/Private/SMaterialEditorViewport.cpp
+++ b/Engine/Source/Editor/MaterialEditor/Private/SMaterialEditorViewport.cpp
@@ -1030,8 +1030,9 @@ TSharedPtr<SWidget> SMaterialEditor3DPreviewViewport::BuildViewportToolbar()
 
 		// Add the UnrealEd viewport toolbar context.
 		{
-			UUnrealEdViewportToolbarContext* const ContextObject = NewObject<UUnrealEdViewportToolbarContext>();
-			ContextObject->Viewport = SharedThis(this);
+			UUnrealEdViewportToolbarContext* const ContextObject =
+				UE::UnrealEd::CreateViewportToolbarDefaultContext(SharedThis(this));
+
 			ViewportToolbarContext.AddObject(ContextObject);
 		}
 	}
diff --git a/Engine/Source/Editor/StaticMeshEditor/Private/SStaticMeshEditorViewport.cpp b/Engine/Source/Editor/StaticMeshEditor/Private/SStaticMeshEditorViewport.cpp
--- a/Engine/Source/Editor/StaticMeshEditor/Private/SStaticMeshEditorViewport.cpp
+++ b/Engine/Source/Editor/StaticMeshEditor/Private/SStaticMeshEditorViewport.cpp
@@ -892,28 +892,8 @@ TSharedPtr<SWidget> SStaticMeshEditorViewport::BuildViewportToolbar()
 
 		// Add the UnrealEd viewport toolbar context.
 		{
-			UUnrealEdViewportToolbarContext* const ContextObject = NewObject<UUnrealEdViewportToolbarContext>();
-			ContextObject->Viewport = SharedThis(this);
-
-			// Hook up our toolbar's filter for supported view modes.
-			ContextObject->IsViewModeSupported = UE::UnrealEd::IsViewModeSupportedDelegate::CreateLambda(
-				[](EViewModeIndex ViewModeIndex) -> bool
-				{
-					// This code is taken from SViewportToolBar::IsViewModeSupported
-					// SSCSEditorViewportToolBar does not override it, so we just take it as-is
-					// TODO: maybe create a private function for it, or move IsViewModeSupported to SEditorViewport
-
-					switch (ViewModeIndex)
-					{
-					case VMI_PrimitiveDistanceAccuracy:
-					case VMI_MaterialTextureScaleAccuracy:
-					case VMI_RequiredTextureResolution:
-						return false;
-					default:
-						return true;
-					}
-				}
-			);
+			UUnrealEdViewportToolbarContext* const ContextObject =
+				UE::UnrealEd::CreateViewportToolbarDefaultContext(SharedThis(this));
 
 			ViewportToolbarContext.AddObject(ContextObject);
 		}
