[Water]

Early out of the UpdateWaterInfo function earlier when Terrain components are still compiling.

diff --git a/Engine/Plugins/Experimental/Water/Source/Runtime/Private/WaterZoneActor.cpp b/Engine/Plugins/Experimental/Water/Source/Runtime/Private/WaterZoneActor.cpp
--- a/Engine/Plugins/Experimental/Water/Source/Runtime/Private/WaterZoneActor.cpp
+++ b/Engine/Plugins/Experimental/Water/Source/Runtime/Private/WaterZoneActor.cpp
@@ -595,11 +595,15 @@ bool AWaterZone::UpdateWaterInfoTexture()
 		}
 
 #if WITH_EDITOR
-		// Check all the ground components have complete shader maps before we try to render them into the water info texture
+		// Check all the ground components have complete shader maps and are all compiled before we try to render them into the water info texture
 		for (TWeakObjectPtr<UPrimitiveComponent> GroundPrimCompPtr : GroundPrimitiveComponents)
 		{
 			if (UPrimitiveComponent* GroundPrimComp = GroundPrimCompPtr.Get())
 			{
+				if (GroundPrimComp->IsCompiling())
+				{
+					return false;
+				}
 				TArray<UMaterialInterface*> TmpUsedMaterials;
 				GroundPrimComp->GetUsedMaterials(TmpUsedMaterials, false);
 				UsedMaterials.Append(TmpUsedMaterials);
@@ -617,9 +621,8 @@ bool AWaterZone::UpdateWaterInfoTexture()
 				{
 					if (!MaterialResource->IsGameThreadShaderMapComplete())
 					{
-#if WITH_EDITOR
 						MaterialResource->SubmitCompileJobs_GameThread(EShaderCompileJobPriority::High);
-#endif
+
 						bHasIncompleteShaderMaps = true;
 					}
 				}
