[Navigation]
Updated multiple navigation relevant components to reduce amount of redundant operations.
Also fix some missing navigation updates for components modified after registration.

diff --git a/Engine/Plugins/Experimental/Water/Source/Runtime/Private/WaterBodyActor.cpp b/Engine/Plugins/Experimental/Water/Source/Runtime/Private/WaterBodyActor.cpp
--- a/Engine/Plugins/Experimental/Water/Source/Runtime/Private/WaterBodyActor.cpp
+++ b/Engine/Plugins/Experimental/Water/Source/Runtime/Private/WaterBodyActor.cpp
@@ -19,6 +19,7 @@
 #include "WaterBodyStaticMeshComponent.h"
 #include "WaterModule.h"
 #include "WaterVersion.h"
+#include "AI/NavigationSystemBase.h"
 #include "Algo/RemoveIf.h"
 
 #include UE_INLINE_GENERATED_CPP_BY_NAME(WaterBodyActor)
@@ -684,6 +685,9 @@ void AWaterBody::PostRegisterAllComponents()
 #endif // WITH_EDITOR
 
 		WaterBodyComponent->OnPostRegisterAllComponents();
+
+		// Now that all components are registered and setup we can notify the navigation system.
+		FNavigationSystem::UpdateActorAndComponentData(*this);
 	}
 }
 
@@ -692,6 +696,14 @@ bool AWaterBody::IsHLODRelevant() const
 	return true;
 }
 
+bool AWaterBody::IsComponentRelevantForNavigation(UActorComponent* Component) const
+{
+	// Wait for the whole setup to be completed before making child components relevant to navigation (i.e. PostRegisterAllComponents).
+	// Then we update the actor and all its components (i.e. FNavigationSystem::UpdateActorAndComponentData).
+	// This reduces the number of redundant processing in the navigation system while setting up the components.
+	return HasActorRegisteredAllComponents();
+}
+
 #if WITH_EDITOR
 void AWaterBody::PopulatePIEDuplicationSeed(AActor::FDuplicationSeedInterface& DuplicationSeed)
 {
diff --git a/Engine/Plugins/Experimental/Water/Source/Runtime/Private/WaterBodyComponent.cpp b/Engine/Plugins/Experimental/Water/Source/Runtime/Private/WaterBodyComponent.cpp
--- a/Engine/Plugins/Experimental/Water/Source/Runtime/Private/WaterBodyComponent.cpp
+++ b/Engine/Plugins/Experimental/Water/Source/Runtime/Private/WaterBodyComponent.cpp
@@ -1508,7 +1508,11 @@ void UWaterBodyComponent::UpdateAll(const FOnWaterBodyChangedParams& InParams)
 			}
 
 			MarkOwningWaterZoneForRebuild(RebuildFlags);
+		}
 
+		// 'UpdateWaterBody' creates/updates collision components so we need to update the navigation data
+		if (IsRegistered() && IsNavigationRelevant())
+		{
 			FNavigationSystem::UpdateActorAndComponentData(*WaterBodyOwner);
 		}
 
diff --git a/Engine/Plugins/Experimental/Water/Source/Runtime/Private/WaterBodyOceanComponent.cpp b/Engine/Plugins/Experimental/Water/Source/Runtime/Private/WaterBodyOceanComponent.cpp
--- a/Engine/Plugins/Experimental/Water/Source/Runtime/Private/WaterBodyOceanComponent.cpp
+++ b/Engine/Plugins/Experimental/Water/Source/Runtime/Private/WaterBodyOceanComponent.cpp
@@ -622,10 +622,6 @@ void UWaterBodyOceanComponent::OnUpdateBody(bool bWithExclusionVolumes)
 				CollisionBoxes.Add(BoxComponent);
 			}
 
-			if (!BoxComponent->IsRegistered())
-			{
-				BoxComponent->RegisterComponent();
-			}
 			BoxComponent->SetNetAddressable(); // it's deterministically named so it's addressable over network (needed for collision)
 			BoxComponent->bDrawOnlyIfSelected = true;
 			BoxComponent->SetRelativeLocation(FVector::ZeroVector);
@@ -637,6 +633,11 @@ void UWaterBodyOceanComponent::OnUpdateBody(bool bWithExclusionVolumes)
 			BoxComponent->SetRelativeLocation(RelativePosition);
 			BoxComponent->SetBoxExtent(Box.BoxExtent);
 
+			// Registers as the last step to reduce amount of updates required when modifying an already registered component
+			if (!BoxComponent->IsRegistered())
+			{
+				BoxComponent->RegisterComponent();
+			}
 		}
 
 		// create the convex-hull components
@@ -659,10 +660,6 @@ void UWaterBodyOceanComponent::OnUpdateBody(bool bWithExclusionVolumes)
 				CollisionHullSets.Add(CollisionComponent);
 			}
 
-			if (!CollisionComponent->IsRegistered())
-			{
-				CollisionComponent->RegisterComponent();
-			}
 			CollisionComponent->SetNetAddressable(); // it's deterministically named so it's addressable over network (needed for collision)
 			CollisionComponent->SetRelativeLocation(FVector::ZeroVector);
 
@@ -670,6 +667,12 @@ void UWaterBodyOceanComponent::OnUpdateBody(bool bWithExclusionVolumes)
 			CopySharedNavigationSettingsToComponent(CollisionComponent);
 
 			CollisionComponent->InitializeFromConvexElements(ConvexSet);
+
+			// Registers as the last step to reduce amount of updates required when modifying an already registered component
+			if (!CollisionComponent->IsRegistered())
+			{
+				CollisionComponent->RegisterComponent();
+			}
 		}
 	}
 	else
diff --git a/Engine/Plugins/Experimental/Water/Source/Runtime/Public/WaterBodyActor.h b/Engine/Plugins/Experimental/Water/Source/Runtime/Public/WaterBodyActor.h
--- a/Engine/Plugins/Experimental/Water/Source/Runtime/Public/WaterBodyActor.h
+++ b/Engine/Plugins/Experimental/Water/Source/Runtime/Public/WaterBodyActor.h
@@ -73,6 +73,7 @@ public:
 	virtual void PostLoad() override;
 	virtual void PostRegisterAllComponents() override;
 	virtual bool IsHLODRelevant() const override;
+	virtual bool IsComponentRelevantForNavigation(UActorComponent* Component) const override;
 
 #if WITH_EDITOR
 	virtual void SetActorHiddenInGame(bool bNewHidden) override;
diff --git a/Engine/Source/Runtime/Engine/Classes/Components/StaticMeshComponent.h b/Engine/Source/Runtime/Engine/Classes/Components/StaticMeshComponent.h
--- a/Engine/Source/Runtime/Engine/Classes/Components/StaticMeshComponent.h
+++ b/Engine/Source/Runtime/Engine/Classes/Components/StaticMeshComponent.h
@@ -289,7 +289,15 @@ public:
 	/** For Nanite enabled meshes, we'll only show the proxy mesh if this is true */
 	UPROPERTY()
 	uint8 bDisplayNaniteFallbackMesh:1;
-#endif
+
+	/**
+	 * Transient flag used during registration to handle edge case with mesh compilation completion callback.
+	 * We perform actions to register the mesh properly when it gets called async, but we end up doing those
+	 * twice when it gets called while the registration is not completed.
+	 */
+	UPROPERTY(transient)
+	uint8 bRegistering : 1;
+#endif // WITH_EDITORONLY_DATA
 
 	/** Enable dynamic sort mesh's triangles to remove ordering issue when rendered with a translucent material */
 	UPROPERTY(EditAnywhere, AdvancedDisplay, BlueprintReadOnly, Category = Lighting, meta = (UIMin = "0", UIMax = "1", DisplayName = "Sort Triangles"))
diff --git a/Engine/Source/Runtime/Engine/Private/Components/ActorComponent.cpp b/Engine/Source/Runtime/Engine/Private/Components/ActorComponent.cpp
--- a/Engine/Source/Runtime/Engine/Private/Components/ActorComponent.cpp
+++ b/Engine/Source/Runtime/Engine/Private/Components/ActorComponent.cpp
@@ -2555,11 +2555,18 @@ void UActorComponent::HandleCanEverAffectNavigationChange(bool bForceUpdate)
 	{
 		if (bCanEverAffectNavigation)
 		{
+			// Update cached value
 			bNavigationRelevant = IsNavigationRelevant();
+
+			// Notify the navigation system
 			FNavigationSystem::OnComponentRegistered(*this);
 		}
 		else
 		{
+			// Update cached value
+			bNavigationRelevant = false;
+
+			// Notify the navigation system
 			FNavigationSystem::OnComponentUnregistered(*this);
 		}
 	}
diff --git a/Engine/Source/Runtime/Engine/Private/Components/SceneComponent.cpp b/Engine/Source/Runtime/Engine/Private/Components/SceneComponent.cpp
--- a/Engine/Source/Runtime/Engine/Private/Components/SceneComponent.cpp
+++ b/Engine/Source/Runtime/Engine/Private/Components/SceneComponent.cpp
@@ -1422,6 +1422,12 @@ void USceneComponent::UpdateBounds()
 		}
 	}
 
+	// Refresh navigation data if component is relevant to navigation.
+	// Note that 'UpdateNavigationData' will ignore unregistered components.
+	if (bNavigationRelevant)
+	{
+		UpdateNavigationData();
+	}
 
 #if ENABLE_NAN_DIAGNOSTIC
 	if (Bounds.ContainsNaN())
@@ -2326,10 +2332,13 @@ bool USceneComponent::AttachToComponent(USceneComponent* Parent, const FAttachme
 			}
 		}
 
-		// Update overlaps, in case location changed or overlap state depends on attachment.
 		if (IsRegistered())
 		{
+			// Update overlaps, in case location changed or overlap state depends on attachment.
 			UpdateOverlaps();
+
+			// Update our owner actor in the navigation system since its associated bounds have changed
+			FNavigationSystem::UpdateActorAndComponentData(*GetOwner());
 		}
 
 		return true;
diff --git a/Engine/Source/Runtime/Engine/Private/Components/StaticMeshComponent.cpp b/Engine/Source/Runtime/Engine/Private/Components/StaticMeshComponent.cpp
--- a/Engine/Source/Runtime/Engine/Private/Components/StaticMeshComponent.cpp
+++ b/Engine/Source/Runtime/Engine/Private/Components/StaticMeshComponent.cpp
@@ -34,7 +34,8 @@
 #include "HierarchicalLODUtilitiesModule.h"
 #include "Rendering/StaticLightingSystemInterface.h"
 #include "Streaming/ActorTextureStreamingBuildDataComponent.h"
-#endif
+#include "Templates/UnrealTemplate.h"
+#endif // WITH_EDITOR
 #include "LightMap.h"
 #include "ShadowMap.h"
 #include "AI/Navigation/NavCollisionBase.h"
@@ -226,6 +227,7 @@ UStaticMeshComponent::UStaticMeshComponent(const FObjectInitializer& ObjectIniti
 	bDisplayVertexColors = false;
 	bDisplayPhysicalMaterialMasks = false;
 	bDisplayNaniteFallbackMesh = false;
+	bRegistering = false;
 #endif
 }
 
@@ -731,6 +733,10 @@ void UStaticMeshComponent::PostEditImport()
 
 void UStaticMeshComponent::OnRegister()
 {
+#if WITH_EDITORONLY_DATA
+	FGuardValue_Bitfield(bRegistering, true);
+#endif
+
 	NotifyIfStaticMeshChanged();
 
 	UpdateCollisionFromStaticMesh();
@@ -3212,6 +3218,13 @@ void UStaticMeshComponent::UpdateBounds()
 
 void UStaticMeshComponent::PostStaticMeshCompilation()
 {
+	// Flag indicates that the component is currently registering and all the following actions
+	// are going to be handled by the registration itself. No need to perform them in that case.
+	if (bRegistering)
+	{
+		return;
+	}
+
 	CachePaintedDataIfNecessary();
 
 	FixupOverrideColorsIfNecessary(true);
@@ -3222,10 +3235,10 @@ void UStaticMeshComponent::PostStaticMeshCompilation()
 
 	RecreatePhysicsState();
 
-	FNavigationSystem::UpdateComponentData(*this);
-
 	if (IsRegistered())
 	{
+		FNavigationSystem::UpdateComponentData(*this);
+
 		FStaticLightingSystemInterface::OnPrimitiveComponentUnregistered.Broadcast(this);
 		if (HasValidSettingsForStaticLighting(false))
 		{
diff --git a/Engine/Source/Runtime/NavigationSystem/Private/NavLinkCustomComponent.cpp b/Engine/Source/Runtime/NavigationSystem/Private/NavLinkCustomComponent.cpp
--- a/Engine/Source/Runtime/NavigationSystem/Private/NavLinkCustomComponent.cpp
+++ b/Engine/Source/Runtime/NavigationSystem/Private/NavLinkCustomComponent.cpp
@@ -314,7 +314,10 @@ void UNavLinkCustomComponent::SetLinkData(const FVector& RelativeStart, const FV
 	LinkRelativeStart = RelativeStart;
 	LinkRelativeEnd = RelativeEnd;
 	LinkDirection = Direction;
-	
+
+	// Link start and end positions have changed, we need to update the bounds
+	UpdateNavigationBounds();
+
 	RefreshNavigationModifiers();
 }
 
diff --git a/Engine/Source/Runtime/NavigationSystem/Private/NavModifierVolume.cpp b/Engine/Source/Runtime/NavigationSystem/Private/NavModifierVolume.cpp
--- a/Engine/Source/Runtime/NavigationSystem/Private/NavModifierVolume.cpp
+++ b/Engine/Source/Runtime/NavigationSystem/Private/NavModifierVolume.cpp
@@ -18,12 +18,15 @@ namespace UE::Navigation::ModVolume::Private
 {
 	void OnNavAreaRegistrationChanged(ANavModifierVolume& ModifierVolume, const UWorld& World, const UClass* NavAreaClass)
 	{
-		if (NavAreaClass && NavAreaClass == ModifierVolume.GetAreaClass() && &World == ModifierVolume.GetWorld())
+		if (NavAreaClass
+			&& NavAreaClass == ModifierVolume.GetAreaClass()
+			&& &World == ModifierVolume.GetWorld()
+			&& ModifierVolume.HasActorRegisteredAllComponents()) // Update only required after initial registration was completed
 		{
 			FNavigationSystem::UpdateActorData(ModifierVolume);
 		}
 	}
-} // UE::Navigation::ModVolumne::Private
+} // UE::Navigation::ModVolume::Private
 #endif // WITH_EDITOR
 
 //----------------------------------------------------------------------//
