- Track changes to the water body component
- This fixes an issue where it was found on the first frame, then stomped the next frame with nullptr

#rnx
#tests Tested auto finding component vs having the component set directly

diff --git a/Engine/Plugins/Experimental/Water/Source/Runtime/Private/NiagaraDataInterfaceWater.cpp b/Engine/Plugins/Experimental/Water/Source/Runtime/Private/NiagaraDataInterfaceWater.cpp
--- a/Engine/Plugins/Experimental/Water/Source/Runtime/Private/NiagaraDataInterfaceWater.cpp
+++ b/Engine/Plugins/Experimental/Water/Source/Runtime/Private/NiagaraDataInterfaceWater.cpp
@@ -17,6 +17,7 @@ struct FNDIWater_InstanceData
 {
 	bool								bFindClosestBody = false;
 	TWeakObjectPtr<UWaterBodyComponent>	WaterBodyComponent;
+	uint32								WaterBodyChangeId = 0;
 	FNiagaraLWCConverter				LWCConverter;
 };
 
@@ -299,8 +300,9 @@ int32 UNiagaraDataInterfaceWater::PerInstanceDataSize() const
 bool UNiagaraDataInterfaceWater::InitPerInstanceData(void* PerInstanceData, FNiagaraSystemInstance* SystemInstance)
 {
 	FNDIWater_InstanceData* InstData = new (PerInstanceData) FNDIWater_InstanceData();
-	InstData->bFindClosestBody		= SourceBodyComponent == nullptr && bFindWaterBodyOnSpawn;
-	InstData->WaterBodyComponent	= SourceBodyComponent;
+	InstData->bFindClosestBody		= bFindWaterBodyOnSpawn;
+	InstData->WaterBodyComponent	= nullptr;
+	InstData->WaterBodyChangeId		= SourceBodyChangeId - 1;
 	InstData->LWCConverter			= SystemInstance->GetLWCConverter();
 
 	return true;
@@ -317,18 +319,23 @@ bool UNiagaraDataInterfaceWater::PerInstanceTick(void* PerInstanceData, FNiagara
 	check(SystemInstance);
 	FNDIWater_InstanceData* InstData = static_cast<FNDIWater_InstanceData*>(PerInstanceData);
 
-	// If the search for closest was enabled, perform the search
-	// Note: we do this here rather than in Init as the system might be auto activate and the user parameter not set until after the spawn
-	if (SourceBodyComponent == nullptr && InstData->bFindClosestBody)
+	// Do we need to update the water body component?
+	if (InstData->WaterBodyChangeId != SourceBodyChangeId)
 	{
-		const FVector QueryLocation = SystemInstance->GetWorldTransform().GetTranslation();
-		InstData->WaterBodyComponent = NDIWaterPrivate::FindClosestWaterBody(SystemInstance->GetWorld(), QueryLocation);
-	}
-	else
-	{
-		InstData->WaterBodyComponent = SourceBodyComponent;
+		// If the search for closest was enabled, perform the search
+		// Note: we do this here rather than in Init as the system might be auto activate and the user parameter not set until after the spawn
+		if (SourceBodyComponent == nullptr && InstData->bFindClosestBody)
+		{
+			const FVector QueryLocation = SystemInstance->GetWorldTransform().GetTranslation();
+			InstData->WaterBodyComponent = NDIWaterPrivate::FindClosestWaterBody(SystemInstance->GetWorld(), QueryLocation);
+		}
+		else
+		{
+			InstData->WaterBodyComponent = SourceBodyComponent;
+		}
+		InstData->bFindClosestBody = false;
+		InstData->WaterBodyChangeId = SourceBodyChangeId;
 	}
-	InstData->bFindClosestBody = false;
 
 	return false;
 }
diff --git a/Engine/Plugins/Experimental/Water/Source/Runtime/Public/NiagaraDataInterfaceWater.h b/Engine/Plugins/Experimental/Water/Source/Runtime/Public/NiagaraDataInterfaceWater.h
--- a/Engine/Plugins/Experimental/Water/Source/Runtime/Public/NiagaraDataInterfaceWater.h
+++ b/Engine/Plugins/Experimental/Water/Source/Runtime/Public/NiagaraDataInterfaceWater.h
@@ -39,7 +39,11 @@ public:
 	void GetWaveParamLookupTableOffset(FVectorVMExternalFunctionContext& Context);
 
 	/** Sets the current water body to be used by this data interface */
-	void SetWaterBodyComponent(UWaterBodyComponent* InWaterBodyComponent) { SourceBodyComponent = InWaterBodyComponent; }
+	void SetWaterBodyComponent(UWaterBodyComponent* InWaterBodyComponent)
+	{
+		SourceBodyComponent = InWaterBodyComponent;
+		++SourceBodyChangeId;
+	}
 
 protected:
 #if WITH_EDITORONLY_DATA
@@ -52,6 +56,8 @@ private:
 
 	UPROPERTY(EditAnywhere, Category = "Water") 
 	TObjectPtr<UWaterBodyComponent> SourceBodyComponent;
+
+	uint32 SourceBodyChangeId = 0;
 };
 
 #if UE_ENABLE_INCLUDE_ORDER_DEPRECATED_IN_5_2
