[Viewport Toolbar] Add Show and Preview Profiles Submenu to Water Editor

Note: this CL also removes redundant Performance and Scalability menu (it was being added twice)

diff --git a/Engine/Plugins/Experimental/Water/Source/Editor/Private/WaterWavesEditorViewport.cpp b/Engine/Plugins/Experimental/Water/Source/Editor/Private/WaterWavesEditorViewport.cpp
--- a/Engine/Plugins/Experimental/Water/Source/Editor/Private/WaterWavesEditorViewport.cpp
+++ b/Engine/Plugins/Experimental/Water/Source/Editor/Private/WaterWavesEditorViewport.cpp
@@ -157,18 +157,13 @@ TSharedPtr<SWidget> SWaterWavesEditorViewport::BuildViewportToolbar()
 				FToolMenuEntry ViewModesSubmenu = UE::UnrealEd::CreateViewportToolbarViewModesSubmenu();
 				ViewModesSubmenu.InsertPosition.Position = EToolMenuInsertType::First;
 				RightSection.AddEntry(ViewModesSubmenu);
-
-				FToolMenuEntry PerformanceAndScalabilitySubmenu = UE::UnrealEd::CreatePerformanceAndScalabilitySubmenu();
-				PerformanceAndScalabilitySubmenu.InsertPosition.Position = EToolMenuInsertType::First;
-				RightSection.AddEntry(PerformanceAndScalabilitySubmenu);
 			}
 
 			// Add the Show submenu.
 			{
-				// TODO: needs some shared function e.g. in UnrealEdToolbar.cpp since this menu is quite common
-				/*FToolMenuEntry ShowSubmenu = ;
+				FToolMenuEntry ShowSubmenu = UE::UnrealEd::CreateDefaultShowSubmenu();
 				ShowSubmenu.InsertPosition.Position = EToolMenuInsertType::First;
-				RightSection.AddEntry(ShowSubmenu);*/
+				RightSection.AddEntry(ShowSubmenu);
 			}
 
 			// Add the Performance and Scalability submenu.
