A single WaterVertexFactory instance can now render multiple similar quadtrees at different world space locations. This is a necessary step towards supporting per-view quadtrees for split screen.
Also fixed selection/outlining not working with the GPU driven water quadtree path.

diff --git a/Engine/Plugins/Experimental/Water/Shaders/Private/WaterMeshVertexFactory.ush b/Engine/Plugins/Experimental/Water/Shaders/Private/WaterMeshVertexFactory.ush
--- a/Engine/Plugins/Experimental/Water/Shaders/Private/WaterMeshVertexFactory.ush
+++ b/Engine/Plugins/Experimental/Water/Shaders/Private/WaterMeshVertexFactory.ush
@@ -10,6 +10,7 @@
 
 #if WATER_MESH_DRAW_INDIRECT_INTERNAL
 float3 QuadTreePosition;
+float CaptureDepthRange;
 uint bLODMorphingEnabled;
 
 #if INSTANCED_STEREO
@@ -303,7 +304,7 @@ FWaterVertexFactoryInstanceInput UnpackWaterVertexFactoryInstanceInput(float4 In
 	const uint RawDensityIndex = (InData0 >> 27u) & 0x1Fu;
 	const uint LODLevel = RawLODLevel + RawDensityIndex;
 	const uint DensityIndex = min(RawDensityIndex, (uint)WaterVF.NumDensities - 1);
-	const float WaterHeight = f16tof32(InData1) * WaterVF.CaptureDepthRange;
+	const float WaterHeight = f16tof32(InData1) * CaptureDepthRange;
 	const uint2 TileCoord = uint2((InData1 >> 16u) & 0xFFu, (InData1 >> 24u) & 0xFFu);
 	const float HeightLODFactor = float(InData2 & 0xFFu) / 255.0f;
 	const uint WaterBodyIndex = InData2 >> 8u;
@@ -430,7 +431,7 @@ FVertexFactoryIntermediates GetVertexFactoryIntermediates(FVertexFactoryInput In
 	#if INSTANCED_STEREO
 	float SelectedValue = UnpackRGBA8(InstanceData3Buffer.Load(InstanceDataIndex)).w;
 	#else
-	float SelectedValue = UnpackRGBA8(Input.InstanceData2).w;
+	float SelectedValue = UnpackRGBA8(Input.InstanceData3).w;
 	#endif // INSTANCED_STEREO
 #else // WATER_MESH_DRAW_INDIRECT_INTERNAL
 	float SelectedValue = Input.InstanceData2.w;
@@ -452,7 +453,7 @@ float4 VertexFactoryGetInstanceHitProxyId(FVertexFactoryInput Input, FVertexFact
 	const uint InstanceDataIndex = InstanceDataOffsetsBuffer.Load(DrawBucketIndex) + InstanceIndex;
 	return float4(UnpackRGBA8(InstanceData3Buffer.Load(InstanceDataIndex)).bgr, 0);
 	#else
-	return float4(UnpackRGBA8(Input.InstanceData2).bgr, 0);
+	return float4(UnpackRGBA8(Input.InstanceData3).bgr, 0);
 	#endif // INSTANCED_STEREO
 #else // WATER_MESH_DRAW_INDIRECT_INTERNAL
 	return float4(Input.InstanceData2.rgb, 0);
diff --git a/Engine/Plugins/Experimental/Water/Source/Runtime/Private/WaterMeshSceneProxy.cpp b/Engine/Plugins/Experimental/Water/Source/Runtime/Private/WaterMeshSceneProxy.cpp
--- a/Engine/Plugins/Experimental/Water/Source/Runtime/Private/WaterMeshSceneProxy.cpp
+++ b/Engine/Plugins/Experimental/Water/Source/Runtime/Private/WaterMeshSceneProxy.cpp
@@ -224,8 +224,6 @@ FWaterMeshSceneProxy::FWaterMeshSceneProxy(UWaterMeshComponent* Component)
 
 	DensityCount = FMath::Min(WaterQuadTree.GetTreeDepth(), (int32)FMath::FloorLog2(NumQuadsLOD0));
 
-	const FVector QuadTreeCorner = FVector(WaterQuadTree.GetTileRegion().Min, WaterQuadTreeMinHeight);
-	const float WaterQuadTreeDepthRange = WaterQuadTreeMaxHeight - WaterQuadTreeMinHeight;
 	const float LeafSize = WaterQuadTree.GetLeafSize();
 	
 	// bIsGPUQuadTree is constant over the lifetime of this scene proxy, so we know up front if we need the indirect draw vertex factory
@@ -235,12 +233,12 @@ FWaterMeshSceneProxy::FWaterMeshSceneProxy(UWaterMeshComponent* Component)
 		const UE::StereoRenderUtils::FStereoShaderAspects Aspects(GetScene().GetShaderPlatform());
 		if (Aspects.IsInstancedStereoEnabled())
 		{
-			WaterVertexFactoryIndirectDrawISR = new FWaterVertexFactoryIndirectDrawISRType(GetScene().GetFeatureLevel(), QuadTreeCorner, NumQuadsPerIndirectDrawTile, NumQuadsLOD0, DensityCount, LeafSize, LODScale, WaterQuadTreeDepthRange);
+			WaterVertexFactoryIndirectDrawISR = new FWaterVertexFactoryIndirectDrawISRType(GetScene().GetFeatureLevel(), NumQuadsPerIndirectDrawTile, NumQuadsLOD0, DensityCount, LeafSize, LODScale);
 			BeginInitResource(WaterVertexFactoryIndirectDrawISR);
 		}
 		else
 		{
-			WaterVertexFactoryIndirectDraw = new FWaterVertexFactoryIndirectDrawType(GetScene().GetFeatureLevel(), QuadTreeCorner, NumQuadsPerIndirectDrawTile, NumQuadsLOD0, DensityCount, LeafSize, LODScale, WaterQuadTreeDepthRange);
+			WaterVertexFactoryIndirectDraw = new FWaterVertexFactoryIndirectDrawType(GetScene().GetFeatureLevel(), NumQuadsPerIndirectDrawTile, NumQuadsLOD0, DensityCount, LeafSize, LODScale);
 			BeginInitResource(WaterVertexFactoryIndirectDraw);
 		}
 	}
@@ -249,7 +247,7 @@ FWaterMeshSceneProxy::FWaterMeshSceneProxy(UWaterMeshComponent* Component)
 	WaterVertexFactories.Reserve(WaterQuadTree.GetTreeDepth());
 	for (uint8 i = 0; i < WaterQuadTree.GetTreeDepth(); i++)
 	{
-		WaterVertexFactories.Add(new FWaterVertexFactoryType(GetScene().GetFeatureLevel(), QuadTreeCorner, NumQuads, NumQuadsLOD0, DensityCount, LeafSize, LODScale, WaterQuadTreeDepthRange));
+		WaterVertexFactories.Add(new FWaterVertexFactoryType(GetScene().GetFeatureLevel(), NumQuads, NumQuadsLOD0, DensityCount, LeafSize, LODScale));
 		BeginInitResource(WaterVertexFactories.Last());
 
 		NumQuads /= 2;
@@ -535,6 +533,8 @@ void FWaterMeshSceneProxy::GetDynamicMeshElements(const TArray<const FSceneView*
 		{
 			FWaterVertexFactoryUserDataWrapperType& UserDataWrapper = Collector.AllocateOneFrameResource<FWaterVertexFactoryUserDataWrapperType>();
 			UserDataWrapper.UserData.RenderGroupType = RenderGroup;
+			UserDataWrapper.UserData.QuadTreePosition = FVector(WaterQuadTree.GetTileRegion().Min, WaterQuadTreeMinHeight);
+			UserDataWrapper.UserData.CaptureDepthRange = WaterQuadTreeMaxHeight - WaterQuadTreeMinHeight;
 			UserDataWrapper.UserData.IndirectInstanceData0 = IndirectDrawResources.InstanceData0->GetRHI();
 			UserDataWrapper.UserData.IndirectInstanceData1 = IndirectDrawResources.InstanceData1->GetRHI();
 			UserDataWrapper.UserData.IndirectInstanceData2 = IndirectDrawResources.InstanceData2->GetRHI();
diff --git a/Engine/Plugins/Experimental/Water/Source/Runtime/Private/WaterVertexFactory.cpp b/Engine/Plugins/Experimental/Water/Source/Runtime/Private/WaterVertexFactory.cpp
--- a/Engine/Plugins/Experimental/Water/Source/Runtime/Private/WaterVertexFactory.cpp
+++ b/Engine/Plugins/Experimental/Water/Source/Runtime/Private/WaterVertexFactory.cpp
@@ -29,6 +29,7 @@ public:
 		if (DrawMode == EWaterVertexFactoryDrawMode::Indirect || DrawMode == EWaterVertexFactoryDrawMode::IndirectInstancedStereo)
 		{
 			QuadTreePositionParameter.Bind(ParameterMap, TEXT("QuadTreePosition"));
+			CaptureDepthRangeParameter.Bind(ParameterMap, TEXT("CaptureDepthRange"));
 			LODMorphingEnabledParameter.Bind(ParameterMap, TEXT("bLODMorphingEnabled"));
 
 			if (DrawMode == EWaterVertexFactoryDrawMode::IndirectInstancedStereo)
@@ -115,7 +116,8 @@ public:
 			}
 
 			const FVector PreViewTranslation = View->ViewMatrices.GetPreViewTranslation();
-			ShaderBindings.Add(QuadTreePositionParameter, FVector3f(PreViewTranslation + VertexFactory->GetQuadTreePositionWS()));
+			ShaderBindings.Add(QuadTreePositionParameter, FVector3f(PreViewTranslation + WaterMeshUserData->QuadTreePosition));
+			ShaderBindings.Add(CaptureDepthRangeParameter, WaterMeshUserData->CaptureDepthRange);
 
 			static const auto CVar = IConsoleManager::Get().FindTConsoleVariableDataInt(TEXT("r.Water.WaterMesh.LODMorphEnabled"));
 			const bool bLODMorphingEnabled = CVar && CVar->GetValueOnRenderThread() != 0;
@@ -143,6 +145,7 @@ public:
 
 private:
 	LAYOUT_FIELD(FShaderParameter, QuadTreePositionParameter);
+	LAYOUT_FIELD(FShaderParameter, CaptureDepthRangeParameter);
 	LAYOUT_FIELD(FShaderParameter, LODMorphingEnabledParameter);
 	LAYOUT_FIELD(FShaderResourceParameter, InstanceDataOffsetsBufferParameter);
 	LAYOUT_FIELD(FShaderResourceParameter, InstanceData0BufferParameter);
diff --git a/Engine/Plugins/Experimental/Water/Source/Runtime/Public/WaterVertexFactory.h b/Engine/Plugins/Experimental/Water/Source/Runtime/Public/WaterVertexFactory.h
--- a/Engine/Plugins/Experimental/Water/Source/Runtime/Public/WaterVertexFactory.h
+++ b/Engine/Plugins/Experimental/Water/Source/Runtime/Public/WaterVertexFactory.h
@@ -16,7 +16,6 @@ struct FShaderCompilerEnvironment;
 BEGIN_GLOBAL_SHADER_PARAMETER_STRUCT(FWaterVertexFactoryParameters, )
 	SHADER_PARAMETER(float, LODScale)
 	SHADER_PARAMETER(float, LeafSize)
-	SHADER_PARAMETER(float, CaptureDepthRange)
 	SHADER_PARAMETER(int32, NumQuadsPerTileSide)
 	SHADER_PARAMETER(int32, NumQuadsLOD0)
 	SHADER_PARAMETER(int32, NumDensities)
@@ -194,7 +193,7 @@ public:
 	static constexpr int32 NumRenderGroups = bWithWaterSelectionSupport ? 3 : 1; // Must match EWaterMeshRenderGroupType
 	static constexpr int32 NumAdditionalVertexStreams = TWaterInstanceDataBuffers<bWithWaterSelectionSupport>::NumBuffers;
 
-	TWaterVertexFactory(ERHIFeatureLevel::Type InFeatureLevel, const FVector& InQuadTreePositionWS, int32 InNumQuadsPerSide, int32 InNumQuadsLOD0, int32 InNumDensities, float InLeafSize, float InLODScale, float InCaptureDepthRange);
+	TWaterVertexFactory(ERHIFeatureLevel::Type InFeatureLevel, int32 InNumQuadsPerSide, int32 InNumQuadsLOD0, int32 InNumDensities, float InLeafSize, float InLODScale);
 	~TWaterVertexFactory();
 
 	/**
@@ -220,8 +219,6 @@ public:
 
 	inline const FUniformBufferRHIRef GetWaterVertexFactoryUniformBuffer(EWaterMeshRenderGroupType InRenderGroupType) const { return UniformBuffers[(int32)InRenderGroupType]; }
 
-	inline FVector GetQuadTreePositionWS() const { return QuadTreePositionWS; }
-
 	static constexpr bool UsesIndirectDraws() { return DrawMode == EWaterVertexFactoryDrawMode::Indirect || DrawMode == EWaterVertexFactoryDrawMode::IndirectInstancedStereo; }
 	static constexpr bool UsesInstancedStereo() { return DrawMode == EWaterVertexFactoryDrawMode::IndirectInstancedStereo; }
 
@@ -235,13 +232,11 @@ public:
 private:
 	TStaticArray<FWaterVertexFactoryBufferRef, NumRenderGroups> UniformBuffers;
 
-	const FVector QuadTreePositionWS = FVector::ZeroVector;
 	const int32 NumQuadsPerSide = 0;
 	const int32 NumQuadsLOD0 = 0;
 	const int32 NumDensities = 0;
 	const float LeafSize = 0.0f;
 	const float LODScale = 0.0f;
-	const float CaptureDepthRange = 0.0f;
 };
 
 extern const FVertexFactoryType* GetWaterVertexFactoryType(bool bWithWaterSelectionSupport, EWaterVertexFactoryDrawMode DrawMode);
@@ -268,6 +263,8 @@ struct TWaterMeshUserData
 	FUniformBufferRHIRef WaterVertexFactoryRaytracingVFUniformBuffer = nullptr;
 #endif
 
+	FVector QuadTreePosition = FVector::ZeroVector;
+	float CaptureDepthRange = 0.0f;
 	FRHIBuffer* IndirectInstanceData0 = nullptr;
 	FRHIBuffer* IndirectInstanceData1 = nullptr;
 	FRHIBuffer* IndirectInstanceData2 = nullptr;
diff --git a/Engine/Plugins/Experimental/Water/Source/Runtime/Public/WaterVertexFactory.inl b/Engine/Plugins/Experimental/Water/Source/Runtime/Public/WaterVertexFactory.inl
--- a/Engine/Plugins/Experimental/Water/Source/Runtime/Public/WaterVertexFactory.inl
+++ b/Engine/Plugins/Experimental/Water/Source/Runtime/Public/WaterVertexFactory.inl
@@ -15,15 +15,13 @@
 // ----------------------------------------------------------------------------------
 
 template <bool bWithWaterSelectionSupport, EWaterVertexFactoryDrawMode DrawMode>
-TWaterVertexFactory<bWithWaterSelectionSupport, DrawMode>::TWaterVertexFactory(ERHIFeatureLevel::Type InFeatureLevel, const FVector& InQuadTreePositionWS, int32 InNumQuadsPerSide, int32 InNumQuadsLOD0, int32 InNumDensities, float InLeafSize, float InLODScale, float InCaptureDepthRange)
+TWaterVertexFactory<bWithWaterSelectionSupport, DrawMode>::TWaterVertexFactory(ERHIFeatureLevel::Type InFeatureLevel, int32 InNumQuadsPerSide, int32 InNumQuadsLOD0, int32 InNumDensities, float InLeafSize, float InLODScale)
 	: FVertexFactory(InFeatureLevel)
-	, QuadTreePositionWS(InQuadTreePositionWS)
 	, NumQuadsPerSide(InNumQuadsPerSide)
 	, NumQuadsLOD0(InNumQuadsLOD0)
 	, NumDensities(InNumDensities)
 	, LeafSize(InLeafSize)
 	, LODScale(InLODScale)
-	, CaptureDepthRange(InCaptureDepthRange)
 {
 	VertexBuffer = new FWaterMeshVertexBuffer(NumQuadsPerSide);
 	IndexBuffer = new FWaterMeshIndexBuffer(NumQuadsPerSide);
@@ -151,7 +149,6 @@ void TWaterVertexFactory<bWithWaterSelectionSupport, DrawMode>::SetupUniformData
 	UniformParams.NumDensities = NumDensities;
 	UniformParams.LODScale = LODScale;
 	UniformParams.LeafSize = LeafSize;
-	UniformParams.CaptureDepthRange = CaptureDepthRange;
 	UniformParams.bRenderSelected = true;
 	UniformParams.bRenderUnselected = true;
 
