- Disable clustering on the water data interface
- Blueprint / C++ calls can modify a property directly

#rnx
#tests Ran test asset in level

diff --git a/Engine/Plugins/Experimental/Water/Source/Runtime/Public/NiagaraDataInterfaceWater.h b/Engine/Plugins/Experimental/Water/Source/Runtime/Public/NiagaraDataInterfaceWater.h
--- a/Engine/Plugins/Experimental/Water/Source/Runtime/Public/NiagaraDataInterfaceWater.h
+++ b/Engine/Plugins/Experimental/Water/Source/Runtime/Public/NiagaraDataInterfaceWater.h
@@ -13,6 +13,7 @@ class WATER_API UNiagaraDataInterfaceWater : public UNiagaraDataInterface
 
 public:
 	virtual void PostInitProperties() override;
+	virtual bool CanBeInCluster() const override { return false; }	// Note: Due to BP functionality we can change a UObject property on this DI we can not put into a cluster
 
 	/** UNiagaraDataInterface interface */
 	virtual void GetVMExternalFunction(const FVMExternalFunctionBindingInfo& BindingInfo, void* InstanceData, FVMExternalFunction &OutFunc) override;
