Water: Rewrote and reorganized large parts of FWaterMeshSceneProxy to allow for supporting multiple local quadtrees with a single scene proxy instance. This is useful for supporting splitscreen. The scene proxy now has a TMap of int keys and local quadtrees and uses a new FindBestQuadTreeForView() function to map from any given view to a quadtree. Currently it uses the views PlayerIndex to do so and falls back to picking the closest quadtree otherwise. The scene proxy now also exposes functions for creating, updating and destroying these local quadtrees.
Since the FWaterViewExtension does not yet support multiple quadtrees, there's a new FORCE_SINGLE_WATER_QUADTREE macro to force the scene proxy to only ever have a single quadtree, effectively behaving the same as before this CL.
Since quadtree construction now always happens in the scene proxy, lots of quadtree building code was removed from UWaterMeshComponent.
Also did some general refactoring of the water mesh scene proxy to make the GDME function less cluttered.

diff --git a/Engine/Plugins/Experimental/Water/Source/Runtime/Private/WaterMeshComponent.cpp b/Engine/Plugins/Experimental/Water/Source/Runtime/Private/WaterMeshComponent.cpp
--- a/Engine/Plugins/Experimental/Water/Source/Runtime/Private/WaterMeshComponent.cpp
+++ b/Engine/Plugins/Experimental/Water/Source/Runtime/Private/WaterMeshComponent.cpp
@@ -63,20 +63,6 @@ TAutoConsoleVariable<int32> CVarWaterMeshEnabled(
 	ECVF_RenderThreadSafe
 );
 
-TAutoConsoleVariable<int32> CVarWaterMeshMIDDeduplication(
-	TEXT("r.Water.WaterMesh.MIDDeduplication"),
-	0,
-	TEXT("Deduplicate per-water body MIDs"),
-	ECVF_RenderThreadSafe
-);
-
-TAutoConsoleVariable<int32> CVarWaterMeshBuildQuadTreeInSceneProxy(
-	TEXT("r.Water.WaterMesh.BuildQuadTreeInSceneProxy"),
-	1,
-	TEXT("Moves water quadtree construction into the water mesh scene proxy."),
-	ECVF_RenderThreadSafe
-);
-
 extern TAutoConsoleVariable<float> CVarWaterSplineResampleMaxDistance;
 
 
@@ -199,25 +185,30 @@ FMaterialRelevance UWaterMeshComponent::GetWaterMaterialRelevance(ERHIFeatureLev
 	return Result;
 }
 
-bool UWaterMeshComponent::ShouldBuildQuadTreeInSceneProxy() const
+void UWaterMeshComponent::SetDynamicWaterMeshCenter(const FVector2D& NewCenter)
 {
-	const bool bShouldBuildQuadTreeInSceneProxy = CVarWaterMeshBuildQuadTreeInSceneProxy.GetValueOnAnyThread() != 0;
-	return bShouldBuildQuadTreeInSceneProxy;
+	if (!DynamicWaterMeshCenter.Equals(NewCenter))
+	{
+		DynamicWaterMeshCenter = NewCenter;
+		MarkWaterMeshGridDirty();
+	}
 }
 
-const FWaterQuadTree& UWaterMeshComponent::GetWaterQuadTree() const
+FVector2D UWaterMeshComponent::GetGlobalWaterMeshCenter() const
 {
-	ensureMsgf(!ShouldBuildQuadTreeInSceneProxy(), TEXT("Retrieving the water quadtree from the WaterMeshComponent is invalid if the quadtree is built in the scene proxy!"));
-	return WaterQuadTree;
+	const float LODCountBiasFactor = FMath::Pow(2.0f, (float)CVarWaterMeshLODCountBias.GetValueOnGameThread());
+	const float EffectiveTileSize = TileSize / LODCountBiasFactor;
+	FVector2D Result = FVector2D(FMath::GridSnap<FVector::FReal>(GetComponentLocation().X, EffectiveTileSize), FMath::GridSnap<FVector::FReal>(GetComponentLocation().Y, EffectiveTileSize));
+	return Result;
 }
 
-void UWaterMeshComponent::SetDynamicWaterMeshCenter(const FVector2D& NewCenter)
+bool UWaterMeshComponent::IsLocalOnlyTessellationEnabled() const
 {
-	if (!DynamicWaterMeshCenter.Equals(NewCenter))
+	if (const AWaterZone* WaterZone = GetOwner<AWaterZone>())
 	{
-		DynamicWaterMeshCenter = NewCenter;
-		MarkWaterMeshGridDirty();
+		return WaterZone->IsLocalOnlyTessellationEnabled();
 	}
+	return false;
 }
 
 void UWaterMeshComponent::SetTileSize(float NewTileSize)
@@ -249,27 +240,14 @@ FIntPoint UWaterMeshComponent::GetExtentInTiles() const
 FBoxSphereBounds UWaterMeshComponent::CalcBounds(const FTransform& LocalToWorld) const
 {
 	FBox NewBounds = FBox(ForceInit);
-	if (ShouldBuildQuadTreeInSceneProxy())
-	{
-		// With the water quadtree building moved into the scene proxy, the scene proxy is free to create quadtree(s) anywhere within the water zone.
-		// To ensure that the scene proxy isn't frustum culled, we need to be conservative and use the water zone bounds instead. The quadtree does
-		// fine grained frustum culling internally anyways, so this shouldn't be a performance issue.
-		AWaterZone* WaterZone = Cast<AWaterZone>(GetOwner());
-		if (WaterZone)
-		{
-			NewBounds = WaterZone->GetZoneBounds();
-		}
-	}
-	else
+	
+	// With the water quadtree building moved into the scene proxy, the scene proxy is free to create quadtree(s) anywhere within the water zone.
+	// To ensure that the scene proxy isn't frustum culled, we need to be conservative and use the water zone bounds instead. The quadtree does
+	// fine grained frustum culling internally anyways, so this shouldn't be a performance issue.
+	AWaterZone* WaterZone = Cast<AWaterZone>(GetOwner());
+	if (WaterZone)
 	{
-		// Always return valid bounds (tree is initialized with invalid bounds and if nothing is inserted, the tree bounds will stay invalid)
-		NewBounds = WaterQuadTree.GetBounds();
-
-		if (NewBounds.Min.Z >= NewBounds.Max.Z)
-		{
-			NewBounds.Min.Z = 0.0f;
-			NewBounds.Max.Z = 100.0f;
-		}
+		NewBounds = WaterZone->GetZoneBounds();
 	}
 	
 	// Add the far distance to the bounds if it's valid
@@ -293,613 +271,79 @@ void UWaterMeshComponent::RebuildWaterMesh(float InTileSize, const FIntPoint& In
 
 	bool bAnyWaterMeshesNotReady = false;
 	const bool bIsGPUQuadTree = CVarWaterMeshGPUQuadTree.GetValueOnGameThread() != 0;
-	const bool bBuildQuadTreeInSceneProxy = ShouldBuildQuadTreeInSceneProxy();
-
-	// Recreate WaterQuadTreeBuilder
-	if (bBuildQuadTreeInSceneProxy)
-	{
-		FMaterialRenderProxy* FarDistanceMaterialProxy = IsMaterialUsedWithWater(FarDistanceMaterial) ? FarDistanceMaterial->GetRenderProxy() : nullptr;
-		WaterQuadTreeBuilder.Init(WaterZone->GetZoneBounds2D(), InExtentInTiles, InTileSize, FarDistanceMaterialProxy, FarDistanceMeshExtent, bUseFarMeshWithoutOcean, bIsGPUQuadTree);
-
-		UsedMaterials.Empty();
-
-		bool bUsesFarDistanceMaterial = bUseFarMeshWithoutOcean;
-
-		WaterZone->ForEachWaterBodyComponent([&](UWaterBodyComponent* WaterBodyComponent) -> bool
-		{
-			check(WaterBodyComponent);
-			AActor* Actor = WaterBodyComponent->GetOwner();
-			check(Actor);
-
-			// Skip invisible water bodies
-			if (!WaterBodyComponent->ShouldRender() || !WaterBodyComponent->ShouldGenerateWaterMeshTile())
-			{
-				return true;
-			}
-
-			UWaterBodyInfoMeshComponent* WaterBodyInfoMeshComponent = WaterBodyComponent->GetWaterInfoMeshComponent();
-			UStaticMesh* StaticMesh = WaterBodyInfoMeshComponent ? WaterBodyInfoMeshComponent->GetStaticMesh() : nullptr;
-			bAnyWaterMeshesNotReady |= StaticMesh && StaticMesh->IsCompiling();
-			FStaticMeshRenderData* StaticMeshRenderData = StaticMesh ? StaticMesh->GetRenderData() : nullptr;
-			if (!ensure(WaterBodyInfoMeshComponent) || !ensure(StaticMesh) || !StaticMeshRenderData)
-			{
-				return true;
-			}
-
-			bUsesFarDistanceMaterial = bUsesFarDistanceMaterial || (WaterBodyComponent->GetWaterBodyType() == EWaterBodyType::Ocean);
-
-			auto GetMaterialProxy = [&](UMaterialInstanceDynamic* MID, bool bUseFallback) -> FMaterialRenderProxy*
-			{
-				UMaterialInterface* MaterialInterface = nullptr;
-				if (!MID || !IsMaterialUsedWithWater(MID))
-				{
-					MaterialInterface = bUseFallback ? UMaterial::GetDefaultMaterial(MD_Surface) : nullptr;
-				}
-				else
-				{
-					WaterBodyComponent->SetDynamicParametersOnMID(MID);
-					MaterialInterface = MID;
-				}
-				if (MaterialInterface)
-				{
-					UsedMaterials.Add(MaterialInterface);
-					return MaterialInterface->GetRenderProxy();
-				}
-				return nullptr;
-			};
-			
-			const bool bIsRiver = WaterBodyComponent->GetWaterBodyType() == EWaterBodyType::River;
-
-			FWaterQuadTreeBuilder::FWaterBody WaterBody = {};
-			WaterBody.Material = GetMaterialProxy(WaterBodyComponent->GetWaterMaterialInstance(), true);
-			WaterBody.RiverToLakeMaterial = bIsRiver ? GetMaterialProxy(WaterBodyComponent->GetRiverToLakeTransitionMaterialInstance(), false) : nullptr;
-			WaterBody.RiverToOceanMaterial = bIsRiver ? GetMaterialProxy(WaterBodyComponent->GetRiverToOceanTransitionMaterialInstance(), false) : nullptr;
-			WaterBody.StaticMeshRenderData = StaticMeshRenderData;
-			WaterBody.LocalToWorld = WaterBodyComponent->GetComponentTransform();
-			WaterBody.Bounds = WaterBodyComponent->Bounds;
-			WaterBody.OverlapMaterialPriority = WaterBodyComponent->GetOverlapMaterialPriority();
-			WaterBody.Type = WaterBodyComponent->GetWaterBodyType();
-			WaterBody.WaterBodyIndex = WaterBodyComponent->GetWaterBodyIndex();
-			WaterBody.SurfaceBaseHeight = WaterBodyComponent->GetComponentLocation().Z;
-			WaterBody.MaxWaveHeight = WaterBodyComponent->GetMaxWaveHeight();
-#if WITH_WATER_SELECTION_SUPPORT
-			WaterBody.HitProxy = new HActor(/*InActor = */Actor, /*InPrimComponent = */nullptr);
-			WaterBody.bWaterBodySelected = Actor->IsSelected();
-#endif // WITH_WATER_SELECTION_SUPPORT
-
-			if (!bIsGPUQuadTree)
-			{
-				switch (WaterBody.Type)
-				{
-				case EWaterBodyType::River:
-				{
-					TRACE_CPUPROFILER_EVENT_SCOPE(River);
-
-					TArray<FBox, TInlineAllocator<16>> Boxes;
-					TArray<UPrimitiveComponent*> CollisionComponents = WaterBodyComponent->GetCollisionComponents();
-					for (UPrimitiveComponent* Comp : CollisionComponents)
-					{
-						if (UBodySetup* BodySetup = (Comp != nullptr) ? Comp->GetBodySetup() : nullptr)
-						{
-							// Go through all sub shapes on the bodysetup to get a tight fit along water body
-							for (const FKConvexElem& ConvElem : BodySetup->AggGeom.ConvexElems)
-							{
-								TRACE_CPUPROFILER_EVENT_SCOPE(Add);
-
-								FBox SubBox = ConvElem.ElemBox.TransformBy(Comp->GetComponentTransform().ToMatrixWithScale());
-								SubBox.Max.Z += WaterBodyComponent->GetMaxWaveHeight();
-
-								Boxes.Add(SubBox);
-							}
-						}
-						else
-						{
-							// fallback on global AABB: 
-							FVector Center;
-							FVector Extent;
-							Actor->GetActorBounds(false, Center, Extent);
-							FBox Box(FBox::BuildAABB(Center, Extent));
-							Box.Max.Z += WaterBodyComponent->GetMaxWaveHeight();
-							Boxes.Add(Box);
-						}
-					}
-
-					for (const FBox& Box : Boxes)
-					{
-#if !(UE_BUILD_SHIPPING || UE_BUILD_TEST)
-						if (!!CVarWaterMeshShowTileGenerationGeometry.GetValueOnGameThread())
-						{
-							DrawDebugBox(GetWorld(), Box.GetCenter(), Box.GetExtent(), FColor::Red);
-						}
-#endif
-					}
-					WaterBody.RiverBoxes = Boxes;
-					break;
-				}
-				case EWaterBodyType::Lake:
-				{
-					TRACE_CPUPROFILER_EVENT_SCOPE(Lake);
-
-					const UWaterSplineComponent* SplineComp = WaterBodyComponent->GetWaterSpline();
-					const int32 NumOriginalSplinePoints = SplineComp->GetNumberOfSplinePoints();
-					float ConstantZ = SplineComp->GetLocationAtDistanceAlongSpline(0.0f, ESplineCoordinateSpace::World).Z;
-
-					FBox LakeBounds = Actor->GetComponentsBoundingBox(/* bNonColliding = */true);
-					LakeBounds.Max.Z += WaterBodyComponent->GetMaxWaveHeight();
-
-					// Skip lakes with less than 3 spline points
-					if (NumOriginalSplinePoints < 3)
-					{
-						return true;
-					}
-
-					TArray<TArray<FVector2D>> PolygonBatches;
-					// Reuse the convex hulls generated for the physics shape because the work has already been done, but we can fallback to a simple spline evaluation method in case there's no physics:
-					bool bUseFallbackMethod = true;
-
-					TArray<UPrimitiveComponent*> CollisionComponents = WaterBodyComponent->GetCollisionComponents();
-					if (CollisionComponents.Num() > 0)
-					{
-						UPrimitiveComponent* Comp = CollisionComponents[0];
-						if (UBodySetup* BodySetup = (Comp != nullptr) ? Comp->GetBodySetup() : nullptr)
-						{
-							FTransform CompTransform = Comp->GetComponentTransform();
-							// Go through all sub shapes on the bodysetup to get a tight fit along water body
-							for (const FKConvexElem& ConvElem : BodySetup->AggGeom.ConvexElems)
-							{
-								TRACE_CPUPROFILER_EVENT_SCOPE(Add);
-
-								// Vertex data contains the bottom vertices first, then the top ones :
-								int32 TotalNumVertices = ConvElem.VertexData.Num();
-								check(TotalNumVertices % 2 == 0);
-								int32 NumVertices = TotalNumVertices / 2;
-								if (NumVertices > 0)
-								{
-									bUseFallbackMethod = false;
-
-									// Because the physics shape is made of multiple convex hulls, we cannot simply add their vertices to one big list but have to have 1 batch of polygons per convex hull
-									TArray<FVector2D>& Polygon = PolygonBatches.Emplace_GetRef();
-									Polygon.SetNum(NumVertices);
-									for (int32 i = 0; i < NumVertices; ++i)
-									{
-										// Gather the top vertices :
-										Polygon[i] = FVector2D(CompTransform.TransformPosition(ConvElem.VertexData[NumVertices + i]));
-									}
-								}
-							}
-						}
-					}
-
-					if (bUseFallbackMethod)
-					{
-						TArray<FVector> PolyLineVertices;
-						SplineComp->ConvertSplineToPolyLine(ESplineCoordinateSpace::World, FMath::Square(CVarWaterSplineResampleMaxDistance.GetValueOnGameThread()), PolyLineVertices);
-
-						TArray<FVector2D>& Polygon = PolygonBatches.Emplace_GetRef();
-						Polygon.Reserve(PolyLineVertices.Num());
-						Algo::Transform(PolyLineVertices, Polygon, [](const FVector& Vertex) { return FVector2D(Vertex); });
-					}
-
-					for (const TArray<FVector2D>& Polygon : PolygonBatches)
-					{
-#if !(UE_BUILD_SHIPPING || UE_BUILD_TEST)
-						if (!!CVarWaterMeshShowTileGenerationGeometry.GetValueOnGameThread())
-						{
-							float Z = SplineComp->GetLocationAtDistanceAlongSpline(0.0f, ESplineCoordinateSpace::World).Z;
-							int32 NumVertices = Polygon.Num();
-							for (int32 i = 0; i < NumVertices; i++)
-							{
-								const FVector2D& Point0 = Polygon[i];
-								const FVector2D& Point1 = Polygon[(i + 1) % NumVertices];
-								DrawDebugLine(GetWorld(), FVector(Point0.X, Point0.Y, Z), FVector(Point1.X, Point1.Y, Z), FColor::Green);
-							}
-						}
-#endif // !(UE_BUILD_SHIPPING || UE_BUILD_TEST)
-					}
-
-					WaterBody.PolygonBounds = LakeBounds;
-					WaterBody.PolygonBatches = MoveTemp(PolygonBatches);
-
-					break;
-				}
-				case EWaterBodyType::Ocean:
-				{
-					TRACE_CPUPROFILER_EVENT_SCOPE(Ocean);
-
-					const UWaterSplineComponent* SplineComp = WaterBodyComponent->GetWaterSpline();
-					const int32 NumOriginalSplinePoints = SplineComp->GetNumberOfSplinePoints();
-
-					// Skip oceans with less than 3 spline points
-					if (NumOriginalSplinePoints < 3)
-					{
-						return true;
-					}
-
-					TArray<FVector2D> Polygon;
-
-					TArray<FVector> PolyLineVertices;
-					SplineComp->ConvertSplineToPolyLine(ESplineCoordinateSpace::World, FMath::Square(CVarWaterSplineResampleMaxDistance.GetValueOnGameThread()), PolyLineVertices);
-
-					Polygon.Reserve(PolyLineVertices.Num());
-					Algo::Transform(PolyLineVertices, Polygon, [](const FVector& Vertex) { return FVector2D(Vertex); });
-
-					FBox OceanBounds = WaterBodyComponent->Bounds.GetBox();
-					OceanBounds.Max.Z += WaterBodyComponent->GetMaxWaveHeight();
-
-#if !(UE_BUILD_SHIPPING || UE_BUILD_TEST)
-					if (!!CVarWaterMeshShowTileGenerationGeometry.GetValueOnGameThread())
-					{
-						float Z = SplineComp->GetLocationAtDistanceAlongSpline(0.0f, ESplineCoordinateSpace::World).Z;
-						int32 NumVertices = Polygon.Num();
-						for (int32 i = 0; i < NumVertices; i++)
-						{
-							const FVector2D& Point0 = Polygon[i];
-							const FVector2D& Point1 = Polygon[(i + 1) % NumVertices];
-							DrawDebugLine(GetWorld(), FVector(Point0.X, Point0.Y, Z), FVector(Point1.X, Point1.Y, Z), FColor::Blue);
-						}
-
-						DrawDebugBox(GetWorld(), OceanBounds.GetCenter(), OceanBounds.GetExtent(), FColor::Blue);
-					}
-#endif //!(UE_BUILD_SHIPPING || UE_BUILD_TEST)
-
-					WaterBody.PolygonBounds = OceanBounds;
-					WaterBody.PolygonBatches.Add(MoveTemp(Polygon));
-
-					break;
-				}
-				case EWaterBodyType::Transition:
-					// Transitions dont require rendering
-					break;
-				default:
-					ensureMsgf(false, TEXT("This water body type is not implemented and will not produce any water tiles. "));
-				}
-			}
-
-			WaterQuadTreeBuilder.AddWaterBody(MoveTemp(WaterBody));
-
-			return true;
-		});
-
-		if (bUsesFarDistanceMaterial && FarDistanceMaterialProxy && (FarDistanceMeshExtent > 0.0f))
-		{
-			UsedMaterials.Add(FarDistanceMaterial);
-		}
-
-		MarkRenderStateDirty();
-
-		// Force another rebuild next frame if water body meshes are still compiling
-		if (bIsGPUQuadTree && bAnyWaterMeshesNotReady)
-		{
-			bNeedsRebuild = true;
-		}
-
-		return;
-	}
-
-	// Position snapped to the grid
-	const FVector2D GridPosition = WaterZone->IsLocalOnlyTessellationEnabled() ? GetDynamicWaterMeshCenter() : FVector2D(FMath::GridSnap<FVector::FReal>(GetComponentLocation().X, InTileSize), FMath::GridSnap<FVector::FReal>(GetComponentLocation().Y, InTileSize));
-	const FVector2D WorldExtent = FVector2D(InTileSize * InExtentInTiles.X, InTileSize * InExtentInTiles.Y);
-
-	FBox2D WaterWorldBox = FBox2D(-WorldExtent + GridPosition, WorldExtent + GridPosition);
 	
-	// If the dynamic bounds is outside the full bounds of the water mesh, we shouldn't regenerate the quadtree
-	if (!(WaterWorldBox.GetArea() > 0.f))
-	{
-		return;
-	}
-
-	// This resets the tree to an initial state, ready for node insertion
-	WaterQuadTree.InitTree(WaterWorldBox, InTileSize, InExtentInTiles, bIsGPUQuadTree);
+	FMaterialRenderProxy* FarDistanceMaterialProxy = IsMaterialUsedWithWater(FarDistanceMaterial) ? FarDistanceMaterial->GetRenderProxy() : nullptr;
+	WaterQuadTreeBuilder.Init(WaterZone->GetZoneBounds2D(), InExtentInTiles, InTileSize, FarDistanceMaterialProxy, FarDistanceMeshExtent, bUseFarMeshWithoutOcean, bIsGPUQuadTree);
 
 	UsedMaterials.Empty();
 
-	// Will be updated with the ocean min bound, to be used to place the far mesh just under the ocean to avoid seams
-	float FarMeshHeight = GetComponentLocation().Z;
-	// Only use a far mesh when there is an ocean in the zone.
-	bool bHasOcean = false;
-
-	const UWaterSubsystem* WaterSubsystem = UWaterSubsystem::GetWaterSubsystem(GetWorld());
-
-	const float GlobalOceanHeight = WaterSubsystem ? WaterSubsystem->GetOceanTotalHeight() : TNumericLimits<float>::Lowest();
-	const float OceanFlood = WaterSubsystem ? WaterSubsystem->GetOceanFloodHeight() : 0.0f;
-	const bool bIsFlooded = OceanFlood > 0.0f;
-
-	// Go through all water body actors to figure out bounds and water tiles
-	AWaterZone* OwningZone = GetOwner<AWaterZone>();
-	check(OwningZone);
-
-	// Min and max user defined priority range. (Input also clamped on OverlapMaterialPriority in AWaterBody)
-	constexpr int32 MinWaterBodyPriority = -8192;
-	constexpr int32 MaxWaterBodyPriority = 8191;
-	constexpr int32 GPUQuadTreeMaxNumPriorities = 8;
-
-	// The GPU quadtree only supports 8 different priority values, so we need to remap priorities into that space.
-	// Fortunately, rivers and non-river water bodies are rendered into their own "priority space", so we don't need to worry about
-	// moving river priorities into their own range.
-	TArray<int16> SortedPriorities;
-	if (bIsGPUQuadTree)
-	{
-		OwningZone->ForEachWaterBodyComponent([this, WaterWorldBox, &SortedPriorities](UWaterBodyComponent* WaterBodyComponent)
-		{
-			check(WaterBodyComponent);
-			AActor* Actor = WaterBodyComponent->GetOwner();
-			check(Actor);
-
-			// Skip invisible water bodies
-			if (!WaterBodyComponent->ShouldRender() || !WaterBodyComponent->ShouldGenerateWaterMeshTile())
-			{
-				return true;
-			}
-
-			// Don't process water bodies which have their spline outside of this water mesh
-			const FBox WaterBodyBounds = WaterBodyComponent->Bounds.GetBox();
-			if (!WaterBodyBounds.IntersectXY(FBox(FVector(WaterWorldBox.Min, 0.0f), FVector(WaterWorldBox.Max, 0.0f))))
-			{
-				return true;
-			}
-
-			const int16 Priority = static_cast<int16>(FMath::Clamp(WaterBodyComponent->GetOverlapMaterialPriority(), MinWaterBodyPriority, MaxWaterBodyPriority));
-			SortedPriorities.AddUnique(Priority);
-
-			return true;
-		});
-
-		SortedPriorities.Sort();
-
-		if (SortedPriorities.Num() > GPUQuadTreeMaxNumPriorities)
-		{
-			UE_LOG(LogWater, Warning, TEXT("WaterZone has more unique water body priorities (%i) than can be supported with GPU driven water quadtree rendering (%i)!"), SortedPriorities.Num(), GPUQuadTreeMaxNumPriorities);
-		}
-	}
-
-	// Lambda for setting parameters on shared/deduplicated MIDs
-	auto SetDynamicParametersOnSharedMID = [&](UMaterialInstanceDynamic* InMID)
-	{
-		if ((InMID == nullptr) || (WaterSubsystem == nullptr))
-		{
-			return false;
-		}
-
-		InMID->SetScalarParameterValue(UWaterBodyComponent::WaterBodyIndexParamName, -1);
-		InMID->SetScalarParameterValue(UWaterBodyComponent::GlobalOceanHeightName, WaterSubsystem->GetOceanTotalHeight());
-		InMID->SetScalarParameterValue(UWaterBodyComponent::WaterZoneIndexParamName, OwningZone->GetWaterZoneIndex());
-		InMID->SetTextureParameterValue(UWaterBodyComponent::WaterVelocityAndHeightName, OwningZone->WaterInfoTextureArray);
-		return true;
-	};
+	bool bUsesFarDistanceMaterial = bUseFarMeshWithoutOcean;
 
-	OwningZone->ForEachWaterBodyComponent([this, bIsGPUQuadTree, WaterWorldBox, bIsFlooded, GlobalOceanHeight, OceanFlood, &FarMeshHeight, &bHasOcean, &SortedPriorities, &bAnyWaterMeshesNotReady, &SetDynamicParametersOnSharedMID](UWaterBodyComponent* WaterBodyComponent)
+	WaterZone->ForEachWaterBodyComponent([&](UWaterBodyComponent* WaterBodyComponent) -> bool
 	{
 		check(WaterBodyComponent);
 		AActor* Actor = WaterBodyComponent->GetOwner();
 		check(Actor);
 
 		// Skip invisible water bodies
-		if (!WaterBodyComponent->ShouldRender())
-		{
-			return true;
-		}
-
-		// No need to generate anything in the case of a custom water
-		if (!WaterBodyComponent->ShouldGenerateWaterMeshTile())
+		if (!WaterBodyComponent->ShouldRender() || !WaterBodyComponent->ShouldGenerateWaterMeshTile())
 		{
 			return true;
 		}
 
-		// Don't process water bodies that has their spline outside of this water mesh
-		const FBox WaterBodyBounds = WaterBodyComponent->Bounds.GetBox();
-		if (!WaterBodyBounds.IntersectXY(FBox(FVector(WaterWorldBox.Min, 0.0f), FVector(WaterWorldBox.Max, 0.0f))))
+		UWaterBodyInfoMeshComponent* WaterBodyInfoMeshComponent = WaterBodyComponent->GetWaterInfoMeshComponent();
+		UStaticMesh* StaticMesh = WaterBodyInfoMeshComponent ? WaterBodyInfoMeshComponent->GetStaticMesh() : nullptr;
+		bAnyWaterMeshesNotReady |= StaticMesh && StaticMesh->IsCompiling();
+		FStaticMeshRenderData* StaticMeshRenderData = StaticMesh ? StaticMesh->GetRenderData() : nullptr;
+		if (!ensure(WaterBodyInfoMeshComponent) || !ensure(StaticMesh) || !StaticMeshRenderData)
 		{
 			return true;
 		}
 
-		FWaterBodyRenderData RenderData;
+		bUsesFarDistanceMaterial = bUsesFarDistanceMaterial || (WaterBodyComponent->GetWaterBodyType() == EWaterBodyType::Ocean);
 
-		const EWaterBodyType WaterBodyType = WaterBodyComponent->GetWaterBodyType();
-
-		if (WaterBodyType != EWaterBodyType::Ocean)
+		auto GetMaterialProxy = [&](UMaterialInstanceDynamic* MID, bool bUseFallback) -> FMaterialRenderProxy*
 		{
-			if (bIsFlooded)
+			UMaterialInterface* MaterialInterface = nullptr;
+			if (!MID || !IsMaterialUsedWithWater(MID))
 			{
-				// If water body is below ocean height and not set to snap to the ocean height, skip it
-				const float CompareHeight = (WaterBodyType == EWaterBodyType::River) ? Actor->GetComponentsBoundingBox().Max.Z : WaterBodyComponent->GetComponentLocation().Z;
-				if (CompareHeight <= GlobalOceanHeight)
-				{
-					return true;
-				}
+				MaterialInterface = bUseFallback ? UMaterial::GetDefaultMaterial(MD_Surface) : nullptr;
 			}
-		}
-
-		const bool bDeduplicateMIDs = CVarWaterMeshMIDDeduplication.GetValueOnGameThread() != 0;
-
-		if (bDeduplicateMIDs)
-		{
-			// Assign material instance(s)
-			UMaterialInterface* WaterMaterial = WaterBodyComponent->GetWaterMaterial();
-			if (WaterMaterial)
+			else
 			{
-				const bool bIsRiver = WaterBodyType == EWaterBodyType::River;
-
-				UMaterialInterface* Materials[3] = {};
-				Materials[0] = WaterBodyComponent->GetWaterMaterial();
-				Materials[1] = bIsRiver ? WaterBodyComponent->GetRiverToLakeTransitionMaterial() : nullptr;
-				Materials[2] = bIsRiver ? WaterBodyComponent->GetRiverToOceanTransitionMaterial() : nullptr;
-
-				FName Names[] = { TEXT("WaterMID"), TEXT("LakeTransitionMID"), TEXT("OceanTransitionMID") };
-
-				UMaterialInterface* OutMaterials[3] = {};
-				for (int i = 0; i < 3; ++i)
-				{
-					UMaterialInterface* Material = Materials[i];
-					if (Material)
-					{
-						if (IsMaterialUsedWithWater(Material))
-						{
-							const bool bHasExistingMID = MaterialToMID.Contains(Material);
-							UMaterialInstanceDynamic* MID = nullptr;
-							if (!bHasExistingMID)
-							{
-								MID = FWaterUtils::GetOrCreateTransientMID(nullptr, Names[i], Material, RF_Transient | RF_NonPIEDuplicateTransient | RF_TextExportTransient);
-								SetDynamicParametersOnSharedMID(MID);
-								MaterialToMID.Add(Material, MID);
-							}
-							else
-							{
-								MID = MaterialToMID[Material];
-								// Update the parameters on this MID if this is the first time this material is seen in this rebuild.
-								// This is necessary to handle cases where GlobalOceanHeight or the water info texture pointer change
-								// TODO: GlobalOceanHeight can probably be removed, leaving only the texture parameter.
-								// Can we then skip this and set the parameters only on creation?
-								if (!UsedMaterials.Contains(MID))
-								{
-									SetDynamicParametersOnSharedMID(MID);
-								}
-							}
-							OutMaterials[i] = MID;
-						}
-						else
-						{
-							OutMaterials[i] = UMaterial::GetDefaultMaterial(MD_Surface);
-						}
-						UsedMaterials.Add(OutMaterials[i]);
-					}
-				}
-
-				RenderData.Material = OutMaterials[0] ? OutMaterials[0]->GetRenderProxy() : nullptr;
-				RenderData.RiverToLakeMaterial = OutMaterials[1] ? OutMaterials[1]->GetRenderProxy() : nullptr;
-				RenderData.RiverToOceanMaterial = OutMaterials[2] ? OutMaterials[2]->GetRenderProxy() : nullptr;
+				WaterBodyComponent->SetDynamicParametersOnMID(MID);
+				MaterialInterface = MID;
 			}
-
-			RenderData.Priority = static_cast<int16>(FMath::Clamp(WaterBodyComponent->GetOverlapMaterialPriority(), MinWaterBodyPriority, MaxWaterBodyPriority));
-			RenderData.WaterBodyIndex = static_cast<int16>(WaterBodyComponent->GetWaterBodyIndex());
-			RenderData.SurfaceBaseHeight = WaterBodyComponent->GetComponentLocation().Z;
-			RenderData.MaxWaveHeight = WaterBodyComponent->GetMaxWaveHeight();
-			RenderData.BoundsMinZ = WaterBodyComponent->Bounds.GetBox().Min.Z;
-			RenderData.BoundsMaxZ = WaterBodyComponent->Bounds.GetBox().Max.Z;
-			RenderData.WaterBodyType = static_cast<int8>(WaterBodyType);
-#if WITH_WATER_SELECTION_SUPPORT
-			RenderData.HitProxy = new HActor(/*InActor = */Actor, /*InPrimComponent = */nullptr);
-			RenderData.bWaterBodySelected = Actor->IsSelected();
-#endif // WITH_WATER_SELECTION_SUPPORT
-
-			if (WaterBodyType == EWaterBodyType::Ocean && bIsFlooded)
+			if (MaterialInterface)
 			{
-				RenderData.SurfaceBaseHeight += OceanFlood;
-				RenderData.Priority -= 1;
+				UsedMaterials.Add(MaterialInterface);
+				return MaterialInterface->GetRenderProxy();
 			}
-		}
-		else
-		{
-			// Assign material instance
-			UMaterialInstanceDynamic* WaterMaterial = WaterBodyComponent->GetWaterMaterialInstance();
-			RenderData.Material = WaterMaterial->GetRenderProxy();
-
-			if (RenderData.Material)
-			{
-				if (!IsMaterialUsedWithWater(WaterMaterial))
-				{
-					RenderData.Material = UMaterial::GetDefaultMaterial(MD_Surface)->GetRenderProxy();
-					// Add material so that the component keeps track of all potential materials used
-					UsedMaterials.Add(UMaterial::GetDefaultMaterial(MD_Surface));
-				}
-				else
-				{
-					// Add ocean height as a scalar parameter
-					WaterBodyComponent->SetDynamicParametersOnMID(WaterMaterial);
-					// Add material so that the component keeps track of all potential materials used
-					UsedMaterials.Add(WaterMaterial);
-				}
-			}
-
-			RenderData.Priority = static_cast<int16>(FMath::Clamp(WaterBodyComponent->GetOverlapMaterialPriority(), MinWaterBodyPriority, MaxWaterBodyPriority));
-			RenderData.WaterBodyIndex = static_cast<int16>(WaterBodyComponent->GetWaterBodyIndex());
-			RenderData.SurfaceBaseHeight = WaterBodyComponent->GetComponentLocation().Z;
-			RenderData.MaxWaveHeight = WaterBodyComponent->GetMaxWaveHeight();
-			RenderData.BoundsMinZ = WaterBodyComponent->Bounds.GetBox().Min.Z;
-			RenderData.BoundsMaxZ = WaterBodyComponent->Bounds.GetBox().Max.Z;
-			RenderData.WaterBodyType = static_cast<int8>(WaterBodyType);
+			return nullptr;
+		};
+
+		const bool bIsRiver = WaterBodyComponent->GetWaterBodyType() == EWaterBodyType::River;
+
+		FWaterQuadTreeBuilder::FWaterBody WaterBody = {};
+		WaterBody.Material = GetMaterialProxy(WaterBodyComponent->GetWaterMaterialInstance(), true);
+		WaterBody.RiverToLakeMaterial = bIsRiver ? GetMaterialProxy(WaterBodyComponent->GetRiverToLakeTransitionMaterialInstance(), false) : nullptr;
+		WaterBody.RiverToOceanMaterial = bIsRiver ? GetMaterialProxy(WaterBodyComponent->GetRiverToOceanTransitionMaterialInstance(), false) : nullptr;
+		WaterBody.StaticMeshRenderData = StaticMeshRenderData;
+		WaterBody.LocalToWorld = WaterBodyComponent->GetComponentTransform();
+		WaterBody.Bounds = WaterBodyComponent->Bounds;
+		WaterBody.OverlapMaterialPriority = WaterBodyComponent->GetOverlapMaterialPriority();
+		WaterBody.Type = WaterBodyComponent->GetWaterBodyType();
+		WaterBody.WaterBodyIndex = WaterBodyComponent->GetWaterBodyIndex();
+		WaterBody.SurfaceBaseHeight = WaterBodyComponent->GetComponentLocation().Z;
+		WaterBody.MaxWaveHeight = WaterBodyComponent->GetMaxWaveHeight();
 #if WITH_WATER_SELECTION_SUPPORT
-			RenderData.HitProxy = new HActor(/*InActor = */Actor, /*InPrimComponent = */nullptr);
-			RenderData.bWaterBodySelected = Actor->IsSelected();
+		WaterBody.HitProxy = new HActor(/*InActor = */Actor, /*InPrimComponent = */nullptr);
+		WaterBody.bWaterBodySelected = Actor->IsSelected();
 #endif // WITH_WATER_SELECTION_SUPPORT
 
-			if (WaterBodyType == EWaterBodyType::Ocean && bIsFlooded)
-			{
-				RenderData.SurfaceBaseHeight += OceanFlood;
-				RenderData.Priority -= 1;
-			}
-
-			// For rivers, set up transition materials if they exist
-			if (WaterBodyType == EWaterBodyType::River)
-			{
-				UMaterialInstanceDynamic* RiverToLakeMaterial = WaterBodyComponent->GetRiverToLakeTransitionMaterialInstance();
-				if (IsMaterialUsedWithWater(RiverToLakeMaterial))
-				{
-					RenderData.RiverToLakeMaterial = RiverToLakeMaterial->GetRenderProxy();
-					UsedMaterials.Add(RiverToLakeMaterial);
-					// Add ocean height as a scalar parameter
-					WaterBodyComponent->SetDynamicParametersOnMID(RiverToLakeMaterial);
-				}
-
-				UMaterialInstanceDynamic* RiverToOceanMaterial = WaterBodyComponent->GetRiverToOceanTransitionMaterialInstance();
-				if (IsMaterialUsedWithWater(RiverToOceanMaterial))
-				{
-					RenderData.RiverToOceanMaterial = RiverToOceanMaterial->GetRenderProxy();
-					UsedMaterials.Add(RiverToOceanMaterial);
-					// Add ocean height as a scalar parameter
-					WaterBodyComponent->SetDynamicParametersOnMID(RiverToOceanMaterial);
-				}
-			}
-		}
-
-		if (RenderData.RiverToLakeMaterial || RenderData.RiverToOceanMaterial)
-		{
-			// Move rivers up to it's own priority space, so that they always have precedence if they have transitions and that they only compare agains other rivers with transitions
-			RenderData.Priority += (MaxWaterBodyPriority-MinWaterBodyPriority)+1;
-		}
-
-		uint32 WaterBodyRenderDataIndex = WaterQuadTree.AddWaterBodyRenderData(RenderData);
-
-		if (bIsGPUQuadTree)
+		if (!bIsGPUQuadTree)
 		{
-			// On the GPU path, we only submit FWaterBodyQuadTreeRasterInfo for the water quadtree to be rasterized on the GPU. In particular, we need the static mesh, a transform and a priority/WaterBodyRenderData index.
-
-			// TODO: Add support for flooded ocean water heights. This is a legacy feature and currently not used in any project (?), so it's probably not urgent.
-			// In order to implement this, we would need to render a fullscreen quad/triangle for the ocean instead of using the spline derived mesh.
-
-			UWaterBodyInfoMeshComponent* WaterBodyInfoMeshComponent = WaterBodyComponent->GetWaterInfoMeshComponent();
-			if (ensure(WaterBodyInfoMeshComponent))
-			{
-				UStaticMesh* StaticMesh = WaterBodyInfoMeshComponent->GetStaticMesh();
-				if (ensure(StaticMesh))
-				{
-					bAnyWaterMeshesNotReady |= StaticMesh->IsCompiling();
-					FStaticMeshRenderData* StaticMeshRenderData = StaticMesh->GetRenderData();
-					if (StaticMeshRenderData)
-					{
-						const int16 ClampedPriority = static_cast<int16>(FMath::Clamp(WaterBodyComponent->GetOverlapMaterialPriority(), MinWaterBodyPriority, MaxWaterBodyPriority));
-
-						FWaterBodyQuadTreeRasterInfo RasterInfo;
-						RasterInfo.LocalToWorld = WaterBodyComponent->GetComponentTransform();
-						RasterInfo.RenderData = StaticMeshRenderData;
-						RasterInfo.WaterBodyRenderDataIndex = WaterBodyRenderDataIndex;
-						RasterInfo.Priority = FMath::Clamp(SortedPriorities.IndexOfByKey(ClampedPriority), 0, GPUQuadTreeMaxNumPriorities - 1);
-						RasterInfo.bIsRiver = WaterBodyType == EWaterBodyType::River;
-
-						WaterQuadTree.AddWaterBodyRasterInfo(RasterInfo);
-					}
-				}
-			}
-
-			if (WaterBodyType == EWaterBodyType::Ocean)
-			{
-				// Place far mesh height just below the ocean level
-				FarMeshHeight = RenderData.SurfaceBaseHeight - RenderData.MaxWaveHeight;
-				bHasOcean = true;
-			}
-		}
-		else
-		{
-			switch (WaterBodyType)
+			switch (WaterBody.Type)
 			{
 			case EWaterBodyType::River:
 			{
@@ -936,8 +380,6 @@ void UWaterMeshComponent::RebuildWaterMesh(float InTileSize, const FIntPoint& In
 
 				for (const FBox& Box : Boxes)
 				{
-					WaterQuadTree.AddWaterTilesInsideBounds(Box, WaterBodyRenderDataIndex);
-
 #if !(UE_BUILD_SHIPPING || UE_BUILD_TEST)
 					if (!!CVarWaterMeshShowTileGenerationGeometry.GetValueOnGameThread())
 					{
@@ -945,6 +387,7 @@ void UWaterMeshComponent::RebuildWaterMesh(float InTileSize, const FIntPoint& In
 					}
 #endif
 				}
+				WaterBody.RiverBoxes = Boxes;
 				break;
 			}
 			case EWaterBodyType::Lake:
@@ -961,7 +404,7 @@ void UWaterMeshComponent::RebuildWaterMesh(float InTileSize, const FIntPoint& In
 				// Skip lakes with less than 3 spline points
 				if (NumOriginalSplinePoints < 3)
 				{
-					break;
+					return true;
 				}
 
 				TArray<TArray<FVector2D>> PolygonBatches;
@@ -1013,8 +456,6 @@ void UWaterMeshComponent::RebuildWaterMesh(float InTileSize, const FIntPoint& In
 
 				for (const TArray<FVector2D>& Polygon : PolygonBatches)
 				{
-					WaterQuadTree.AddLake(Polygon, LakeBounds, WaterBodyRenderDataIndex);
-
 #if !(UE_BUILD_SHIPPING || UE_BUILD_TEST)
 					if (!!CVarWaterMeshShowTileGenerationGeometry.GetValueOnGameThread())
 					{
@@ -1030,63 +471,53 @@ void UWaterMeshComponent::RebuildWaterMesh(float InTileSize, const FIntPoint& In
 #endif // !(UE_BUILD_SHIPPING || UE_BUILD_TEST)
 				}
 
+				WaterBody.PolygonBounds = LakeBounds;
+				WaterBody.PolygonBatches = MoveTemp(PolygonBatches);
+
 				break;
 			}
 			case EWaterBodyType::Ocean:
 			{
 				TRACE_CPUPROFILER_EVENT_SCOPE(Ocean);
 
-				// Add ocean based on the ocean spline when there is no flood. Otherwise add ocean everywhere
-				if (bIsFlooded)
+				const UWaterSplineComponent* SplineComp = WaterBodyComponent->GetWaterSpline();
+				const int32 NumOriginalSplinePoints = SplineComp->GetNumberOfSplinePoints();
+
+				// Skip oceans with less than 3 spline points
+				if (NumOriginalSplinePoints < 3)
 				{
-					FBox OceanBounds = Actor->GetComponentsBoundingBox();
-					OceanBounds.Max.Z += WaterBodyComponent->GetMaxWaveHeight() + OceanFlood;
-					WaterQuadTree.AddWaterTilesInsideBounds(OceanBounds, WaterBodyRenderDataIndex);
+					return true;
 				}
-				else
-				{
-					const UWaterSplineComponent* SplineComp = WaterBodyComponent->GetWaterSpline();
-					const int32 NumOriginalSplinePoints = SplineComp->GetNumberOfSplinePoints();
 
-					// Skip oceans with less than 3 spline points
-					if (NumOriginalSplinePoints < 3)
-					{
-						break;
-					}
+				TArray<FVector2D> Polygon;
 
-					TArray<FVector2D> Polygon;
+				TArray<FVector> PolyLineVertices;
+				SplineComp->ConvertSplineToPolyLine(ESplineCoordinateSpace::World, FMath::Square(CVarWaterSplineResampleMaxDistance.GetValueOnGameThread()), PolyLineVertices);
 
-					TArray<FVector> PolyLineVertices;
-					SplineComp->ConvertSplineToPolyLine(ESplineCoordinateSpace::World, FMath::Square(CVarWaterSplineResampleMaxDistance.GetValueOnGameThread()), PolyLineVertices);
-
-					Polygon.Reserve(PolyLineVertices.Num());
-					Algo::Transform(PolyLineVertices, Polygon, [](const FVector& Vertex) { return FVector2D(Vertex); });
+				Polygon.Reserve(PolyLineVertices.Num());
+				Algo::Transform(PolyLineVertices, Polygon, [](const FVector& Vertex) { return FVector2D(Vertex); });
 
-					FBox OceanBounds = WaterBodyComponent->Bounds.GetBox();
-					OceanBounds.Max.Z += WaterBodyComponent->GetMaxWaveHeight();
+				FBox OceanBounds = WaterBodyComponent->Bounds.GetBox();
+				OceanBounds.Max.Z += WaterBodyComponent->GetMaxWaveHeight();
 
 #if !(UE_BUILD_SHIPPING || UE_BUILD_TEST)
-					if (!!CVarWaterMeshShowTileGenerationGeometry.GetValueOnGameThread())
+				if (!!CVarWaterMeshShowTileGenerationGeometry.GetValueOnGameThread())
+				{
+					float Z = SplineComp->GetLocationAtDistanceAlongSpline(0.0f, ESplineCoordinateSpace::World).Z;
+					int32 NumVertices = Polygon.Num();
+					for (int32 i = 0; i < NumVertices; i++)
 					{
-						float Z = SplineComp->GetLocationAtDistanceAlongSpline(0.0f, ESplineCoordinateSpace::World).Z;
-						int32 NumVertices = Polygon.Num();
-						for (int32 i = 0; i < NumVertices; i++)
-						{
-							const FVector2D& Point0 = Polygon[i];
-							const FVector2D& Point1 = Polygon[(i + 1) % NumVertices];
-							DrawDebugLine(GetWorld(), FVector(Point0.X, Point0.Y, Z), FVector(Point1.X, Point1.Y, Z), FColor::Blue);
-						}
-
-						DrawDebugBox(GetWorld(), OceanBounds.GetCenter(), OceanBounds.GetExtent(), FColor::Blue);
+						const FVector2D& Point0 = Polygon[i];
+						const FVector2D& Point1 = Polygon[(i + 1) % NumVertices];
+						DrawDebugLine(GetWorld(), FVector(Point0.X, Point0.Y, Z), FVector(Point1.X, Point1.Y, Z), FColor::Blue);
 					}
-#endif //!(UE_BUILD_SHIPPING || UE_BUILD_TEST)
 
-					WaterQuadTree.AddOcean(Polygon, OceanBounds, WaterBodyRenderDataIndex);
+					DrawDebugBox(GetWorld(), OceanBounds.GetCenter(), OceanBounds.GetExtent(), FColor::Blue);
 				}
+#endif //!(UE_BUILD_SHIPPING || UE_BUILD_TEST)
 
-				// Place far mesh height just below the ocean level
-				FarMeshHeight = RenderData.SurfaceBaseHeight - WaterBodyComponent->GetMaxWaveHeight();
-				bHasOcean = true;
+				WaterBody.PolygonBounds = OceanBounds;
+				WaterBody.PolygonBatches.Add(MoveTemp(Polygon));
 
 				break;
 			}
@@ -1098,25 +529,16 @@ void UWaterMeshComponent::RebuildWaterMesh(float InTileSize, const FIntPoint& In
 			}
 		}
 
+		WaterQuadTreeBuilder.AddWaterBody(MoveTemp(WaterBody));
+
 		return true;
 	});
 
-	// Build the far distance mesh instances if needed
-	if ((bHasOcean || bUseFarMeshWithoutOcean) && (IsMaterialUsedWithWater(FarDistanceMaterial) && FarDistanceMeshExtent > 0.0f))
+	if (bUsesFarDistanceMaterial && FarDistanceMaterialProxy && (FarDistanceMeshExtent > 0.0f))
 	{
 		UsedMaterials.Add(FarDistanceMaterial);
-
-		// Far Mesh should stitch to the edge of the water zone
-		const FBox2D FarMeshBounds = WaterZone->GetZoneBounds2D();
-
-		WaterQuadTree.AddFarMesh(FarDistanceMaterial->GetRenderProxy(), FarMeshBounds, FarDistanceMeshExtent, FarMeshHeight);
 	}
 
-	// Remove all materials from the MaterialToMID map that aren't currently used
-	MaterialToMID = MaterialToMID.FilterByPredicate([this](const auto& Pair) { return UsedMaterials.Contains(Pair.Value); });
-
-	WaterQuadTree.Unlock(true);
-
 	MarkRenderStateDirty();
 
 	// Force another rebuild next frame if water body meshes are still compiling
diff --git a/Engine/Plugins/Experimental/Water/Source/Runtime/Private/WaterMeshSceneProxy.cpp b/Engine/Plugins/Experimental/Water/Source/Runtime/Private/WaterMeshSceneProxy.cpp
--- a/Engine/Plugins/Experimental/Water/Source/Runtime/Private/WaterMeshSceneProxy.cpp
+++ b/Engine/Plugins/Experimental/Water/Source/Runtime/Private/WaterMeshSceneProxy.cpp
@@ -19,6 +19,9 @@
 #include "RenderGraphUtils.h"
 #include "StereoRenderUtils.h"
 
+// TODO: Remove this once FWaterViewExtension actually supports managing multiple local quadtrees
+#define FORCE_SINGLE_WATER_QUADTREE 1
+
 DECLARE_STATS_GROUP(TEXT("Water Mesh"), STATGROUP_WaterMesh, STATCAT_Advanced);
 
 DECLARE_DWORD_COUNTER_STAT(TEXT("Tiles Drawn"), STAT_WaterTilesDrawn, STATGROUP_WaterMesh);
@@ -164,90 +167,53 @@ public:
 FWaterMeshSceneProxy::FWaterMeshSceneProxy(UWaterMeshComponent* Component)
 	: FPrimitiveSceneProxy(Component)
 	, MaterialRelevance(Component->GetWaterMaterialRelevance(GetScene().GetFeatureLevel()))
+	, WaterQuadTreeBuilder(Component->GetWaterQuadTreeBuilder())
+	, bIsLocalOnlyTessellationEnabled(Component->IsLocalOnlyTessellationEnabled())
 {
-	AWaterZone* WaterZone = CastChecked<AWaterZone>(Component->GetOwner());
-	checkf(WaterZone != nullptr, TEXT("WaterMeshComponent is owned by an actor that is not a WaterZone. This is not supported!"));
-
 	// Need to enable single pass GetDynamicMeshElements, as the function uses all Views together to compute visible water tiles
 	bSinglePassGDME = true;
 
-	if (Component->ShouldBuildQuadTreeInSceneProxy())
-	{
-		// Get the QuadTreeBuilder and build the quadtree
-		const FWaterQuadTreeBuilder& WaterQuadTreeBuilder = Component->GetWaterQuadTreeBuilder();
-		const FVector2D GridPosition = WaterZone->IsLocalOnlyTessellationEnabled() ?
-			Component->GetDynamicWaterMeshCenter() :
-			FVector2D(FMath::GridSnap<FVector::FReal>(Component->GetComponentLocation().X, WaterQuadTreeBuilder.GetLeafSize()), FMath::GridSnap<FVector::FReal>(Component->GetComponentLocation().Y, WaterQuadTreeBuilder.GetLeafSize()));
-		WaterQuadTreeBuilder.BuildWaterQuadTree(WaterQuadTree, GridPosition);
-	}
-	else
-	{
-		WaterQuadTree = Component->GetWaterQuadTree();
-	}
-
 	// Leaf size * 0.5 equals the tightest possible LOD Scale that doesn't break the morphing. Can be scaled larger
-	LODScale = WaterQuadTree.GetLeafSize() * FMath::Max(Component->GetLODScale(), 0.5f);
+	const float LeafSize = WaterQuadTreeBuilder.GetLeafSize();
+	WaterQuadTreeConstants.LODScale = LeafSize * FMath::Max(Component->GetLODScale(), 0.5f);
+
+	int32 NumQuads = (int32)FMath::Pow(2.0f, (float)Component->GetTessellationFactor());
+	WaterQuadTreeConstants.NumQuadsLOD0 = NumQuads;
+	WaterQuadTreeConstants.NumQuadsPerIndirectDrawTile = FMath::Min((int32)FMath::RoundUpToPowerOfTwo(FMath::Clamp(CVarWaterMeshGPUQuadTreeNumQuads.GetValueOnGameThread(), 2, 128)), WaterQuadTreeConstants.NumQuadsLOD0);
 
 	// Assign the force collapse level if there is one, otherwise leave it at the default
 	if (Component->ForceCollapseDensityLevel > -1)
 	{
-		ForceCollapseDensityLevel = Component->ForceCollapseDensityLevel;
+		WaterQuadTreeConstants.ForceCollapseDensityLevel = Component->ForceCollapseDensityLevel;
 	}
 
-	int32 NumQuads = (int32)FMath::Pow(2.0f, (float)Component->GetTessellationFactor());
-	NumQuadsLOD0 = NumQuads;
-	NumQuadsPerIndirectDrawTile = FMath::Min((int32)FMath::RoundUpToPowerOfTwo(FMath::Clamp(CVarWaterMeshGPUQuadTreeNumQuads.GetValueOnGameThread(), 2, 128)), NumQuadsLOD0);
-
-	const bool bIsGPUQuadTree = WaterQuadTree.IsGPUQuadTree();
-
-	// Initialize Z bounds needed for GPU driven rendering
-	if (bIsGPUQuadTree)
-	{
-		WaterQuadTreeMinHeight = DBL_MAX;
-		WaterQuadTreeMaxHeight = -DBL_MAX;
-
-		const TArrayView<const FWaterBodyRenderData> WaterBodyRenderData = WaterQuadTree.GetWaterBodyRenderData();
-		for (const FWaterBodyRenderData& RenderData : WaterBodyRenderData)
-		{
-			WaterQuadTreeMinHeight = FMath::Min(WaterQuadTreeMinHeight, RenderData.BoundsMinZ - RenderData.MaxWaveHeight);
-			WaterQuadTreeMaxHeight = FMath::Max(WaterQuadTreeMaxHeight, RenderData.BoundsMaxZ + RenderData.MaxWaveHeight);
-		}
-
-		if (WaterQuadTreeMinHeight == DBL_MAX)
-		{
-			WaterQuadTreeMinHeight = 0.0;
-			WaterQuadTreeMaxHeight = 0.0;
-		}
-		WaterQuadTreeMinHeight -= 1.0;
-		WaterQuadTreeMaxHeight += 1.0;
-	}
+	const int32 TreeDepth = WaterQuadTreeBuilder.GetTreeDepth();
+	WaterQuadTreeConstants.DensityCount = FMath::Min(TreeDepth, (int32)FMath::FloorLog2(WaterQuadTreeConstants.NumQuadsLOD0));
 
-	DensityCount = FMath::Min(WaterQuadTree.GetTreeDepth(), (int32)FMath::FloorLog2(NumQuadsLOD0));
 
-	const float LeafSize = WaterQuadTree.GetLeafSize();
-	
 	// bIsGPUQuadTree is constant over the lifetime of this scene proxy, so we know up front if we need the indirect draw vertex factory
+	const bool bIsGPUQuadTree = WaterQuadTreeBuilder.IsGPUQuadTree();
 	if (bIsGPUQuadTree)
 	{
 		// Only create one indirect draw vertex factory, depending on whether ISR is enabled
 		const UE::StereoRenderUtils::FStereoShaderAspects Aspects(GetScene().GetShaderPlatform());
 		if (Aspects.IsInstancedStereoEnabled())
 		{
-			WaterVertexFactoryIndirectDrawISR = new FWaterVertexFactoryIndirectDrawISRType(GetScene().GetFeatureLevel(), NumQuadsPerIndirectDrawTile, NumQuadsLOD0, DensityCount, LeafSize, LODScale);
+			WaterVertexFactoryIndirectDrawISR = new FWaterVertexFactoryIndirectDrawISRType(GetScene().GetFeatureLevel(), WaterQuadTreeConstants.NumQuadsPerIndirectDrawTile, WaterQuadTreeConstants.NumQuadsLOD0, WaterQuadTreeConstants.DensityCount, LeafSize, WaterQuadTreeConstants.LODScale);
 			BeginInitResource(WaterVertexFactoryIndirectDrawISR);
 		}
 		else
 		{
-			WaterVertexFactoryIndirectDraw = new FWaterVertexFactoryIndirectDrawType(GetScene().GetFeatureLevel(), NumQuadsPerIndirectDrawTile, NumQuadsLOD0, DensityCount, LeafSize, LODScale);
+			WaterVertexFactoryIndirectDraw = new FWaterVertexFactoryIndirectDrawType(GetScene().GetFeatureLevel(), WaterQuadTreeConstants.NumQuadsPerIndirectDrawTile, WaterQuadTreeConstants.NumQuadsLOD0, WaterQuadTreeConstants.DensityCount, LeafSize, WaterQuadTreeConstants.LODScale);
 			BeginInitResource(WaterVertexFactoryIndirectDraw);
 		}
 	}
 
 	// We always need the basic CPU-driven rendering vertex factory because the far mesh uses it
-	WaterVertexFactories.Reserve(WaterQuadTree.GetTreeDepth());
-	for (uint8 i = 0; i < WaterQuadTree.GetTreeDepth(); i++)
+	WaterVertexFactories.Reserve(TreeDepth);
+	for (int32 i = 0; i < TreeDepth; i++)
 	{
-		WaterVertexFactories.Add(new FWaterVertexFactoryType(GetScene().GetFeatureLevel(), NumQuads, NumQuadsLOD0, DensityCount, LeafSize, LODScale));
+		WaterVertexFactories.Add(new FWaterVertexFactoryType(GetScene().GetFeatureLevel(), NumQuads, WaterQuadTreeConstants.NumQuadsLOD0, WaterQuadTreeConstants.DensityCount, LeafSize, WaterQuadTreeConstants.LODScale));
 		BeginInitResource(WaterVertexFactories.Last());
 
 		NumQuads /= 2;
@@ -260,32 +226,43 @@ FWaterMeshSceneProxy::FWaterMeshSceneProxy(UWaterMeshComponent* Component)
 	}
 
 	WaterVertexFactories.Shrink();
-	check(DensityCount == WaterVertexFactories.Num());
-
-	const int32 TotalLeafNodes = WaterQuadTree.GetMaxLeafCount();
-	WaterInstanceDataBuffers = new FWaterInstanceDataBuffersType(TotalLeafNodes);
-
-	WaterMeshUserDataBuffers = new FWaterMeshUserDataBuffersType(WaterInstanceDataBuffers);
-
-	WaterQuadTree.BuildMaterialIndices();
-
-	const FBox WaterInfoBounds3D = WaterZone->GetDynamicWaterInfoBounds();
-	WaterInfoBounds = FBox2D(FVector2D(WaterInfoBounds3D.Min), FVector2D(WaterInfoBounds3D.Max));
+	check(WaterQuadTreeConstants.DensityCount == WaterVertexFactories.Num());
 
 #if RHI_RAYTRACING
-	RayTracingWaterData.SetNum(DensityCount);
+	RayTracingWaterData.SetNum(WaterQuadTreeConstants.DensityCount);
 #endif
 
-	// Always do CPU occlusion queries, even if this is a GPU quadtree. The GPU quadtree still potentially uses the far mesh which is CPU driven.
-	const int32 MaxQueries = CVarWaterMeshOcclusionCullingMaxQueries.GetValueOnGameThread();
-	const bool bIncludeFarMeshOcclusionQueries = CVarWaterMeshOcclusionCullingIncludeFarMesh.GetValueOnGameThread() != 0;
-	OcclusionCullingBounds = WaterQuadTree.ComputeNodeBounds(MaxQueries, CVarWaterMeshOcclusionCullExpandBoundsAmountXY.GetValueOnGameThread(), bIncludeFarMeshOcclusionQueries, &OcclusionResultsFarMeshOffset);
-	EmptyOcclusionCullingBounds.Add(WaterQuadTree.GetBoundsIncludingFarMesh());
-	// If this is a GPU quadtree, the CPU root node will have an invalid bounding box, so derive conservative bounds now
-	if (bIsGPUQuadTree && !OcclusionCullingBounds.IsEmpty())
+	// In case we don't need to support multiple local quadtrees, create a global quadtree right away and initialize occlusion culling bounds.
+	// Local quadtrees are incompatible with CPU-based occlusion culling as the number of quadtrees can change during the lifetime of the scene proxy.
+#if !FORCE_SINGLE_WATER_QUADTREE
+	if (!bIsLocalOnlyTessellationEnabled)
+#endif
+	{
+#if FORCE_SINGLE_WATER_QUADTREE
+		CreateViewWaterQuadTree(INDEX_NONE, bIsLocalOnlyTessellationEnabled ? Component->GetDynamicWaterMeshCenter() : Component->GetGlobalWaterMeshCenter());
+#else
+		CreateViewWaterQuadTree(INDEX_NONE, Component->GetGlobalWaterMeshCenter());
+#endif
+		FViewWaterQuadTree* ViewWaterQuadTree = ViewQuadTrees.Find(INDEX_NONE);
+		const FWaterQuadTree& WaterQuadTree = ViewWaterQuadTree->GetWaterQuadTree();
+
+		// Always do CPU occlusion queries, even if this is a GPU quadtree. The GPU quadtree still potentially uses the far mesh which is CPU driven.
+		const int32 MaxQueries = CVarWaterMeshOcclusionCullingMaxQueries.GetValueOnGameThread();
+		const bool bIncludeFarMeshOcclusionQueries = CVarWaterMeshOcclusionCullingIncludeFarMesh.GetValueOnGameThread() != 0;
+		OcclusionCullingBounds = WaterQuadTree.ComputeNodeBounds(MaxQueries, CVarWaterMeshOcclusionCullExpandBoundsAmountXY.GetValueOnGameThread(), bIncludeFarMeshOcclusionQueries, &OcclusionResultsFarMeshOffset);
+		EmptyOcclusionCullingBounds.Add(WaterQuadTree.GetBoundsIncludingFarMesh());
+		// If this is a GPU quadtree, the CPU root node will have an invalid bounding box, so derive conservative bounds now
+		if (bIsGPUQuadTree && !OcclusionCullingBounds.IsEmpty())
+		{
+			OcclusionCullingBounds[0] = FBox(FVector(WaterQuadTree.GetTileRegion().Min, ViewWaterQuadTree->GetMinHeight()), FVector(WaterQuadTree.GetTileRegion().Max, ViewWaterQuadTree->GetMaxHeight()));
+		}
+	}
+#if !FORCE_SINGLE_WATER_QUADTREE
+	else
 	{
-		OcclusionCullingBounds[0] = FBox(FVector(WaterQuadTree.GetTileRegion().Min, WaterQuadTreeMinHeight), FVector(WaterQuadTree.GetTileRegion().Max, WaterQuadTreeMaxHeight));
+		EmptyOcclusionCullingBounds.Add(Component->Bounds);
 	}
+#endif
 
 	// When using a GPU quadtree, this scene proxy allocates pooled buffers in GetDynamicMeshElements. AllocatePooledBuffer must be called on the renderthread.
 	// The callback on GWaterMeshGPUWork seems to be invoked after all GetDynamicMeshElements tasks finish, so there should be no race condition there.
@@ -311,10 +288,6 @@ FWaterMeshSceneProxy::~FWaterMeshSceneProxy()
 		delete WaterVertexFactoryIndirectDrawISR;
 	}
 
-	delete WaterInstanceDataBuffers;
-
-	delete WaterMeshUserDataBuffers;
-
 #if RHI_RAYTRACING
 	for (auto& WaterDataArray : RayTracingWaterData)
 	{
@@ -331,23 +304,16 @@ void FWaterMeshSceneProxy::CreateRenderThreadResources(FRHICommandListBase& RHIC
 {
 	SceneProxyCreatedFrameNumberRenderThread = GFrameNumberRenderThread;
 
-	if (WaterQuadTree.IsGPUQuadTree())
+	// Register a callback with the FWaterViewExtension to be executed every frame before rendering the scene.
+	if (WaterQuadTreeBuilder.IsGPUQuadTree())
 	{
 		FWaterMeshGPUWork::FCallback Callback;
 		Callback.Proxy = this;
 		Callback.Function = [this](FRDGBuilder& GraphBuilder, bool bDepthBufferIsPopulated)
 		{
-			if (bNeedToTraverseGPUQuadTree)
+			for (auto& Pair : ViewQuadTrees)
 			{
-				if (!QuadTreeGPU.IsInitialized())
-				{
-					BuildGPUQuadTree(GraphBuilder);
-				}
-				FWaterQuadTreeGPU::FTraverseParams TraverseParams = MoveTemp(WaterQuadTreeGPUTraverseParams);
-				TraverseParams.bDepthBufferIsPopulated = bDepthBufferIsPopulated;
-				QuadTreeGPU.Traverse(GraphBuilder, TraverseParams);
-				bNeedToTraverseGPUQuadTree = false;
-				WaterQuadTreeGPUTraverseParams = {};
+				Pair.Value.TraverseGPUQuadTree(GraphBuilder, bDepthBufferIsPopulated);
 			}
 		};
 		GWaterMeshGPUWork.Callbacks.Add(MoveTemp(Callback));
@@ -356,218 +322,283 @@ void FWaterMeshSceneProxy::CreateRenderThreadResources(FRHICommandListBase& RHIC
 
 void FWaterMeshSceneProxy::DestroyRenderThreadResources()
 {
-	if (WaterQuadTree.IsGPUQuadTree())
+	if (WaterQuadTreeBuilder.IsGPUQuadTree())
 	{
 		GWaterMeshGPUWork.Callbacks.RemoveAllSwap([this](const FWaterMeshGPUWork::FCallback& Callback) { return Callback.Proxy == this; });
 	}
 }
 
-FWaterMeshSceneProxy::FWaterLODParams FWaterMeshSceneProxy::GetWaterLODParams(const FVector& Position) const
-{
-	float WaterHeightForLOD = 0.0f;
-	WaterQuadTree.QueryInterpolatedTileBaseHeightAtLocation(FVector2D(Position), WaterHeightForLOD);
-
-	// Need to let the lowest LOD morph globally towards the next LOD. When the LOD is done morphing, simply clamp the LOD in the LOD selection to effectively promote the lowest LOD to the same LOD level as the one above
-	float DistToWater = FMath::Abs(Position.Z - WaterHeightForLOD) / LODScale;
-	DistToWater = FMath::Max(DistToWater - 2.0f, 0.0f);
-	DistToWater *= 2.0f;
-
-	// Clamp to WaterTileQuadTree.GetLODCount() - 1.0f prevents the last LOD to morph
-	const float FloatLOD = FMath::Clamp(FMath::Log2(DistToWater), 0.0f, WaterQuadTree.GetTreeDepth() - 1.0f);
-
-	FWaterLODParams WaterLODParams;
-	WaterLODParams.HeightLODFactor = FMath::Frac(FloatLOD);
-	WaterLODParams.LowestLOD = FMath::Clamp(FMath::FloorToInt(FloatLOD), 0, WaterQuadTree.GetTreeDepth() - 1);
-	WaterLODParams.WaterHeightForLOD = WaterHeightForLOD;
-
-	return WaterLODParams;
-}
-
 void FWaterMeshSceneProxy::GetDynamicMeshElements(const TArray<const FSceneView*>& Views, const FSceneViewFamily& ViewFamily, uint32 VisibilityMap, FMeshElementCollector& Collector) const
 {
 	CSV_SCOPED_TIMING_STAT_EXCLUSIVE(Water);
 	TRACE_CPUPROFILER_EVENT_SCOPE(FWaterMeshSceneProxy::GetDynamicMeshElements);
-	FRHICommandListBase& RHICmdList = Collector.GetRHICommandList();
-
-	// The water render groups we have to render for this batch : 
-	TArray<EWaterMeshRenderGroupType, TInlineAllocator<FWaterVertexFactoryType::NumRenderGroups>> BatchRenderGroups;
-	// By default, render all water tiles : 
-	BatchRenderGroups.Add(EWaterMeshRenderGroupType::RG_RenderWaterTiles);
-
-#if WITH_WATER_SELECTION_SUPPORT
-	bool bHasSelectedInstances = IsSelected();
-	const bool bSelectionRenderEnabled = GIsEditor && ViewFamily.EngineShowFlags.Selection;
-
-	if (bSelectionRenderEnabled && bHasSelectedInstances)
-	{
-		// Don't render all in one group: instead, render 2 groups : first, the selected only then, the non-selected only :
-		BatchRenderGroups[0] = EWaterMeshRenderGroupType::RG_RenderSelectedWaterTilesOnly;
-		BatchRenderGroups.Add(EWaterMeshRenderGroupType::RG_RenderUnselectedWaterTilesOnly);
-	}
-#endif // WITH_WATER_SELECTION_SUPPORT
-
+	
 	if (!HasWaterData() || !FWaterUtils::IsWaterMeshRenderingEnabled(/*bIsRenderThread = */true))
 	{
 		return;
 	}
+	
+	FRHICommandListBase& RHICmdList = Collector.GetRHICommandList();
 
-	// Figure out which views have wireframe enabled
-	uint32 WireframeVisibilityMap = 0;
-	if (AllowDebugViewmodes())
+	// Handle wireframe views
+	FColoredMaterialRenderProxy* WireframeMaterialInstance = nullptr;
+	const uint32 WireframeVisibilityMap = GetWireframeVisibilityMapAndMaterial(Views, VisibilityMap, Collector, WireframeMaterialInstance);
+
+	// Get a mapping from view index to quadtree key
+	TArray<int32, TInlineAllocator<8>> ViewToQuadTreeKey = GetViewToQuadTreeMapping(Views, VisibilityMap);
+	
+	// Iterate over all quadtrees and generate draw calls for all views that map to that quadtree.
+	for (const auto& Pair : ViewQuadTrees)
 	{
-		if (CVarWaterMeshShowWireframe.GetValueOnRenderThread() == 1)
-		{
-			// Force all views to wireframe
-			WireframeVisibilityMap = VisibilityMap;
-		}
-		else
+		const FViewWaterQuadTree& ViewWaterQuadTree = Pair.Value;
+		const FWaterQuadTree& WaterQuadTree = ViewWaterQuadTree.GetWaterQuadTree();
+
+		// Compute a subset mask of VisibilityMap of all views that use this quadtree
+		uint32 LocalViewMask = 0;
+		bool bEncounteredISRView = false;
+		int32 InstanceFactor = 1;
+		for (int32 ViewIndex = 0; ViewIndex < Views.Num(); ViewIndex++)
 		{
-			for (int32 ViewIndex = 0; ViewIndex < Views.Num(); ViewIndex++)
+			if (VisibilityMap & (1 << ViewIndex) && ViewToQuadTreeKey[ViewIndex] == Pair.Key)
 			{
-				if ((VisibilityMap & (1 << ViewIndex)) && Views[ViewIndex]->Family->EngineShowFlags.Wireframe)
+				LocalViewMask |= 1 << ViewIndex;
+				const FSceneView* View = Views[ViewIndex];
+				if (!bEncounteredISRView && View->IsInstancedStereoPass())
 				{
-					WireframeVisibilityMap |= (1 << ViewIndex);
+					bEncounteredISRView = true;
+					InstanceFactor = View->GetStereoPassInstanceFactor();
 				}
 			}
 		}
-	}
 
-	// Set up wireframe material (if needed)
-	FColoredMaterialRenderProxy* WireframeMaterialInstance = nullptr;
-	if (WireframeVisibilityMap && CVarWaterMeshShowWireframeAtBaseHeight.GetValueOnRenderThread() == 1)
-	{
-		WireframeMaterialInstance = new FColoredMaterialRenderProxy(
-			GEngine->WireframeMaterial ? GEngine->WireframeMaterial->GetRenderProxy() : NULL,
-			FColor::Cyan);
+		// The water render groups we have to render for this batch : 
+		const TArray<EWaterMeshRenderGroupType, TInlineAllocator<FWaterVertexFactoryType::NumRenderGroups>> BatchRenderGroups = GetBatchRenderGroups(Views, LocalViewMask);
 
-		Collector.RegisterOneFrameMaterialProxy(WireframeMaterialInstance);
-	}
+		const int32 NumWaterMaterials = WaterQuadTree.GetWaterMaterials().Num();
+		const int32 NumBuckets = NumWaterMaterials * WaterQuadTreeConstants.DensityCount;
+		const int32 NumBucketsIndirect = NumWaterMaterials; // Indirect draws use a single density mesh tile
 
-	const int32 NumWaterMaterials = WaterQuadTree.GetWaterMaterials().Num();
-	const int32 NumBuckets = NumWaterMaterials * DensityCount;
-	const int32 NumBucketsIndirect = NumWaterMaterials; // Indirect draws use a single density mesh tile
+		if (WaterQuadTree.IsGPUQuadTree() && NumBucketsIndirect > 0)
+		{
+			const FViewWaterQuadTree::FUserDataAndIndirectArgs WaterMeshUserDataAndIndirectArgs = ViewWaterQuadTree.PrepareGPUQuadTreeForRendering(Views, LocalViewMask, Collector, WaterQuadTreeConstants, BatchRenderGroups, RHICmdList);
 
-	// If this is a GPU quadtree, we are using GPU driven rendering to first create a quadtree on the GPU (once) and then traverse it every frame to build a list of indirect draws
-	const bool bIsGPUQuadTree = WaterQuadTree.IsGPUQuadTree();
-	if (bIsGPUQuadTree && NumBucketsIndirect > 0)
-	{
-		bool bEncounteredISRView = false;
-		int32 InstanceFactor = 1;
-		TArray<const FSceneView*> VisibleViews;
+			// Go through all buckets and issue one batched draw call per LOD level per material per view
+			int32 CompactViewIndex = 0;
+			for (int32 ViewIndex = 0; ViewIndex < Views.Num(); ViewIndex++)
+			{
+				// when rendering ISR, don't process the instanced view
+				if ((LocalViewMask & (1 << ViewIndex)) && (!bEncounteredISRView || Views[ViewIndex]->IsPrimarySceneView()))
+				{
+					TRACE_CPUPROFILER_EVENT_SCOPE(BucketsPerView);
+
+					const bool bInstancedStereo = Views[ViewIndex]->bIsInstancedStereoEnabled;
+					const bool bViewIsWireframe = WireframeVisibilityMap & (1 << ViewIndex);
+
+					for (int32 MaterialIndex = 0; MaterialIndex < NumWaterMaterials; ++MaterialIndex)
+					{
+						TRACE_CPUPROFILER_EVENT_SCOPE(MaterialBucket);
+						bool bMaterialDrawn = false;
+
+						// We only render all tiles with one or multiple instances of a single density quad mesh, so no need to create one bucket per density.
+						{
+							const int32 BucketIndex = MaterialIndex;
+
+							TRACE_CPUPROFILER_EVENT_SCOPE(DensityBucket);
+
+							const FMaterialRenderProxy* MaterialRenderProxy = bViewIsWireframe && WireframeMaterialInstance ? WireframeMaterialInstance : WaterQuadTree.GetWaterMaterials()[MaterialIndex];
+							check(MaterialRenderProxy != nullptr);
+
+							bool bUseForDepthPass = false;
+
+							// If there's a valid material, use that to figure out the depth pass status
+							if (const FMaterial* BucketMaterial = MaterialRenderProxy->GetMaterialNoFallback(GetScene().GetFeatureLevel()))
+							{
+								// Preemptively turn off depth rendering for this mesh batch if the material doesn't need it
+								bUseForDepthPass = !BucketMaterial->GetShadingModels().HasShadingModel(MSM_SingleLayerWater) && !IsTranslucentOnlyBlendMode(*BucketMaterial);
+							}
+
+							bMaterialDrawn = true;
+							for (EWaterMeshRenderGroupType RenderGroup : BatchRenderGroups)
+							{
+								// Set up mesh batch
+								FMeshBatch& Mesh = Collector.AllocateMesh();
+								Mesh.bWireframe = bViewIsWireframe;
+								Mesh.VertexFactory = bInstancedStereo ? (FVertexFactory*)WaterVertexFactoryIndirectDrawISR : (FVertexFactory*)WaterVertexFactoryIndirectDraw;
+								Mesh.MaterialRenderProxy = MaterialRenderProxy;
+								Mesh.ReverseCulling = IsLocalToWorldDeterminantNegative();
+								Mesh.Type = PT_TriangleList;
+								Mesh.DepthPriorityGroup = SDPG_World;
+								Mesh.bCanApplyViewModeOverrides = false;
+								Mesh.bUseForMaterial = true;
+								Mesh.CastShadow = false;
+								// Preemptively turn off depth rendering for this mesh batch if the material doesn't need it
+								Mesh.bUseForDepthPass = bUseForDepthPass;
+								Mesh.bUseAsOccluder = false;
+
+#if WITH_WATER_SELECTION_SUPPORT
+								Mesh.bUseSelectionOutline = (RenderGroup == EWaterMeshRenderGroupType::RG_RenderSelectedWaterTilesOnly);
+								Mesh.bUseWireframeSelectionColoring = (RenderGroup == EWaterMeshRenderGroupType::RG_RenderSelectedWaterTilesOnly);
+#endif // WITH_WATER_SELECTION_SUPPORT
+
+								Mesh.Elements.SetNumZeroed(1);
+								{
+									TRACE_CPUPROFILER_EVENT_SCOPE_STR("Setup batch element");
+
+									// Set up one mesh batch element
+									FMeshBatchElement& BatchElement = Mesh.Elements[0];
+
+									// Set up for indirect draw
+									BatchElement.FirstIndex = 0;
+									BatchElement.NumPrimitives = 0; // Must be 0 to enable usage of IndirectArgsBuffer
+									BatchElement.IndirectArgsBuffer = WaterMeshUserDataAndIndirectArgs.IndirectArgs->GetRHI();
+									BatchElement.IndirectArgsOffset = (CompactViewIndex * NumBucketsIndirect + BucketIndex) * sizeof(FRHIDrawIndexedIndirectParameters);
+									BatchElement.UserData = (void*)WaterMeshUserDataAndIndirectArgs.UserData[(int32)RenderGroup];
+									BatchElement.UserIndex = CompactViewIndex * NumBucketsIndirect + BucketIndex;
+
+									BatchElement.IndexBuffer = bInstancedStereo ? WaterVertexFactoryIndirectDrawISR->IndexBuffer : WaterVertexFactoryIndirectDraw->IndexBuffer;
+									BatchElement.PrimitiveIdMode = PrimID_ForceZero;
+
+									// We need the uniform buffer of this primitive because it stores the proper value for the bOutputVelocity flag.
+									// The identity primitive uniform buffer simply stores false for this flag which leads to missing motion vectors.
+									BatchElement.PrimitiveUniformBuffer = GetUniformBuffer();
+								}
+
+								{
+									INC_DWORD_STAT(STAT_WaterDrawCalls);
+									TRACE_CPUPROFILER_EVENT_SCOPE(Collector.AddMesh);
+
+									Collector.AddMesh(ViewIndex, Mesh);
+								}
+							}
+						}
+
+						INC_DWORD_STAT_BY(STAT_WaterDrawnMats, (int32)bMaterialDrawn);
+					}
+
+					++CompactViewIndex;
+				}
+			}
+		}
+	
+		// Even if we have a GPU quadtree, we still go through the setup for the non-GPU path because the water quadtree might have an ocean far mesh which we render traditionally.
+		// If this is a GPU quadtree, traversing it on the CPU will only return the far mesh tiles (if present).
+
+		TArray<FWaterQuadTree::FTraversalOutput, TInlineAllocator<4>> WaterInstanceDataPerView;
 
-		// Gather per view data for all renderable views
+		// Gather visible tiles, their lod and materials for all renderable views (skip right view when stereo pair is rendered instanced)
 		for (int32 ViewIndex = 0; ViewIndex < Views.Num(); ViewIndex++)
 		{
 			const FSceneView* View = Views[ViewIndex];
-			if (!bEncounteredISRView && View->IsInstancedStereoPass())
-			{
-				bEncounteredISRView = true;
-				InstanceFactor = View->GetStereoPassInstanceFactor();
-			}
 
 			// skip gathering visible tiles from instanced right eye views
-			if ((VisibilityMap & (1 << ViewIndex)) && (!bEncounteredISRView || View->IsPrimarySceneView()))
+			if ((LocalViewMask & (1 << ViewIndex)) && (!bEncounteredISRView || View->IsPrimarySceneView()))
 			{
-				VisibleViews.Add(View);
+				const FVector ObserverPosition = View->ViewMatrices.GetViewOrigin();
+
+				const FViewWaterQuadTree::FWaterLODParams WaterLODParams = ViewWaterQuadTree.GetWaterLODParams(ObserverPosition, WaterQuadTreeConstants.LODScale);
+
+#if !(UE_BUILD_SHIPPING || UE_BUILD_TEST)
+				if (CVarWaterMeshShowLODLevels.GetValueOnRenderThread())
+				{
+					for (int32 i = WaterLODParams.LowestLOD; i < WaterQuadTree.GetTreeDepth(); i++)
+					{
+						float LODDist = FWaterQuadTree::GetLODDistance(i, WaterQuadTreeConstants.LODScale);
+						FVector Orig = FVector(FVector2D(ObserverPosition), WaterLODParams.WaterHeightForLOD);
+
+						DrawCircle(Collector.GetPDI(ViewIndex), Orig, FVector::ForwardVector, FVector::RightVector, GColorList.GetFColorByIndex(i + 1), LODDist, 64, 0);
+					}
+				}
+#endif
+				TRACE_CPUPROFILER_EVENT_SCOPE(QuadTreeTraversalPerView);
+
+				FWaterQuadTree::FTraversalOutput& WaterInstanceData = WaterInstanceDataPerView.Emplace_GetRef();
+				WaterInstanceData.BucketInstanceCounts.Empty(NumBuckets);
+				WaterInstanceData.BucketInstanceCounts.AddZeroed(NumBuckets);
+				if (!!CVarWaterMeshPreAllocStagingInstanceMemory.GetValueOnRenderThread())
+				{
+					WaterInstanceData.StagingInstanceData.Empty(HistoricalMaxViewInstanceCount);
+				}
+
+				FWaterQuadTree::FTraversalDesc TraversalDesc;
+				TraversalDesc.LowestLOD = WaterLODParams.LowestLOD;
+				TraversalDesc.HeightMorph = WaterLODParams.HeightLODFactor;
+				TraversalDesc.LODCount = WaterQuadTree.GetTreeDepth();
+				TraversalDesc.DensityCount = WaterQuadTreeConstants.DensityCount;
+				TraversalDesc.ForceCollapseDensityLevel = WaterQuadTreeConstants.ForceCollapseDensityLevel;
+				TraversalDesc.Frustum = View->ViewFrustum;
+				TraversalDesc.ObserverPosition = ObserverPosition;
+				TraversalDesc.PreViewTranslation = View->ViewMatrices.GetPreViewTranslation();
+				TraversalDesc.LODScale = WaterQuadTreeConstants.LODScale;
+				TraversalDesc.bLODMorphingEnabled = !!CVarWaterMeshLODMorphEnabled.GetValueOnRenderThread();
+				TraversalDesc.WaterInfoBounds = WaterQuadTree.GetTileRegion();
+				TraversalDesc.OcclusionCullingResults = nullptr;
+				TraversalDesc.OcclusionCullingFarMeshOffset = OcclusionResultsFarMeshOffset;
+				if (CVarWaterMeshOcclusionCulling.GetValueOnRenderThread() != 0)
+				{
+					const FOcclusionCullingResults* CullingResults = OcclusionResults.Find(View->GetViewKey());
+					if (CullingResults && !CullingResults->Results.IsEmpty())
+					{
+						TraversalDesc.OcclusionCullingResults = &CullingResults->Results;
+					}
+				}
+
+#if !(UE_BUILD_SHIPPING || UE_BUILD_TEST)
+				//Debug
+				TraversalDesc.DebugPDI = Collector.GetPDI(ViewIndex);
+				TraversalDesc.DebugShowTile = CVarWaterMeshShowTileBounds.GetValueOnRenderThread();
+				TraversalDesc.bDebugDrawIntoForeground = CVarWaterMeshShowTileBoundsForeground.GetValueOnRenderThread() != 0;
+#endif
+				WaterQuadTree.BuildWaterTileInstanceData(TraversalDesc, WaterInstanceData);
+
+				HistoricalMaxViewInstanceCount = FMath::Max(HistoricalMaxViewInstanceCount, WaterInstanceData.InstanceCount);
 			}
 		}
 
-		const int32 NumVisibleViews = VisibleViews.Num();
-		const int32 NumIndirectDrawCalls = NumBucketsIndirect * NumVisibleViews;
-		const int32 LeafCountUpperBound = WaterQuadTree.GetMaxLeafCount() * NumVisibleViews;
-		const int32 NumInstanceDataSlots = FMath::Max(1, static_cast<int32>(LeafCountUpperBound * FMath::Clamp(CVarWaterMeshGPUQuadInstanceDataAllocMult.GetValueOnRenderThread(), 0.0f, 8.0f)));
-
-		// Allocate buffers for the indirect draws
-		struct FIndirectDrawResources
-		{
-			TRefCountPtr<FRDGPooledBuffer> IndirectArgs;
-			TRefCountPtr<FRDGPooledBuffer> InstanceDataOffsets;
-			TRefCountPtr<FRDGPooledBuffer> InstanceData0;
-			TRefCountPtr<FRDGPooledBuffer> InstanceData1;
-			TRefCountPtr<FRDGPooledBuffer> InstanceData2;
-			TRefCountPtr<FRDGPooledBuffer> InstanceData3;
-		};
-		FIndirectDrawResources& IndirectDrawResources = Collector.AllocateOneFrameResource<FIndirectDrawResources>();
-		IndirectDrawResources.IndirectArgs = AllocatePooledBuffer(FRDGBufferDesc::CreateIndirectDesc<FRHIDrawIndexedIndirectParameters>(FMath::Max(1, NumIndirectDrawCalls)), TEXT("WaterQuadTree.IndirectArgsBuffer"));
-		IndirectDrawResources.InstanceDataOffsets = AllocatePooledBuffer(FRDGBufferDesc::CreateBufferDesc(sizeof(uint32), FMath::Max(1, NumIndirectDrawCalls)), TEXT("WaterQuadTree.InstanceDataOffsetsBuffer"));
-		IndirectDrawResources.InstanceData0 = AllocatePooledBuffer(FRDGBufferDesc::CreateBufferDesc(sizeof(uint32), FMath::Max(1, NumInstanceDataSlots)), TEXT("WaterQuadTree.InstanceDataBuffer0"));
-		IndirectDrawResources.InstanceData1 = AllocatePooledBuffer(FRDGBufferDesc::CreateBufferDesc(sizeof(uint32), FMath::Max(1, NumInstanceDataSlots)), TEXT("WaterQuadTree.InstanceDataBuffer1"));
-		IndirectDrawResources.InstanceData2 = AllocatePooledBuffer(FRDGBufferDesc::CreateBufferDesc(sizeof(uint32), FMath::Max(1, NumInstanceDataSlots)), TEXT("WaterQuadTree.InstanceDataBuffer2"));
-		if (WITH_WATER_SELECTION_SUPPORT != 0)
+		// Get number of total instances for all views
+		int32 TotalInstanceCount = 0;
+		for (const FWaterQuadTree::FTraversalOutput& WaterInstanceData : WaterInstanceDataPerView)
 		{
-			IndirectDrawResources.InstanceData3 = AllocatePooledBuffer(FRDGBufferDesc::CreateBufferDesc(sizeof(uint32), FMath::Max(1, NumInstanceDataSlots)), TEXT("WaterQuadTree.InstanceDataBuffer3"));
+			TotalInstanceCount += WaterInstanceData.InstanceCount;
 		}
-		
 
-		// Traverse quadtree. We actually only set up the parameters here; the actual traversal is done later in the frame in a callback invoked from FWaterViewExtension::PreRenderBasePass_RenderThread.
-		// That callback accesses these params, so we need to ensure that they are still valid at that point.
+		if (TotalInstanceCount == 0)
 		{
-			WaterQuadTreeGPUTraverseParams = {};
-			WaterQuadTreeGPUTraverseParams.OutIndirectArgsBuffer = IndirectDrawResources.IndirectArgs;
-			WaterQuadTreeGPUTraverseParams.OutInstanceDataOffsetsBuffer = IndirectDrawResources.InstanceDataOffsets;
-			WaterQuadTreeGPUTraverseParams.OutInstanceData0Buffer = IndirectDrawResources.InstanceData0;
-			WaterQuadTreeGPUTraverseParams.OutInstanceData1Buffer = IndirectDrawResources.InstanceData1;
-			WaterQuadTreeGPUTraverseParams.OutInstanceData2Buffer = IndirectDrawResources.InstanceData2;
-			WaterQuadTreeGPUTraverseParams.OutInstanceData3Buffer = IndirectDrawResources.InstanceData3;
-			WaterQuadTreeGPUTraverseParams.Views = MoveTemp(VisibleViews);
-			WaterQuadTreeGPUTraverseParams.QuadTreePosition = FVector(WaterQuadTree.GetTileRegion().Min, WaterQuadTreeMinHeight);
-			WaterQuadTreeGPUTraverseParams.CullingBounds = WaterInfoBounds.ShiftBy(-WaterQuadTree.GetTileRegion().Min);
-			WaterQuadTreeGPUTraverseParams.NumDensities = DensityCount;
-			WaterQuadTreeGPUTraverseParams.NumMaterials = NumWaterMaterials;
-			WaterQuadTreeGPUTraverseParams.NumQuadsLOD0 = NumQuadsLOD0;
-			WaterQuadTreeGPUTraverseParams.NumQuadsPerTileSide = NumQuadsPerIndirectDrawTile;
-			WaterQuadTreeGPUTraverseParams.ForceCollapseDensityLevel = ForceCollapseDensityLevel;
-			WaterQuadTreeGPUTraverseParams.LeafSize = WaterQuadTree.GetLeafSize();
-			WaterQuadTreeGPUTraverseParams.LODScale = LODScale;
-			WaterQuadTreeGPUTraverseParams.DebugShowTile = CVarWaterMeshShowTileBounds.GetValueOnRenderThread();
-			WaterQuadTreeGPUTraverseParams.bWithWaterSelectionSupport = WITH_WATER_SELECTION_SUPPORT != 0;
-			WaterQuadTreeGPUTraverseParams.bLODMorphingEnabled = !!CVarWaterMeshLODMorphEnabled.GetValueOnRenderThread();
-
-			bNeedToTraverseGPUQuadTree = true;
+			// no instance visible, early exit
+			continue;
 		}
 
+		ViewWaterQuadTree.GetWaterInstanceDataBuffers()->Lock(RHICmdList, TotalInstanceCount* InstanceFactor);
 
-		using FWaterVertexFactoryUserDataWrapperType = TWaterVertexFactoryUserDataWrapper<WITH_WATER_SELECTION_SUPPORT>;
-		FWaterVertexFactoryUserDataWrapperType* UserDataWrappers[3] = {};
-		for (EWaterMeshRenderGroupType RenderGroup : BatchRenderGroups)
-		{
-			FWaterVertexFactoryUserDataWrapperType& UserDataWrapper = Collector.AllocateOneFrameResource<FWaterVertexFactoryUserDataWrapperType>();
-			UserDataWrapper.UserData.RenderGroupType = RenderGroup;
-			UserDataWrapper.UserData.QuadTreePosition = FVector(WaterQuadTree.GetTileRegion().Min, WaterQuadTreeMinHeight);
-			UserDataWrapper.UserData.CaptureDepthRange = WaterQuadTreeMaxHeight - WaterQuadTreeMinHeight;
-			UserDataWrapper.UserData.IndirectInstanceData0 = IndirectDrawResources.InstanceData0->GetRHI();
-			UserDataWrapper.UserData.IndirectInstanceData1 = IndirectDrawResources.InstanceData1->GetRHI();
-			UserDataWrapper.UserData.IndirectInstanceData2 = IndirectDrawResources.InstanceData2->GetRHI();
-			UserDataWrapper.UserData.IndirectInstanceData3 = (WITH_WATER_SELECTION_SUPPORT != 0) ? IndirectDrawResources.InstanceData3->GetRHI() : nullptr;
-			UserDataWrapper.UserData.IndirectInstanceDataOffsetsSRV = IndirectDrawResources.InstanceDataOffsets->GetSRV(RHICmdList, FRHIBufferSRVCreateInfo(PF_R32_UINT));
-			UserDataWrapper.UserData.IndirectInstanceData0SRV = IndirectDrawResources.InstanceData0->GetSRV(RHICmdList, FRHIBufferSRVCreateInfo(PF_R32_UINT));
-			UserDataWrapper.UserData.IndirectInstanceData1SRV = IndirectDrawResources.InstanceData1->GetSRV(RHICmdList, FRHIBufferSRVCreateInfo(PF_R32_UINT));
-			UserDataWrapper.UserData.IndirectInstanceData2SRV = IndirectDrawResources.InstanceData2->GetSRV(RHICmdList, FRHIBufferSRVCreateInfo(PF_R32_UINT));
-			UserDataWrapper.UserData.IndirectInstanceData3SRV = (WITH_WATER_SELECTION_SUPPORT != 0) ? IndirectDrawResources.InstanceData3->GetSRV(RHICmdList, FRHIBufferSRVCreateInfo(PF_R32_UINT)) : nullptr;
-			
-			UserDataWrappers[(int32)RenderGroup] = &UserDataWrapper;
-		}
+		int32 InstanceDataOffset = 0;
 
 		// Go through all buckets and issue one batched draw call per LOD level per material per view
-		int32 CompactViewIndex = 0;
+		int32 TraversalIndex = 0;
 		for (int32 ViewIndex = 0; ViewIndex < Views.Num(); ViewIndex++)
 		{
 			// when rendering ISR, don't process the instanced view
-			if ((VisibilityMap & (1 << ViewIndex)) && (!bEncounteredISRView || Views[ViewIndex]->IsPrimarySceneView()))
+			if ((LocalViewMask & (1 << ViewIndex)) && (!bEncounteredISRView || Views[ViewIndex]->IsPrimarySceneView()))
 			{
 				TRACE_CPUPROFILER_EVENT_SCOPE(BucketsPerView);
 
-				const bool bInstancedStereo = Views[ViewIndex]->bIsInstancedStereoEnabled;
 				const bool bViewIsWireframe = WireframeVisibilityMap & (1 << ViewIndex);
 
+				FWaterQuadTree::FTraversalOutput& WaterInstanceData = WaterInstanceDataPerView[TraversalIndex];
+				TraversalIndex++;
+
 				for (int32 MaterialIndex = 0; MaterialIndex < NumWaterMaterials; ++MaterialIndex)
 				{
 					TRACE_CPUPROFILER_EVENT_SCOPE(MaterialBucket);
 					bool bMaterialDrawn = false;
 
-					// We only render all tiles with one or multiple instances of a single density quad mesh, so no need to create one bucket per density.
+					for (int32 DensityIndex = 0; DensityIndex < WaterQuadTreeConstants.DensityCount; ++DensityIndex)
 					{
-						const int32 BucketIndex = MaterialIndex;
+						const int32 BucketIndex = MaterialIndex * WaterQuadTreeConstants.DensityCount + DensityIndex;
+						const int32 InstanceCount = WaterInstanceData.BucketInstanceCounts[BucketIndex];
+
+						if (!InstanceCount)
+						{
+							continue;
+						}
 
 						TRACE_CPUPROFILER_EVENT_SCOPE(DensityBucket);
 
@@ -589,7 +620,7 @@ void FWaterMeshSceneProxy::GetDynamicMeshElements(const TArray<const FSceneView*
 							// Set up mesh batch
 							FMeshBatch& Mesh = Collector.AllocateMesh();
 							Mesh.bWireframe = bViewIsWireframe;
-							Mesh.VertexFactory = bInstancedStereo ? (FVertexFactory*)WaterVertexFactoryIndirectDrawISR : (FVertexFactory*)WaterVertexFactoryIndirectDraw;
+							Mesh.VertexFactory = WaterVertexFactories[DensityIndex];
 							Mesh.MaterialRenderProxy = MaterialRenderProxy;
 							Mesh.ReverseCulling = IsLocalToWorldDeterminantNegative();
 							Mesh.Type = PT_TriangleList;
@@ -607,21 +638,25 @@ void FWaterMeshSceneProxy::GetDynamicMeshElements(const TArray<const FSceneView*
 #endif // WITH_WATER_SELECTION_SUPPORT
 
 							Mesh.Elements.SetNumZeroed(1);
+
 							{
 								TRACE_CPUPROFILER_EVENT_SCOPE_STR("Setup batch element");
 
 								// Set up one mesh batch element
 								FMeshBatchElement& BatchElement = Mesh.Elements[0];
 
-								// Set up for indirect draw
+								// Set up for instancing
+								//BatchElement.bIsInstancedMesh = true;
+								BatchElement.NumInstances = InstanceCount;
+								BatchElement.UserData = (void*)ViewWaterQuadTree.GetWaterMeshUserDataBuffers()->GetUserData(RenderGroup);
+								BatchElement.UserIndex = InstanceDataOffset * InstanceFactor;
+
 								BatchElement.FirstIndex = 0;
-								BatchElement.NumPrimitives = 0; // Must be 0 to enable usage of IndirectArgsBuffer
-								BatchElement.IndirectArgsBuffer = IndirectDrawResources.IndirectArgs->GetRHI();
-								BatchElement.IndirectArgsOffset = (CompactViewIndex * NumBucketsIndirect + BucketIndex) * sizeof(FRHIDrawIndexedIndirectParameters);
-								BatchElement.UserData = (void*)&UserDataWrappers[(int32)RenderGroup]->UserData;
-								BatchElement.UserIndex = CompactViewIndex * NumBucketsIndirect + BucketIndex;
+								BatchElement.NumPrimitives = WaterVertexFactories[DensityIndex]->IndexBuffer->GetIndexCount() / 3;
+								BatchElement.MinVertexIndex = 0;
+								BatchElement.MaxVertexIndex = WaterVertexFactories[DensityIndex]->VertexBuffer->GetVertexCount() - 1;
 
-								BatchElement.IndexBuffer = bInstancedStereo ? WaterVertexFactoryIndirectDrawISR->IndexBuffer : WaterVertexFactoryIndirectDraw->IndexBuffer;
+								BatchElement.IndexBuffer = WaterVertexFactories[DensityIndex]->IndexBuffer;
 								BatchElement.PrimitiveIdMode = PrimID_ForceZero;
 
 								// We need the uniform buffer of this primitive because it stores the proper value for the bOutputVelocity flag.
@@ -630,258 +665,55 @@ void FWaterMeshSceneProxy::GetDynamicMeshElements(const TArray<const FSceneView*
 							}
 
 							{
+								INC_DWORD_STAT_BY(STAT_WaterVerticesDrawn, WaterVertexFactories[DensityIndex]->VertexBuffer->GetVertexCount() * InstanceCount);
 								INC_DWORD_STAT(STAT_WaterDrawCalls);
+								INC_DWORD_STAT_BY(STAT_WaterTilesDrawn, InstanceCount);
+
 								TRACE_CPUPROFILER_EVENT_SCOPE(Collector.AddMesh);
 
 								Collector.AddMesh(ViewIndex, Mesh);
 							}
 						}
+
+						// Note : we're repurposing the BucketInstanceCounts array here for storing the actual offset in the buffer. This means that effectively from this point on, BucketInstanceCounts doesn't actually 
+						// contain the number of instances anymore : 
+						WaterInstanceData.BucketInstanceCounts[BucketIndex] = InstanceDataOffset;
+						InstanceDataOffset += InstanceCount;
 					}
 
 					INC_DWORD_STAT_BY(STAT_WaterDrawnMats, (int32)bMaterialDrawn);
 				}
 
-				++CompactViewIndex;
-			}
-		}
-	}
-	
-	// Even if we have a GPU quadtree, we still go through the setup for the non-GPU path because the water quadtree might have an ocean far mesh which we render traditionally.
-	// If this is a GPU quadtree, traversing it on the CPU will only return the far mesh tiles (if present).
-
-	TArray<FWaterQuadTree::FTraversalOutput, TInlineAllocator<4>> WaterInstanceDataPerView;
-
-	bool bEncounteredISRView = false;
-	int32 InstanceFactor = 1;
-
-	// Gather visible tiles, their lod and materials for all renderable views (skip right view when stereo pair is rendered instanced)
-	for (int32 ViewIndex = 0; ViewIndex < Views.Num(); ViewIndex++)
-	{
-		const FSceneView* View = Views[ViewIndex];
-		if (!bEncounteredISRView && View->IsInstancedStereoPass())
-		{
-			bEncounteredISRView = true;
-			InstanceFactor = View->GetStereoPassInstanceFactor();
-		}
-
-		// skip gathering visible tiles from instanced right eye views
-		if ((VisibilityMap & (1 << ViewIndex)) && (!bEncounteredISRView || View->IsPrimarySceneView()))
-		{
-			const FVector ObserverPosition = View->ViewMatrices.GetViewOrigin();
-
-			FWaterLODParams WaterLODParams = GetWaterLODParams(ObserverPosition);
-
-#if !(UE_BUILD_SHIPPING || UE_BUILD_TEST)
-			if (CVarWaterMeshShowLODLevels.GetValueOnRenderThread())
-			{
-				for (int32 i = WaterLODParams.LowestLOD; i < WaterQuadTree.GetTreeDepth(); i++)
-				{
-					float LODDist = FWaterQuadTree::GetLODDistance(i, LODScale);
-					FVector Orig = FVector(FVector2D(ObserverPosition), WaterLODParams.WaterHeightForLOD);
-
-					DrawCircle(Collector.GetPDI(ViewIndex), Orig, FVector::ForwardVector, FVector::RightVector, GColorList.GetFColorByIndex(i + 1), LODDist, 64, 0);
-				}
-			}
-#endif
-			TRACE_CPUPROFILER_EVENT_SCOPE(QuadTreeTraversalPerView);
-
-			FWaterQuadTree::FTraversalOutput& WaterInstanceData = WaterInstanceDataPerView.Emplace_GetRef();
-			WaterInstanceData.BucketInstanceCounts.Empty(NumBuckets);
-			WaterInstanceData.BucketInstanceCounts.AddZeroed(NumBuckets);
-			if (!!CVarWaterMeshPreAllocStagingInstanceMemory.GetValueOnRenderThread())
-			{
-				WaterInstanceData.StagingInstanceData.Empty(HistoricalMaxViewInstanceCount);
-			}
-
-			FWaterQuadTree::FTraversalDesc TraversalDesc;
-			TraversalDesc.LowestLOD = WaterLODParams.LowestLOD;
-			TraversalDesc.HeightMorph = WaterLODParams.HeightLODFactor;
-			TraversalDesc.LODCount = WaterQuadTree.GetTreeDepth();
-			TraversalDesc.DensityCount = DensityCount;
-			TraversalDesc.ForceCollapseDensityLevel = ForceCollapseDensityLevel;
-			TraversalDesc.Frustum = View->ViewFrustum;
-			TraversalDesc.ObserverPosition = ObserverPosition;
-			TraversalDesc.PreViewTranslation = View->ViewMatrices.GetPreViewTranslation();
-			TraversalDesc.LODScale = LODScale;
-			TraversalDesc.bLODMorphingEnabled = !!CVarWaterMeshLODMorphEnabled.GetValueOnRenderThread();
-			TraversalDesc.WaterInfoBounds = WaterInfoBounds;
-			TraversalDesc.OcclusionCullingResults = nullptr;
-			TraversalDesc.OcclusionCullingFarMeshOffset = OcclusionResultsFarMeshOffset;
-			if (CVarWaterMeshOcclusionCulling.GetValueOnRenderThread() != 0)
-			{
-				const FOcclusionCullingResults* CullingResults = OcclusionResults.Find(View->GetViewKey());
-				if (CullingResults && !CullingResults->Results.IsEmpty())
-				{
-					TraversalDesc.OcclusionCullingResults = &CullingResults->Results;
-				}
-			}
-
-#if !(UE_BUILD_SHIPPING || UE_BUILD_TEST)
-			//Debug
-			TraversalDesc.DebugPDI = Collector.GetPDI(ViewIndex);
-			TraversalDesc.DebugShowTile = CVarWaterMeshShowTileBounds.GetValueOnRenderThread();
-			TraversalDesc.bDebugDrawIntoForeground = CVarWaterMeshShowTileBoundsForeground.GetValueOnRenderThread() != 0;
-#endif
-			WaterQuadTree.BuildWaterTileInstanceData(TraversalDesc, WaterInstanceData);
-
-			HistoricalMaxViewInstanceCount = FMath::Max(HistoricalMaxViewInstanceCount, WaterInstanceData.InstanceCount);
-		}
-	}
-
-	// Get number of total instances for all views
-	int32 TotalInstanceCount = 0;
-	for (const FWaterQuadTree::FTraversalOutput& WaterInstanceData : WaterInstanceDataPerView)
-	{
-		TotalInstanceCount += WaterInstanceData.InstanceCount;
-	}
-
-	if (TotalInstanceCount == 0)
-	{
-		// no instance visible, early exit
-		return;
-	}
-
-	WaterInstanceDataBuffers->Lock(RHICmdList, TotalInstanceCount * InstanceFactor);
-
-	int32 InstanceDataOffset = 0;
-
-	// Go through all buckets and issue one batched draw call per LOD level per material per view
-	int32 TraversalIndex = 0;
-	for (int32 ViewIndex = 0; ViewIndex < Views.Num(); ViewIndex++)
-	{
-		// when rendering ISR, don't process the instanced view
-		if ((VisibilityMap & (1 << ViewIndex)) && (!bEncounteredISRView || Views[ViewIndex]->IsPrimarySceneView()))
-		{
-			TRACE_CPUPROFILER_EVENT_SCOPE(BucketsPerView);
-
-			const bool bViewIsWireframe = WireframeVisibilityMap & (1 << ViewIndex);
-
-			FWaterQuadTree::FTraversalOutput& WaterInstanceData = WaterInstanceDataPerView[TraversalIndex];
-			TraversalIndex++;
-
-			for (int32 MaterialIndex = 0; MaterialIndex < NumWaterMaterials; ++MaterialIndex)
-			{
-				TRACE_CPUPROFILER_EVENT_SCOPE(MaterialBucket);
-				bool bMaterialDrawn = false;
-
-				for (int32 DensityIndex = 0; DensityIndex < DensityCount; ++DensityIndex)
+				const int32 NumStagingInstances = WaterInstanceData.StagingInstanceData.Num();
+				for (int32 Idx = 0; Idx < NumStagingInstances; ++Idx)
 				{
-					const int32 BucketIndex = MaterialIndex * DensityCount + DensityIndex;
-					const int32 InstanceCount = WaterInstanceData.BucketInstanceCounts[BucketIndex];
-
-					if (!InstanceCount)
-					{
-						continue;
-					}
+					const FWaterQuadTree::FStagingInstanceData& Data = WaterInstanceData.StagingInstanceData[Idx];
+					const int32 WriteIndex = WaterInstanceData.BucketInstanceCounts[Data.BucketIndex]++;
 
-					TRACE_CPUPROFILER_EVENT_SCOPE(DensityBucket);
-
-					const FMaterialRenderProxy* MaterialRenderProxy = bViewIsWireframe && WireframeMaterialInstance ? WireframeMaterialInstance : WaterQuadTree.GetWaterMaterials()[MaterialIndex];
-					check(MaterialRenderProxy != nullptr);
-
-					bool bUseForDepthPass = false;
-
-					// If there's a valid material, use that to figure out the depth pass status
-					if (const FMaterial* BucketMaterial = MaterialRenderProxy->GetMaterialNoFallback(GetScene().GetFeatureLevel()))
-					{
-						// Preemptively turn off depth rendering for this mesh batch if the material doesn't need it
-						bUseForDepthPass = !BucketMaterial->GetShadingModels().HasShadingModel(MSM_SingleLayerWater) && !IsTranslucentOnlyBlendMode(*BucketMaterial);
-					}
-
-					bMaterialDrawn = true;
-					for (EWaterMeshRenderGroupType RenderGroup : BatchRenderGroups)
+					for (int32 StreamIdx = 0; StreamIdx < FWaterInstanceDataBuffersType::NumBuffers; ++StreamIdx)
 					{
-						// Set up mesh batch
-						FMeshBatch& Mesh = Collector.AllocateMesh();
-						Mesh.bWireframe = bViewIsWireframe;
-						Mesh.VertexFactory = WaterVertexFactories[DensityIndex];
-						Mesh.MaterialRenderProxy = MaterialRenderProxy;
-						Mesh.ReverseCulling = IsLocalToWorldDeterminantNegative();
-						Mesh.Type = PT_TriangleList;
-						Mesh.DepthPriorityGroup = SDPG_World;
-						Mesh.bCanApplyViewModeOverrides = false;
-						Mesh.bUseForMaterial = true;
-						Mesh.CastShadow = false;
-						// Preemptively turn off depth rendering for this mesh batch if the material doesn't need it
-						Mesh.bUseForDepthPass = bUseForDepthPass;
-						Mesh.bUseAsOccluder = false;
-
-#if WITH_WATER_SELECTION_SUPPORT
-						Mesh.bUseSelectionOutline = (RenderGroup == EWaterMeshRenderGroupType::RG_RenderSelectedWaterTilesOnly);
-						Mesh.bUseWireframeSelectionColoring = (RenderGroup == EWaterMeshRenderGroupType::RG_RenderSelectedWaterTilesOnly);
-#endif // WITH_WATER_SELECTION_SUPPORT
-
-						Mesh.Elements.SetNumZeroed(1);
-
-						{
-							TRACE_CPUPROFILER_EVENT_SCOPE_STR("Setup batch element");
-
-							// Set up one mesh batch element
-							FMeshBatchElement& BatchElement = Mesh.Elements[0];
-
-							// Set up for instancing
-							//BatchElement.bIsInstancedMesh = true;
-							BatchElement.NumInstances = InstanceCount;
-							BatchElement.UserData = (void*)WaterMeshUserDataBuffers->GetUserData(RenderGroup);
-							BatchElement.UserIndex = InstanceDataOffset * InstanceFactor;
-
-							BatchElement.FirstIndex = 0;
-							BatchElement.NumPrimitives = WaterVertexFactories[DensityIndex]->IndexBuffer->GetIndexCount() / 3;
-							BatchElement.MinVertexIndex = 0;
-							BatchElement.MaxVertexIndex = WaterVertexFactories[DensityIndex]->VertexBuffer->GetVertexCount() - 1;
-
-							BatchElement.IndexBuffer = WaterVertexFactories[DensityIndex]->IndexBuffer;
-							BatchElement.PrimitiveIdMode = PrimID_ForceZero;
-
-							// We need the uniform buffer of this primitive because it stores the proper value for the bOutputVelocity flag.
-							// The identity primitive uniform buffer simply stores false for this flag which leads to missing motion vectors.
-							BatchElement.PrimitiveUniformBuffer = GetUniformBuffer();
-						}
-
+						TArrayView<FVector4f> BufferMemory = ViewWaterQuadTree.GetWaterInstanceDataBuffers()->GetBufferMemory(StreamIdx);
+						for (int32 IdxMultipliedInstance = 0; IdxMultipliedInstance < InstanceFactor; ++IdxMultipliedInstance)
 						{
-							INC_DWORD_STAT_BY(STAT_WaterVerticesDrawn, WaterVertexFactories[DensityIndex]->VertexBuffer->GetVertexCount() * InstanceCount);
-							INC_DWORD_STAT(STAT_WaterDrawCalls);
-							INC_DWORD_STAT_BY(STAT_WaterTilesDrawn, InstanceCount);
-
-							TRACE_CPUPROFILER_EVENT_SCOPE(Collector.AddMesh);
-
-							Collector.AddMesh(ViewIndex, Mesh);
+							BufferMemory[WriteIndex * InstanceFactor + IdxMultipliedInstance] = Data.Data[StreamIdx];
 						}
 					}
-
-					// Note : we're repurposing the BucketInstanceCounts array here for storing the actual offset in the buffer. This means that effectively from this point on, BucketInstanceCounts doesn't actually 
-					//  contain the number of instances anymore : 
-					WaterInstanceData.BucketInstanceCounts[BucketIndex] = InstanceDataOffset;
-					InstanceDataOffset += InstanceCount;
-				}
-
-				INC_DWORD_STAT_BY(STAT_WaterDrawnMats, (int32)bMaterialDrawn);
-			}
-
-			const int32 NumStagingInstances = WaterInstanceData.StagingInstanceData.Num();
-			for (int32 Idx = 0; Idx < NumStagingInstances; ++Idx)
-			{
-				const FWaterQuadTree::FStagingInstanceData& Data = WaterInstanceData.StagingInstanceData[Idx];
-				const int32 WriteIndex = WaterInstanceData.BucketInstanceCounts[Data.BucketIndex]++;
-
-				for (int32 StreamIdx = 0; StreamIdx < FWaterInstanceDataBuffersType::NumBuffers; ++StreamIdx)
-				{
-					TArrayView<FVector4f> BufferMemory = WaterInstanceDataBuffers->GetBufferMemory(StreamIdx);
-					for (int32 IdxMultipliedInstance = 0; IdxMultipliedInstance < InstanceFactor; ++IdxMultipliedInstance)
-					{
-						BufferMemory[WriteIndex * InstanceFactor + IdxMultipliedInstance] = Data.Data[StreamIdx];
-					}
 				}
 			}
 		}
-	}
 
-	WaterInstanceDataBuffers->Unlock(RHICmdList);
+		ViewWaterQuadTree.GetWaterInstanceDataBuffers()->Unlock(RHICmdList);
+	}
 }
 
 const TArray<FBoxSphereBounds>* FWaterMeshSceneProxy::GetOcclusionQueries(const FSceneView* View) const
 {
-	if (CVarWaterMeshOcclusionCulling.GetValueOnRenderThread() != 0)
+	// Local tessellation/multiple dynamic quadtrees are incompatible with CPU occlusion culling
+	if (
+#if !FORCE_SINGLE_WATER_QUADTREE
+		!bIsLocalOnlyTessellationEnabled && 
+#endif
+		CVarWaterMeshOcclusionCulling.GetValueOnRenderThread() != 0)
 	{
 		return &OcclusionCullingBounds;
 	}
@@ -934,6 +766,128 @@ void FWaterMeshSceneProxy::AcceptOcclusionResults(const FSceneView* View, TArray
 	}
 }
 
+bool FWaterMeshSceneProxy::HasWaterData() const
+{
+	for (const auto& Pair : ViewQuadTrees)
+	{
+		if (Pair.Value.GetWaterQuadTree().GetNodeCount() != 0 && WaterQuadTreeConstants.DensityCount != 0)
+		{
+			return true;
+		}
+	}
+	return false;
+}
+
+TArray<EWaterMeshRenderGroupType, TInlineAllocator<FWaterMeshSceneProxy::FWaterVertexFactoryType::NumRenderGroups>> FWaterMeshSceneProxy::GetBatchRenderGroups(const TArray<const FSceneView*>& Views, uint32 VisibilityMap) const
+{
+	TArray<EWaterMeshRenderGroupType, TInlineAllocator<FWaterVertexFactoryType::NumRenderGroups>> BatchRenderGroups;
+	{
+		// By default, render all water tiles : 
+		BatchRenderGroups.Add(EWaterMeshRenderGroupType::RG_RenderWaterTiles);
+
+#if WITH_WATER_SELECTION_SUPPORT
+		bool bHasSelectedInstances = IsSelected();
+
+		// We're using bSinglePassGDME, so we need to check the ViewFamily of all views and can't use the ViewFamily passed into GDME.
+		bool bSelectionRenderEnabled = false;
+		for (int32 ViewIndex = 0; ViewIndex < Views.Num(); ViewIndex++)
+		{
+			if (VisibilityMap & (1 << ViewIndex))
+			{
+				bSelectionRenderEnabled |= GIsEditor && Views[ViewIndex]->Family->EngineShowFlags.Selection;
+			}
+		}
+
+		if (bSelectionRenderEnabled && bHasSelectedInstances)
+		{
+			// Don't render all in one group: instead, render 2 groups : first, the selected only then, the non-selected only :
+			BatchRenderGroups[0] = EWaterMeshRenderGroupType::RG_RenderSelectedWaterTilesOnly;
+			BatchRenderGroups.Add(EWaterMeshRenderGroupType::RG_RenderUnselectedWaterTilesOnly);
+		}
+#endif // WITH_WATER_SELECTION_SUPPORT
+	}
+
+	return BatchRenderGroups;
+}
+
+uint32 FWaterMeshSceneProxy::GetWireframeVisibilityMapAndMaterial(const TArray<const FSceneView*>& Views, uint32 VisibilityMap, FMeshElementCollector& Collector, FColoredMaterialRenderProxy*& OutMaterialInstance) const
+{
+	// Figure out which views have wireframe enabled
+	uint32 WireframeVisibilityMap = 0;
+	if (AllowDebugViewmodes())
+	{
+		if (CVarWaterMeshShowWireframe.GetValueOnRenderThread() == 1)
+		{
+			// Force all views to wireframe
+			WireframeVisibilityMap = VisibilityMap;
+		}
+		else
+		{
+			for (int32 ViewIndex = 0; ViewIndex < Views.Num(); ViewIndex++)
+			{
+				if ((VisibilityMap & (1 << ViewIndex)) && Views[ViewIndex]->Family->EngineShowFlags.Wireframe)
+				{
+					WireframeVisibilityMap |= (1 << ViewIndex);
+				}
+			}
+		}
+	}
+
+	// Set up wireframe material (if needed)
+	OutMaterialInstance = nullptr;
+	if (WireframeVisibilityMap && CVarWaterMeshShowWireframeAtBaseHeight.GetValueOnRenderThread() == 1)
+	{
+		OutMaterialInstance = new FColoredMaterialRenderProxy(
+			GEngine->WireframeMaterial ? GEngine->WireframeMaterial->GetRenderProxy() : NULL,
+			FColor::Cyan);
+
+		Collector.RegisterOneFrameMaterialProxy(OutMaterialInstance);
+	}
+
+	return WireframeVisibilityMap;
+}
+
+int32 FWaterMeshSceneProxy::FindBestQuadTreeForView(const FSceneView* View) const
+{
+	// Do we have a quadtree matching this views PlayerIndex?
+	if (ViewQuadTrees.Contains(View->PlayerIndex))
+	{
+		return View->PlayerIndex;
+	}
+	// Otherwise just use the closest quadtree.
+	else
+	{
+		const FVector2D ViewPosition2D = FVector2D(View->ViewLocation);
+		double ClosestDistSquared = DBL_MAX;
+		int32 ClosestQuadTreeKey = INDEX_NONE;
+		for (const auto& Pair : ViewQuadTrees)
+		{
+			const double DistSquared = FVector2D::DistSquared(Pair.Value.GetWaterQuadTree().GetTileRegion().GetCenter(), ViewPosition2D);
+			if (DistSquared < ClosestDistSquared)
+			{
+				ClosestDistSquared = DistSquared;
+				ClosestQuadTreeKey = Pair.Key;
+			}
+		}
+		return ClosestQuadTreeKey;
+	}
+}
+
+TArray<int32, TInlineAllocator<8>> FWaterMeshSceneProxy::GetViewToQuadTreeMapping(const TArray<const FSceneView*>& Views, uint32 VisibilityMap) const
+{
+	// Find a quadtree for every view
+	TArray<int32, TInlineAllocator<8>> ViewToQuadTreeKey;
+	ViewToQuadTreeKey.SetNum(Views.Num());
+	for (int32 ViewIndex = 0; ViewIndex < Views.Num(); ViewIndex++)
+	{
+		if (VisibilityMap & (1 << ViewIndex))
+		{
+			ViewToQuadTreeKey[ViewIndex] = FindBestQuadTreeForView(Views[ViewIndex]);
+		}
+	}
+	return ViewToQuadTreeKey;
+}
+
 #if RHI_RAYTRACING
 void FWaterMeshSceneProxy::SetupRayTracingInstances(FRHICommandListBase& RHICmdList, int32 NumInstances, uint32 DensityIndex)
 {
@@ -985,9 +939,14 @@ void FWaterMeshSceneProxy::GetDynamicRayTracingInstances(FRayTracingMaterialGath
 	const FSceneView& SceneView = *Context.ReferenceView;
 	const FVector ObserverPosition = SceneView.ViewMatrices.GetViewOrigin();
 
-	FWaterLODParams WaterLODParams = GetWaterLODParams(ObserverPosition);
+	const int32 QuadTreeKey = FindBestQuadTreeForView(&SceneView);
+	const FViewWaterQuadTree* ViewWaterQuadTree = ViewQuadTrees.Find(QuadTreeKey);
+	check(ViewWaterQuadTree);
+	const FWaterQuadTree& WaterQuadTree = ViewWaterQuadTree->GetWaterQuadTree();
+
+	FViewWaterQuadTree::FWaterLODParams WaterLODParams = ViewWaterQuadTree->GetWaterLODParams(ObserverPosition, WaterQuadTreeConstants.LODScale);
 
-	const int32 NumBuckets = WaterQuadTree.GetWaterMaterials().Num() * DensityCount;
+	const int32 NumBuckets = WaterQuadTree.GetWaterMaterials().Num() * WaterQuadTreeConstants.DensityCount;
 
 	FWaterQuadTree::FTraversalOutput WaterInstanceData;
 	WaterInstanceData.BucketInstanceCounts.Empty(NumBuckets);
@@ -997,14 +956,14 @@ void FWaterMeshSceneProxy::GetDynamicRayTracingInstances(FRayTracingMaterialGath
 	TraversalDesc.LowestLOD = WaterLODParams.LowestLOD;
 	TraversalDesc.HeightMorph = WaterLODParams.HeightLODFactor;
 	TraversalDesc.LODCount = WaterQuadTree.GetTreeDepth();
-	TraversalDesc.DensityCount = DensityCount;
-	TraversalDesc.ForceCollapseDensityLevel = ForceCollapseDensityLevel;
+	TraversalDesc.DensityCount = WaterQuadTreeConstants.DensityCount;
+	TraversalDesc.ForceCollapseDensityLevel = WaterQuadTreeConstants.ForceCollapseDensityLevel;
 	TraversalDesc.PreViewTranslation = SceneView.ViewMatrices.GetPreViewTranslation();
 	TraversalDesc.ObserverPosition = ObserverPosition;
 	TraversalDesc.Frustum = FConvexVolume(); // Default volume to disable frustum culling
-	TraversalDesc.LODScale = LODScale;
+	TraversalDesc.LODScale = WaterQuadTreeConstants.LODScale;
 	TraversalDesc.bLODMorphingEnabled = !!CVarWaterMeshLODMorphEnabled.GetValueOnRenderThread();
-	TraversalDesc.WaterInfoBounds = WaterInfoBounds;
+	TraversalDesc.WaterInfoBounds = WaterQuadTree.GetTileRegion();
 
 	WaterQuadTree.BuildWaterTileInstanceData(TraversalDesc, WaterInstanceData);
 
@@ -1016,12 +975,12 @@ void FWaterMeshSceneProxy::GetDynamicRayTracingInstances(FRayTracingMaterialGath
 
 	const int32 NumWaterMaterials = WaterQuadTree.GetWaterMaterials().Num();	
 
-	for (int32 DensityIndex = 0; DensityIndex < DensityCount; ++DensityIndex)
+	for (int32 DensityIndex = 0; DensityIndex < WaterQuadTreeConstants.DensityCount; ++DensityIndex)
 	{
 		int32 DensityInstanceCount = 0;
 		for (int32 MaterialIndex = 0; MaterialIndex < NumWaterMaterials; ++MaterialIndex)
 		{
-			const int32 BucketIndex = MaterialIndex * DensityCount + DensityIndex;
+			const int32 BucketIndex = MaterialIndex * WaterQuadTreeConstants.DensityCount + DensityIndex;
 			const int32 InstanceCount = WaterInstanceData.BucketInstanceCounts[BucketIndex];
 			DensityInstanceCount += InstanceCount;
 		}
@@ -1052,7 +1011,7 @@ void FWaterMeshSceneProxy::GetDynamicRayTracingInstances(FRayTracingMaterialGath
 	BaseMesh.SegmentIndex = 0;
 	BaseMesh.Elements.AddZeroed();
 
-	for (int32 DensityIndex = 0; DensityIndex < DensityCount; ++DensityIndex)
+	for (int32 DensityIndex = 0; DensityIndex < WaterQuadTreeConstants.DensityCount; ++DensityIndex)
 	{
 		int32 DensityInstanceIndex = 0;
 		
@@ -1074,7 +1033,7 @@ void FWaterMeshSceneProxy::GetDynamicRayTracingInstances(FRayTracingMaterialGath
 
 		for (int32 MaterialIndex = 0; MaterialIndex < NumWaterMaterials; ++MaterialIndex)
 		{
-			const int32 BucketIndex = MaterialIndex * DensityCount + DensityIndex;
+			const int32 BucketIndex = MaterialIndex * WaterQuadTreeConstants.DensityCount + DensityIndex;
 			const int32 InstanceCount = WaterInstanceData.BucketInstanceCounts[BucketIndex];
 
 			if (!InstanceCount)
@@ -1100,7 +1059,7 @@ void FWaterMeshSceneProxy::GetDynamicRayTracingInstances(FRayTracingMaterialGath
 				UniformBufferParams.InstanceData0 = InstanceData.Data[0];
 				UniformBufferParams.InstanceData1 = InstanceData.Data[1];
 
-				UserDataWrapper.UserData.InstanceDataBuffers = WaterMeshUserDataBuffers->GetUserData(EWaterMeshRenderGroupType::RG_RenderWaterTiles)->InstanceDataBuffers;
+				UserDataWrapper.UserData.InstanceDataBuffers = ViewWaterQuadTree->GetWaterMeshUserDataBuffers()->GetUserData(EWaterMeshRenderGroupType::RG_RenderWaterTiles)->InstanceDataBuffers;
 				UserDataWrapper.UserData.RenderGroupType = EWaterMeshRenderGroupType::RG_RenderWaterTiles;
 				UserDataWrapper.UserData.WaterVertexFactoryRaytracingVFUniformBuffer = FWaterVertexFactoryRaytracingParametersRef::CreateUniformBufferImmediate(UniformBufferParams, UniformBuffer_SingleFrame);
 							
@@ -1133,6 +1092,47 @@ void FWaterMeshSceneProxy::GetDynamicRayTracingInstances(FRayTracingMaterialGath
 }
 #endif // RHI_RAYTRACING
 
+bool FWaterMeshSceneProxy::CreateViewWaterQuadTree(int32 Key, const FVector2D& CenterPosition)
+{
+#if !FORCE_SINGLE_WATER_QUADTREE
+	if (bIsLocalOnlyTessellationEnabled)
+#endif
+	{
+		if (!ViewQuadTrees.Contains(Key))
+		{
+			FViewWaterQuadTree& ViewWaterQuadTree = ViewQuadTrees.Add(Key);
+			ViewWaterQuadTree.Update(WaterQuadTreeBuilder, CenterPosition);
+			return true;
+		}
+	}
+	return false;
+}
+
+bool FWaterMeshSceneProxy::UpdateViewWaterQuadTree(int32 Key, const FVector2D& CenterPosition)
+{
+#if !FORCE_SINGLE_WATER_QUADTREE
+	if (bIsLocalOnlyTessellationEnabled)
+#endif
+	{
+		if (FViewWaterQuadTree* ViewWaterQuadTree = ViewQuadTrees.Find(Key))
+		{
+			ViewWaterQuadTree->Update(WaterQuadTreeBuilder, CenterPosition);
+			return true;
+		}
+	}
+	return false;
+}
+
+void FWaterMeshSceneProxy::DestroyViewWaterQuadTree(int32 Key)
+{
+#if !FORCE_SINGLE_WATER_QUADTREE
+	if (bIsLocalOnlyTessellationEnabled)
+#endif
+	{
+		ViewQuadTrees.Remove(Key);
+	}
+}
+
 FPrimitiveViewRelevance FWaterMeshSceneProxy::GetViewRelevance(const FSceneView* View) const
 {
 	FPrimitiveViewRelevance Result;
@@ -1157,14 +1157,190 @@ bool FWaterMeshSceneProxy::HasSubprimitiveOcclusionQueries() const
 #if WITH_WATER_SELECTION_SUPPORT
 HHitProxy* FWaterMeshSceneProxy::CreateHitProxies(UPrimitiveComponent* Component, TArray<TRefCountPtr<HHitProxy> >& OutHitProxies)
 {
-	WaterQuadTree.GatherHitProxies(OutHitProxies);
+	WaterQuadTreeBuilder.GatherHitProxies(OutHitProxies);
 
 	// No default hit proxy.
 	return nullptr;
 }
 #endif // WITH_WATER_SELECTION_SUPPORT
 
-void FWaterMeshSceneProxy::BuildGPUQuadTree(FRDGBuilder& GraphBuilder)
+void FViewWaterQuadTree::Update(const FWaterQuadTreeBuilder& Builder, const FVector2D& CenterPosition)
+{
+	Builder.BuildWaterQuadTree(WaterQuadTree, CenterPosition);
+	QuadTreeGPU = {};
+
+	// Initialize Z bounds needed for GPU driven rendering
+	if (Builder.IsGPUQuadTree())
+	{
+		WaterQuadTreeMinHeight = DBL_MAX;
+		WaterQuadTreeMaxHeight = -DBL_MAX;
+
+		const TArrayView<const FWaterBodyRenderData> WaterBodyRenderData = WaterQuadTree.GetWaterBodyRenderData();
+		for (const FWaterBodyRenderData& RenderData : WaterBodyRenderData)
+		{
+			WaterQuadTreeMinHeight = FMath::Min(WaterQuadTreeMinHeight, RenderData.BoundsMinZ - RenderData.MaxWaveHeight);
+			WaterQuadTreeMaxHeight = FMath::Max(WaterQuadTreeMaxHeight, RenderData.BoundsMaxZ + RenderData.MaxWaveHeight);
+		}
+
+		if (WaterQuadTreeMinHeight == DBL_MAX)
+		{
+			WaterQuadTreeMinHeight = 0.0;
+			WaterQuadTreeMaxHeight = 0.0;
+		}
+		WaterQuadTreeMinHeight -= 1.0;
+		WaterQuadTreeMaxHeight += 1.0;
+	}
+
+	if (!WaterInstanceDataBuffers || !WaterMeshUserDataBuffers)
+	{
+		const int32 TotalLeafNodes = Builder.GetMaxLeafCount();
+		WaterInstanceDataBuffers = MakeUnique<FWaterInstanceDataBuffersType>(TotalLeafNodes);
+		WaterMeshUserDataBuffers = MakeUnique<FWaterMeshUserDataBuffersType>(WaterInstanceDataBuffers.Get());
+	}
+}
+
+void FViewWaterQuadTree::TraverseGPUQuadTree(FRDGBuilder& GraphBuilder, bool bDepthBufferIsPopulated)
+{
+	if (bNeedToTraverseGPUQuadTree)
+	{
+		if (!QuadTreeGPU.IsInitialized())
+		{
+			BuildGPUQuadTree(GraphBuilder);
+		}
+		FWaterQuadTreeGPU::FTraverseParams TraverseParams = MoveTemp(WaterQuadTreeGPUTraverseParams);
+		TraverseParams.bDepthBufferIsPopulated = bDepthBufferIsPopulated;
+		QuadTreeGPU.Traverse(GraphBuilder, TraverseParams);
+		bNeedToTraverseGPUQuadTree = false;
+		WaterQuadTreeGPUTraverseParams = {};
+	}
+}
+
+FViewWaterQuadTree::FUserDataAndIndirectArgs FViewWaterQuadTree::PrepareGPUQuadTreeForRendering(const TArray<const FSceneView*>& Views, uint32 VisibilityMap, FMeshElementCollector& Collector, const FWaterQuadTreeConstants& QuadTreeConstants, const TArrayView<const EWaterMeshRenderGroupType>& BatchRenderGroups, FRHICommandListBase& RHICmdList) const
+{
+	bool bEncounteredISRView = false;
+	TArray<const FSceneView*> VisibleViews;
+
+	// Gather per view data for all renderable views
+	for (int32 ViewIndex = 0; ViewIndex < Views.Num(); ViewIndex++)
+	{
+		const FSceneView* View = Views[ViewIndex];
+		if (!bEncounteredISRView && View->IsInstancedStereoPass())
+		{
+			bEncounteredISRView = true;
+		}
+
+		// skip gathering visible tiles from instanced right eye views
+		if ((VisibilityMap & (1 << ViewIndex)) && (!bEncounteredISRView || View->IsPrimarySceneView()))
+		{
+			VisibleViews.Add(View);
+		}
+	}
+
+	const int32 NumWaterMaterials = WaterQuadTree.GetWaterMaterials().Num();
+	const int32 NumBucketsIndirect = NumWaterMaterials; // Indirect draws use a single density mesh tile
+	const int32 NumVisibleViews = VisibleViews.Num();
+	const int32 NumIndirectDrawCalls = NumBucketsIndirect * NumVisibleViews;
+	const int32 LeafCountUpperBound = WaterQuadTree.GetMaxLeafCount() * NumVisibleViews;
+	const int32 NumInstanceDataSlots = FMath::Max(1, static_cast<int32>(LeafCountUpperBound * FMath::Clamp(CVarWaterMeshGPUQuadInstanceDataAllocMult.GetValueOnRenderThread(), 0.0f, 8.0f)));
+
+	// Allocate buffers for the indirect draws
+	struct FIndirectDrawResources
+	{
+		TRefCountPtr<FRDGPooledBuffer> IndirectArgs;
+		TRefCountPtr<FRDGPooledBuffer> InstanceDataOffsets;
+		TRefCountPtr<FRDGPooledBuffer> InstanceData0;
+		TRefCountPtr<FRDGPooledBuffer> InstanceData1;
+		TRefCountPtr<FRDGPooledBuffer> InstanceData2;
+		TRefCountPtr<FRDGPooledBuffer> InstanceData3;
+	};
+	FIndirectDrawResources& IndirectDrawResources = Collector.AllocateOneFrameResource<FIndirectDrawResources>();
+	IndirectDrawResources.IndirectArgs = AllocatePooledBuffer(FRDGBufferDesc::CreateIndirectDesc<FRHIDrawIndexedIndirectParameters>(FMath::Max(1, NumIndirectDrawCalls)), TEXT("WaterQuadTree.IndirectArgsBuffer"));
+	IndirectDrawResources.InstanceDataOffsets = AllocatePooledBuffer(FRDGBufferDesc::CreateBufferDesc(sizeof(uint32), FMath::Max(1, NumIndirectDrawCalls)), TEXT("WaterQuadTree.InstanceDataOffsetsBuffer"));
+	IndirectDrawResources.InstanceData0 = AllocatePooledBuffer(FRDGBufferDesc::CreateBufferDesc(sizeof(uint32), FMath::Max(1, NumInstanceDataSlots)), TEXT("WaterQuadTree.InstanceDataBuffer0"));
+	IndirectDrawResources.InstanceData1 = AllocatePooledBuffer(FRDGBufferDesc::CreateBufferDesc(sizeof(uint32), FMath::Max(1, NumInstanceDataSlots)), TEXT("WaterQuadTree.InstanceDataBuffer1"));
+	IndirectDrawResources.InstanceData2 = AllocatePooledBuffer(FRDGBufferDesc::CreateBufferDesc(sizeof(uint32), FMath::Max(1, NumInstanceDataSlots)), TEXT("WaterQuadTree.InstanceDataBuffer2"));
+	if (WITH_WATER_SELECTION_SUPPORT != 0)
+	{
+		IndirectDrawResources.InstanceData3 = AllocatePooledBuffer(FRDGBufferDesc::CreateBufferDesc(sizeof(uint32), FMath::Max(1, NumInstanceDataSlots)), TEXT("WaterQuadTree.InstanceDataBuffer3"));
+	}
+
+
+	// Traverse quadtree. We actually only set up the parameters here; the actual traversal is done later in the frame in a callback invoked from FWaterViewExtension::PreRenderBasePass_RenderThread.
+	// That callback accesses these params, so we need to ensure that they are still valid at that point.
+	{
+		WaterQuadTreeGPUTraverseParams = {};
+		WaterQuadTreeGPUTraverseParams.OutIndirectArgsBuffer = IndirectDrawResources.IndirectArgs;
+		WaterQuadTreeGPUTraverseParams.OutInstanceDataOffsetsBuffer = IndirectDrawResources.InstanceDataOffsets;
+		WaterQuadTreeGPUTraverseParams.OutInstanceData0Buffer = IndirectDrawResources.InstanceData0;
+		WaterQuadTreeGPUTraverseParams.OutInstanceData1Buffer = IndirectDrawResources.InstanceData1;
+		WaterQuadTreeGPUTraverseParams.OutInstanceData2Buffer = IndirectDrawResources.InstanceData2;
+		WaterQuadTreeGPUTraverseParams.OutInstanceData3Buffer = IndirectDrawResources.InstanceData3;
+		WaterQuadTreeGPUTraverseParams.Views = MoveTemp(VisibleViews);
+		WaterQuadTreeGPUTraverseParams.QuadTreePosition = FVector(WaterQuadTree.GetTileRegion().Min, WaterQuadTreeMinHeight);
+		WaterQuadTreeGPUTraverseParams.CullingBounds = WaterQuadTree.GetTileRegion().ShiftBy(-WaterQuadTree.GetTileRegion().Min);
+		WaterQuadTreeGPUTraverseParams.NumDensities = QuadTreeConstants.DensityCount;
+		WaterQuadTreeGPUTraverseParams.NumMaterials = WaterQuadTree.GetWaterMaterials().Num();
+		WaterQuadTreeGPUTraverseParams.NumQuadsLOD0 = QuadTreeConstants.NumQuadsLOD0;
+		WaterQuadTreeGPUTraverseParams.NumQuadsPerTileSide = QuadTreeConstants.NumQuadsPerIndirectDrawTile;
+		WaterQuadTreeGPUTraverseParams.ForceCollapseDensityLevel = QuadTreeConstants.ForceCollapseDensityLevel;
+		WaterQuadTreeGPUTraverseParams.LeafSize = WaterQuadTree.GetLeafSize();
+		WaterQuadTreeGPUTraverseParams.LODScale = QuadTreeConstants.LODScale;
+		WaterQuadTreeGPUTraverseParams.DebugShowTile = CVarWaterMeshShowTileBounds.GetValueOnRenderThread();
+		WaterQuadTreeGPUTraverseParams.bWithWaterSelectionSupport = WITH_WATER_SELECTION_SUPPORT != 0;
+		WaterQuadTreeGPUTraverseParams.bLODMorphingEnabled = !!CVarWaterMeshLODMorphEnabled.GetValueOnRenderThread();
+
+		bNeedToTraverseGPUQuadTree = true;
+	}
+
+	FUserDataAndIndirectArgs Result;
+	Result.IndirectArgs = IndirectDrawResources.IndirectArgs;
+
+	using FWaterVertexFactoryUserDataWrapperType = TWaterVertexFactoryUserDataWrapper<WITH_WATER_SELECTION_SUPPORT>;
+	FWaterVertexFactoryUserDataWrapperType* UserDataWrappers[3] = {};
+	for (EWaterMeshRenderGroupType RenderGroup : BatchRenderGroups)
+	{
+		FWaterVertexFactoryUserDataWrapperType& UserDataWrapper = Collector.AllocateOneFrameResource<FWaterVertexFactoryUserDataWrapperType>();
+		UserDataWrapper.UserData.RenderGroupType = RenderGroup;
+		UserDataWrapper.UserData.QuadTreePosition = FVector(WaterQuadTree.GetTileRegion().Min, WaterQuadTreeMinHeight);
+		UserDataWrapper.UserData.CaptureDepthRange = WaterQuadTreeMaxHeight - WaterQuadTreeMinHeight;
+		UserDataWrapper.UserData.IndirectInstanceData0 = IndirectDrawResources.InstanceData0->GetRHI();
+		UserDataWrapper.UserData.IndirectInstanceData1 = IndirectDrawResources.InstanceData1->GetRHI();
+		UserDataWrapper.UserData.IndirectInstanceData2 = IndirectDrawResources.InstanceData2->GetRHI();
+		UserDataWrapper.UserData.IndirectInstanceData3 = (WITH_WATER_SELECTION_SUPPORT != 0) ? IndirectDrawResources.InstanceData3->GetRHI() : nullptr;
+		UserDataWrapper.UserData.IndirectInstanceDataOffsetsSRV = IndirectDrawResources.InstanceDataOffsets->GetSRV(RHICmdList, FRHIBufferSRVCreateInfo(PF_R32_UINT));
+		UserDataWrapper.UserData.IndirectInstanceData0SRV = IndirectDrawResources.InstanceData0->GetSRV(RHICmdList, FRHIBufferSRVCreateInfo(PF_R32_UINT));
+		UserDataWrapper.UserData.IndirectInstanceData1SRV = IndirectDrawResources.InstanceData1->GetSRV(RHICmdList, FRHIBufferSRVCreateInfo(PF_R32_UINT));
+		UserDataWrapper.UserData.IndirectInstanceData2SRV = IndirectDrawResources.InstanceData2->GetSRV(RHICmdList, FRHIBufferSRVCreateInfo(PF_R32_UINT));
+		UserDataWrapper.UserData.IndirectInstanceData3SRV = (WITH_WATER_SELECTION_SUPPORT != 0) ? IndirectDrawResources.InstanceData3->GetSRV(RHICmdList, FRHIBufferSRVCreateInfo(PF_R32_UINT)) : nullptr;
+
+		Result.UserData[(int32)RenderGroup] = &UserDataWrapper.UserData;
+	}
+
+	return Result;
+}
+
+FViewWaterQuadTree::FWaterLODParams FViewWaterQuadTree::GetWaterLODParams(const FVector& Position, float LODScale) const
+{
+	float WaterHeightForLOD = 0.0f;
+	WaterQuadTree.QueryInterpolatedTileBaseHeightAtLocation(FVector2D(Position), WaterHeightForLOD);
+
+	// Need to let the lowest LOD morph globally towards the next LOD. When the LOD is done morphing, simply clamp the LOD in the LOD selection to effectively promote the lowest LOD to the same LOD level as the one above
+	float DistToWater = FMath::Abs(Position.Z - WaterHeightForLOD) / LODScale;
+	DistToWater = FMath::Max(DistToWater - 2.0f, 0.0f);
+	DistToWater *= 2.0f;
+
+	// Clamp to WaterTileQuadTree.GetLODCount() - 1.0f prevents the last LOD to morph
+	const float FloatLOD = FMath::Clamp(FMath::Log2(DistToWater), 0.0f, WaterQuadTree.GetTreeDepth() - 1.0f);
+
+	FWaterLODParams WaterLODParams;
+	WaterLODParams.HeightLODFactor = FMath::Frac(FloatLOD);
+	WaterLODParams.LowestLOD = FMath::Clamp(FMath::FloorToInt(FloatLOD), 0, WaterQuadTree.GetTreeDepth() - 1);
+	WaterLODParams.WaterHeightForLOD = WaterHeightForLOD;
+
+	return WaterLODParams;
+}
+
+void FViewWaterQuadTree::BuildGPUQuadTree(FRDGBuilder& GraphBuilder)
 {
 	auto BuildOrthoMatrix = [](float InOrthoWidth, float InOrthoHeight, float DepthRange)
 	{
diff --git a/Engine/Plugins/Experimental/Water/Source/Runtime/Private/WaterQuadTreeBuilder.cpp b/Engine/Plugins/Experimental/Water/Source/Runtime/Private/WaterQuadTreeBuilder.cpp
--- a/Engine/Plugins/Experimental/Water/Source/Runtime/Private/WaterQuadTreeBuilder.cpp
+++ b/Engine/Plugins/Experimental/Water/Source/Runtime/Private/WaterQuadTreeBuilder.cpp
@@ -15,6 +15,10 @@ void FWaterQuadTreeBuilder::Init(const FBox2D& InWaterZoneBounds2D, const FIntPo
 	FarDistanceMeshExtent = InFarDistanceMeshExtent;
 	bUseFarMeshWithoutOcean = bInUseFarMeshWithoutOcean;
 	bIsGPUQuadTree = bInIsGPUQuadTree;
+
+	const int32 MaxDim = (int32)FMath::Max(InExtentInTiles.X * 2, InExtentInTiles.Y * 2);
+	const float RootDim = (float)FMath::RoundUpToPowerOfTwo(MaxDim);
+	TreeDepth = (int32)FMath::Log2(RootDim);
 }
 
 void FWaterQuadTreeBuilder::AddWaterBody(const FWaterBody& WaterBody)
@@ -177,6 +181,8 @@ bool FWaterQuadTreeBuilder::BuildWaterQuadTree(FWaterQuadTree& WaterQuadTree, co
 
 	WaterQuadTree.Unlock(true);
 
+	WaterQuadTree.BuildMaterialIndices();
+
 	return true;
 }
 
@@ -205,6 +211,11 @@ int32 FWaterQuadTreeBuilder::GetMaxLeafCount() const
 	return ExtentInTiles.X * ExtentInTiles.Y * 4;
 }
 
+int32 FWaterQuadTreeBuilder::GetTreeDepth() const
+{
+	return TreeDepth;
+}
+
 FIntPoint FWaterQuadTreeBuilder::GetResolution() const
 {
 	return ExtentInTiles * 2;
diff --git a/Engine/Plugins/Experimental/Water/Source/Runtime/Public/WaterMeshComponent.h b/Engine/Plugins/Experimental/Water/Source/Runtime/Public/WaterMeshComponent.h
--- a/Engine/Plugins/Experimental/Water/Source/Runtime/Public/WaterMeshComponent.h
+++ b/Engine/Plugins/Experimental/Water/Source/Runtime/Public/WaterMeshComponent.h
@@ -50,9 +50,7 @@ public:
 	/** Use this instead of GetMaterialRelevance, since this one will go over all materials from all tiles */
 	FMaterialRelevance GetWaterMaterialRelevance(ERHIFeatureLevel::Type InFeatureLevel) const;
 
-	bool ShouldBuildQuadTreeInSceneProxy() const;
 	const FWaterQuadTreeBuilder& GetWaterQuadTreeBuilder() const { return WaterQuadTreeBuilder; }
-	const FWaterQuadTree& GetWaterQuadTree() const;
 
 	const TSet<TObjectPtr<UMaterialInterface>>& GetUsedMaterialsSet() const { return UsedMaterials; }
 
@@ -66,6 +64,9 @@ public:
 
 	void SetDynamicWaterMeshCenter(const FVector2D& NewCenter);
 	FVector2D GetDynamicWaterMeshCenter() const { return DynamicWaterMeshCenter; }
+	FVector2D GetGlobalWaterMeshCenter() const;
+
+	bool IsLocalOnlyTessellationEnabled() const;
 
 	void SetTileSize(float NewTileSize);
 	float GetTileSize() const { return TileSize; }
diff --git a/Engine/Plugins/Experimental/Water/Source/Runtime/Public/WaterMeshSceneProxy.h b/Engine/Plugins/Experimental/Water/Source/Runtime/Public/WaterMeshSceneProxy.h
--- a/Engine/Plugins/Experimental/Water/Source/Runtime/Public/WaterMeshSceneProxy.h
+++ b/Engine/Plugins/Experimental/Water/Source/Runtime/Public/WaterMeshSceneProxy.h
@@ -6,6 +6,7 @@
 #include "PrimitiveSceneProxy.h"
 #include "Materials/MaterialRelevance.h"
 #include "WaterQuadTree.h"
+#include "WaterQuadTreeBuilder.h"
 #include "WaterVertexFactory.h"
 #include "RayTracingGeometry.h"
 #include "WaterQuadTreeGPU.h"
@@ -15,6 +16,88 @@
 
 class UWaterMeshComponent;
 
+using FWaterInstanceDataBuffersType = TWaterInstanceDataBuffers<WITH_WATER_SELECTION_SUPPORT>;
+using FWaterMeshUserDataBuffersType = TWaterMeshUserDataBuffers<WITH_WATER_SELECTION_SUPPORT>;
+using FWaterMeshUserDataType = TWaterMeshUserData<WITH_WATER_SELECTION_SUPPORT>;
+
+/** Set of quadtree related constants that do not change over the lifetime of the FWaterMeshSceneProxy and are shared by all quadtrees owned by it. */
+struct FWaterQuadTreeConstants
+{
+	/** Scale of the concentric LOD squares  */
+	float LODScale = -1.0f;
+
+	/** Number of quads per side of a water quad tree tile at LOD0 */
+	int32 NumQuadsLOD0 = 0;
+
+	int32 NumQuadsPerIndirectDrawTile = 0;
+
+	/** Number of densities (same as number of grid index/vertex buffers) */
+	int32 DensityCount = 0;
+
+	int32 ForceCollapseDensityLevel = TNumericLimits<int32>::Max();
+};
+
+/** A water quadtree instance owned by FWaterMeshSceneProxy and associated with a certain view. In splitscreen, each player should get their own FViewWaterQuadTree. */
+class FViewWaterQuadTree
+{
+public:
+	struct FWaterLODParams
+	{
+		int32 LowestLOD;
+		float HeightLODFactor;
+		float WaterHeightForLOD;
+	};
+
+	struct FUserDataAndIndirectArgs
+	{
+		TStaticArray<FWaterMeshUserDataType*, 3> UserData = {};
+		TRefCountPtr<FRDGPooledBuffer> IndirectArgs = nullptr;
+	};
+
+	// Rebuilds the quadtree at the specified position.
+	void Update(const FWaterQuadTreeBuilder& Builder, const FVector2D& CenterPosition);
+	
+	// Traverses the GPU quadtree to build indirect draw calls and associated buffers. May also initialize the GPU quadtree if it wasn't initialized already.
+	void TraverseGPUQuadTree(FRDGBuilder& GraphBuilder, bool bDepthBufferIsPopulated);
+	
+	// Allocates transient GPU resources to build indirect draw calls for the GPU quadtree and returns parameters needed by the water vertex factory to draw the quadtree.
+	FUserDataAndIndirectArgs PrepareGPUQuadTreeForRendering(const TArray<const FSceneView*>& Views, uint32 VisibilityMap, FMeshElementCollector& Collector, const FWaterQuadTreeConstants& QuadTreeConstants, const TArrayView<const EWaterMeshRenderGroupType> & BatchRenderGroups, FRHICommandListBase& RHICmdList) const;
+	
+	// Evaluates the CPU quadtree at the given position and returns parameters used for CPU quadtree draw call generation.
+	FWaterLODParams GetWaterLODParams(const FVector& Position, float LODScale) const;
+	
+	const FWaterQuadTree& GetWaterQuadTree() const { return WaterQuadTree; }
+	FWaterInstanceDataBuffersType* GetWaterInstanceDataBuffers() const { return WaterInstanceDataBuffers.Get(); }
+	FWaterMeshUserDataBuffersType* GetWaterMeshUserDataBuffers() const { return WaterMeshUserDataBuffers.Get(); }
+	double GetMinHeight() const { return WaterQuadTreeMinHeight; }
+	double GetMaxHeight() const { return WaterQuadTreeMaxHeight; }
+
+private:
+
+	/** Tiles containing water, stored in a quad tree. */
+	FWaterQuadTree WaterQuadTree;
+
+	/** GPU quad tree instance. Only initialized and used if WaterQuadTree.IsGPUQuadTree() is true. */
+	FWaterQuadTreeGPU QuadTreeGPU;
+
+	/** Unique Instance data buffer shared accross water batch draw calls */
+	TUniquePtr<FWaterInstanceDataBuffersType> WaterInstanceDataBuffers = nullptr;
+
+	/** Per-"water render group" user data (the number of groups might vary depending on whether we're in the editor or not) */
+	TUniquePtr<FWaterMeshUserDataBuffersType> WaterMeshUserDataBuffers = nullptr;
+
+	/** Vertical extent of the quadtree. */
+	double WaterQuadTreeMinHeight = DBL_MAX;
+	double WaterQuadTreeMaxHeight = -DBL_MAX;
+
+	mutable FWaterQuadTreeGPU::FTraverseParams WaterQuadTreeGPUTraverseParams;
+	mutable bool bNeedToTraverseGPUQuadTree = false;
+
+	/** Initializes the GPU quad tree. */
+	void BuildGPUQuadTree(FRDGBuilder& GraphBuilder);
+};
+
+
 /** Water mesh scene proxy */
 
 class FWaterMeshSceneProxy final : public FPrimitiveSceneProxy
@@ -56,7 +139,7 @@
 
 	uint32 GetAllocatedSize() const
 	{
-		return(FPrimitiveSceneProxy::GetAllocatedSize() + (WaterVertexFactories.GetAllocatedSize() + WaterVertexFactories.Num() * sizeof(FWaterVertexFactoryType)) + WaterQuadTree.GetAllocatedSize());
+		return(FPrimitiveSceneProxy::GetAllocatedSize() + (WaterVertexFactories.GetAllocatedSize() + WaterVertexFactories.Num() * sizeof(FWaterVertexFactoryType)) + ViewQuadTrees.GetAllocatedSize());
 	}
 
 #if WITH_WATER_SELECTION_SUPPORT
@@ -67,8 +150,6 @@
 	using FWaterVertexFactoryType = TWaterVertexFactory<WITH_WATER_SELECTION_SUPPORT, EWaterVertexFactoryDrawMode::NonIndirect>;
 	using FWaterVertexFactoryIndirectDrawType = TWaterVertexFactory<WITH_WATER_SELECTION_SUPPORT, EWaterVertexFactoryDrawMode::Indirect>;
 	using FWaterVertexFactoryIndirectDrawISRType = TWaterVertexFactory<WITH_WATER_SELECTION_SUPPORT, EWaterVertexFactoryDrawMode::IndirectInstancedStereo>;
-	using FWaterInstanceDataBuffersType = TWaterInstanceDataBuffers<WITH_WATER_SELECTION_SUPPORT>;
-	using FWaterMeshUserDataBuffersType = TWaterMeshUserDataBuffers<WITH_WATER_SELECTION_SUPPORT>;
 
 #if RHI_RAYTRACING
 	virtual void GetDynamicRayTracingInstances(FRayTracingMaterialGatheringContext& Context, TArray<FRayTracingInstance>& OutRayTracingInstances) override final;
@@ -76,13 +157,14 @@
 	virtual bool IsRayTracingRelevant() const override { return true; }
 #endif
 
+	// Creates and initializes a new quadtree centered around CenterPosition and associated with the given key. Returns true if the quadtree was created and false if it already exists.
+	bool CreateViewWaterQuadTree(int32 Key, const FVector2D& CenterPosition);
+	// Updates an existing quadtree by reconstructing it at a new CenterPosition. Returns true if the quadtree exists and was updated and false otherwise.
+	bool UpdateViewWaterQuadTree(int32 Key, const FVector2D& CenterPosition);
+	// Destroys the quadtree associated with Key.
+	void DestroyViewWaterQuadTree(int32 Key);
+
 private:
-	struct FWaterLODParams
-	{
-		int32 LowestLOD;
-		float HeightLODFactor;
-		float WaterHeightForLOD;
-	};
 
 #if RHI_RAYTRACING
 	struct FRayTracingWaterData
@@ -92,54 +174,25 @@
 	};
 #endif
 
-	bool HasWaterData() const 
+	struct FOcclusionCullingResults
 	{
-		return WaterQuadTree.GetNodeCount() != 0 && DensityCount != 0;
-	}
+		uint32 FrameNumber;
+		TArray<bool> Results;
+	};
 
-	FWaterLODParams GetWaterLODParams(const FVector& Position) const;
+	TMap<int32, FViewWaterQuadTree> ViewQuadTrees;
 
-#if RHI_RAYTRACING
-	void SetupRayTracingInstances(FRHICommandListBase& RHICmdList, int32 NumInstances, uint32 DensityIndex);
-#endif
-
 	FMaterialRelevance MaterialRelevance;
 
-	// One vertex factory per LOD
+	// One vertex factory per LOD. Only used for CPU driven water quadtree rendering.
 	TArray<FWaterVertexFactoryType*> WaterVertexFactories;
 	FWaterVertexFactoryIndirectDrawType* WaterVertexFactoryIndirectDraw = nullptr;
 	FWaterVertexFactoryIndirectDrawISRType* WaterVertexFactoryIndirectDrawISR = nullptr;
 
-	/** Tiles containing water, stored in a quad tree */
-	FWaterQuadTree WaterQuadTree;
+	FWaterQuadTreeBuilder WaterQuadTreeBuilder;
+	FWaterQuadTreeConstants WaterQuadTreeConstants;
+	// If this is true, then this proxy can manage multiple local quadtrees, potentially associated with different views.
+	bool bIsLocalOnlyTessellationEnabled = false;
 
-	/** GPU quad tree instance. Only initialized and used if WaterQuadTree.IsGPUQuadTree() is true. */
-	FWaterQuadTreeGPU QuadTreeGPU;
-
-	/** Unique Instance data buffer shared accross water batch draw calls */	
-	FWaterInstanceDataBuffersType* WaterInstanceDataBuffers = nullptr;
-
-	/** Per-"water render group" user data (the number of groups might vary depending on whether we're in the editor or not) */
-	FWaterMeshUserDataBuffersType* WaterMeshUserDataBuffers = nullptr;
-
-	double WaterQuadTreeMinHeight = DBL_MAX;
-	double WaterQuadTreeMaxHeight = -DBL_MAX;
-
-	/** The world-space bounds of the current water info texture coverage. The Water mesh should only render tiles within this bounding box. */
-	FBox2D WaterInfoBounds = FBox2D(ForceInit);
-
-	/** Scale of the concentric LOD squares  */
-	float LODScale = -1.0f;
-
-	/** Number of quads per side of a water quad tree tile at LOD0 */
-	int32 NumQuadsLOD0 = 0;
-
-	int32 NumQuadsPerIndirectDrawTile = 0;
-
-	/** Number of densities (same as number of grid index/vertex buffers) */
-	int32 DensityCount = 0;
-
-	int32 ForceCollapseDensityLevel = TNumericLimits<int32>::Max();
-
 	mutable int32 HistoricalMaxViewInstanceCount = 0;
 
@@ -148,12 +201,8 @@
 	TArray<TArray<FRayTracingWaterData>> RayTracingWaterData;	
 #endif
 
-	struct FOcclusionCullingResults
-	{
-		uint32 FrameNumber;
-		TArray<bool> Results;
-	};
-
+	
+	// CPU-driven occlusion culling related members
 	TArray<FBoxSphereBounds> OcclusionCullingBounds;
 	TArray<FBoxSphereBounds> EmptyOcclusionCullingBounds;
 	TMap<uint32, FOcclusionCullingResults> OcclusionResults;
@@ -161,11 +210,14 @@
 	int32 OcclusionResultsFarMeshOffset = INT32_MAX;
 	uint32 SceneProxyCreatedFrameNumberRenderThread = INDEX_NONE;
 
-	mutable FWaterQuadTreeGPU::FTraverseParams WaterQuadTreeGPUTraverseParams;
-	mutable bool bNeedToTraverseGPUQuadTree = false;
-
-	/** Initializes the GPU quad tree */
-	void BuildGPUQuadTree(FRDGBuilder& GraphBuilder);
+	bool HasWaterData() const;
+	TArray<EWaterMeshRenderGroupType, TInlineAllocator<FWaterVertexFactoryType::NumRenderGroups>> GetBatchRenderGroups(const TArray<const FSceneView*>& Views, uint32 VisibilityMap) const;
+	uint32 GetWireframeVisibilityMapAndMaterial(const TArray<const FSceneView*>& Views, uint32 VisibilityMap, FMeshElementCollector& Collector, class FColoredMaterialRenderProxy*& OutMaterialInstance) const;
+	int32 FindBestQuadTreeForView(const FSceneView* View) const;
+	TArray<int32, TInlineAllocator<8>> GetViewToQuadTreeMapping(const TArray<const FSceneView*>& Views, uint32 VisibilityMap) const;
+#if RHI_RAYTRACING
+	void SetupRayTracingInstances(FRHICommandListBase& RHICmdList, int32 NumInstances, uint32 DensityIndex);
+#endif
 };
 
 #if UE_ENABLE_INCLUDE_ORDER_DEPRECATED_IN_5_2
diff --git a/Engine/Plugins/Experimental/Water/Source/Runtime/Public/WaterQuadTreeBuilder.h b/Engine/Plugins/Experimental/Water/Source/Runtime/Public/WaterQuadTreeBuilder.h
--- a/Engine/Plugins/Experimental/Water/Source/Runtime/Public/WaterQuadTreeBuilder.h
+++ b/Engine/Plugins/Experimental/Water/Source/Runtime/Public/WaterQuadTreeBuilder.h
@@ -44,6 +44,7 @@ public:
 	bool IsGPUQuadTree() const;
 	float GetLeafSize() const;
 	int32 GetMaxLeafCount() const;
+	int32 GetTreeDepth() const;
 	FIntPoint GetResolution() const;
 
 private:
@@ -51,6 +52,7 @@ private:
 	FBox2D WaterZoneBounds2D = FBox2D();
 	FIntPoint ExtentInTiles = FIntPoint::ZeroValue;
 	float TileSize = 0.0f;
+	int32 TreeDepth = 0;
 	FMaterialRenderProxy* FarDistanceMaterial = nullptr;
 	float FarDistanceMeshExtent = 0.0f;
 	bool bUseFarMeshWithoutOcean = false;
