[Water]

Fixed potential nullptr when SetupView is called on a water view extension which has an invalid world.

diff --git a/Engine/Plugins/Experimental/Water/Source/Runtime/Private/WaterViewExtension.cpp b/Engine/Plugins/Experimental/Water/Source/Runtime/Private/WaterViewExtension.cpp
--- a/Engine/Plugins/Experimental/Water/Source/Runtime/Private/WaterViewExtension.cpp
+++ b/Engine/Plugins/Experimental/Water/Source/Runtime/Private/WaterViewExtension.cpp
@@ -310,7 +310,11 @@ void FWaterViewExtension::SetupView(FSceneViewFamily& InViewFamily, FSceneView&
 	const FVector ViewLocation = InView.ViewLocation;
 
 	const TWeakObjectPtr<UWorld> WorldPtr = GetWorld();
-	check(WorldPtr.IsValid())
+
+	if (!ensureMsgf(WorldPtr.IsValid(), TEXT("FWaterViewExtension::SetupView was called while it's owning world is not valid! Lifetime of the WaterViewExtension is tied to the world, this should be impossible!")))
+	{
+		return;
+	}
 
 	// Prevent re-entrancy. 
 	// Since the water info render will update the view extensions we could end up with a re-entrant case.
