Water: Fixed r.water.waterinfo.rendermethod 1 crashing when trying to render to a WaterInfoTexture slice other than 0.

diff --git a/Engine/Plugins/Experimental/Water/Source/Runtime/Private/WaterUtils.cpp b/Engine/Plugins/Experimental/Water/Source/Runtime/Private/WaterUtils.cpp
--- a/Engine/Plugins/Experimental/Water/Source/Runtime/Private/WaterUtils.cpp
+++ b/Engine/Plugins/Experimental/Water/Source/Runtime/Private/WaterUtils.cpp
@@ -107,7 +107,8 @@ WATER_API UTextureRenderTarget2DArray* FWaterUtils::GetOrCreateTransientRenderTa
 			&& (InRenderTarget->SizeY == InSize.Y)
 			&& (InRenderTarget->Slices == InSlices)
 			&& (InRenderTarget->GetFormat() == PixelFormat) // Watch out : GetFormat() returns a EPixelFormat (non-class enum), so we can't compare with a ETextureRenderTargetFormat
-			&& (InRenderTarget->ClearColor == InClearColor))
+			&& (InRenderTarget->ClearColor == InClearColor)
+			&& (InRenderTarget->bTargetArraySlicesIndependently))
 		{
 			return InRenderTarget;
 		}
@@ -116,6 +117,7 @@ WATER_API UTextureRenderTarget2DArray* FWaterUtils::GetOrCreateTransientRenderTa
 	UTextureRenderTarget2DArray* NewRenderTarget2DArray = NewObject<UTextureRenderTarget2DArray>(GetTransientPackage(), MakeUniqueObjectName(GetTransientPackage(), UTextureRenderTarget2DArray::StaticClass(), InRenderTargetName));
 	check(NewRenderTarget2DArray);
 	NewRenderTarget2DArray->ClearColor = InClearColor;
+	NewRenderTarget2DArray->bTargetArraySlicesIndependently = true; // When r.water.waterinfo.rendermethod is 1, we directly render into each individual slice
 	NewRenderTarget2DArray->Init(InSize.X, InSize.Y, InSlices, PixelFormat);
 	NewRenderTarget2DArray->UpdateResourceImmediate(true);
 
