Do not copy shared navigation settings onto the water spline componoent

diff --git a/Engine/Plugins/Experimental/Water/Source/Runtime/Private/WaterBodyCustomComponent.cpp b/Engine/Plugins/Experimental/Water/Source/Runtime/Private/WaterBodyCustomComponent.cpp
--- a/Engine/Plugins/Experimental/Water/Source/Runtime/Private/WaterBodyCustomComponent.cpp
+++ b/Engine/Plugins/Experimental/Water/Source/Runtime/Private/WaterBodyCustomComponent.cpp
@@ -6,6 +6,7 @@
 #include "Misc/UObjectToken.h"
 #include "WaterBodyActor.h"
 #include "WaterBodyMeshComponent.h"
+#include "WaterSplineComponent.h"
 
 #include UE_INLINE_GENERATED_CPP_BY_NAME(WaterBodyCustomComponent)
 
@@ -106,7 +107,11 @@ void UWaterBodyCustomComponent::OnUpdateBody(bool bWithExclusionVolumes)
 			continue;
 		}
 
-		CopySharedNavigationSettingsToComponent(Comp);
+		// Do not copy Nav setting over Water Spline Component, its just a spline component and will report Navigation warning if only one spline point because of empty bounds box
+		if (!Comp->IsA<UWaterSplineComponent>())
+		{
+			CopySharedNavigationSettingsToComponent(Comp);
+		}
 
 		Comp->SetMobility(Mobility);
 	}
