WaterMeshSceneProxy: Fix Non-Unity compile error

diff --git a/Engine/Plugins/Experimental/Water/Source/Runtime/Public/WaterMeshSceneProxy.h b/Engine/Plugins/Experimental/Water/Source/Runtime/Public/WaterMeshSceneProxy.h
--- a/Engine/Plugins/Experimental/Water/Source/Runtime/Public/WaterMeshSceneProxy.h
+++ b/Engine/Plugins/Experimental/Water/Source/Runtime/Public/WaterMeshSceneProxy.h
@@ -9,6 +9,7 @@
 #include "WaterQuadTreeBuilder.h"
 #include "WaterVertexFactory.h"
 #include "RayTracingGeometry.h"
+#include "RenderGraphResources.h"
 #include "WaterQuadTreeGPU.h"
 
 class FMeshElementCollector;
