Water: Fixed WaterMeshSceneProxy failing to create a global internal quadtree when local tessellation is disabled.

diff --git a/Engine/Plugins/Experimental/Water/Source/Runtime/Private/WaterMeshSceneProxy.cpp b/Engine/Plugins/Experimental/Water/Source/Runtime/Private/WaterMeshSceneProxy.cpp
--- a/Engine/Plugins/Experimental/Water/Source/Runtime/Private/WaterMeshSceneProxy.cpp
+++ b/Engine/Plugins/Experimental/Water/Source/Runtime/Private/WaterMeshSceneProxy.cpp
@@ -238,13 +238,13 @@ FWaterMeshSceneProxy::FWaterMeshSceneProxy(UWaterMeshComponent* Component)
 	if (!bIsLocalOnlyTessellationEnabled)
 #endif
 	{
+		FViewWaterQuadTree& ViewWaterQuadTree = ViewQuadTrees.Add(INDEX_NONE);
 #if FORCE_SINGLE_WATER_QUADTREE
-		CreateViewWaterQuadTree(INDEX_NONE, bIsLocalOnlyTessellationEnabled ? Component->GetDynamicWaterMeshCenter() : Component->GetGlobalWaterMeshCenter());
+		ViewWaterQuadTree.Update(WaterQuadTreeBuilder, bIsLocalOnlyTessellationEnabled ? Component->GetDynamicWaterMeshCenter() : Component->GetGlobalWaterMeshCenter());
 #else
-		CreateViewWaterQuadTree(INDEX_NONE, Component->GetGlobalWaterMeshCenter());
+		ViewWaterQuadTree.Update(WaterQuadTreeBuilder, Component->GetGlobalWaterMeshCenter());
 #endif
-		FViewWaterQuadTree* ViewWaterQuadTree = ViewQuadTrees.Find(INDEX_NONE);
-		const FWaterQuadTree& WaterQuadTree = ViewWaterQuadTree->GetWaterQuadTree();
+		const FWaterQuadTree& WaterQuadTree = ViewWaterQuadTree.GetWaterQuadTree();
 
 		// Always do CPU occlusion queries, even if this is a GPU quadtree. The GPU quadtree still potentially uses the far mesh which is CPU driven.
 		const int32 MaxQueries = CVarWaterMeshOcclusionCullingMaxQueries.GetValueOnGameThread();
@@ -254,7 +254,7 @@ FWaterMeshSceneProxy::FWaterMeshSceneProxy(UWaterMeshComponent* Component)
 		// If this is a GPU quadtree, the CPU root node will have an invalid bounding box, so derive conservative bounds now
 		if (bIsGPUQuadTree && !OcclusionCullingBounds.IsEmpty())
 		{
-			OcclusionCullingBounds[0] = FBox(FVector(WaterQuadTree.GetTileRegion().Min, ViewWaterQuadTree->GetMinHeight()), FVector(WaterQuadTree.GetTileRegion().Max, ViewWaterQuadTree->GetMaxHeight()));
+			OcclusionCullingBounds[0] = FBox(FVector(WaterQuadTree.GetTileRegion().Min, ViewWaterQuadTree.GetMinHeight()), FVector(WaterQuadTree.GetTileRegion().Max, ViewWaterQuadTree.GetMaxHeight()));
 		}
 	}
 #if !FORCE_SINGLE_WATER_QUADTREE
