Fix cook nondeterminism involving UWaterBodyMeshComponent

Can't use BodySetup->bHasCookedCollisionData to check if there's collision data attached since that bool is initialized
to true in the ctor and not actually serialized, so it's just set to true for everything by default.

We still want to keep the fix where we discard any attached collision data for this component which isn't
supposed to have it, though. Do it in a slightly more careful way that avoids invalidating the GUID and body
setup every time, which is what was causing cook non-determinism issues.

If we don't have a body setup, or do we have a body setup that didn't have bNeverNeedsCookedCollisionData set,
we still need to run the old logic that creates a new body setup.

diff --git a/Engine/Plugins/Experimental/Water/Source/Runtime/Private/WaterBodyMeshComponent.cpp b/Engine/Plugins/Experimental/Water/Source/Runtime/Private/WaterBodyMeshComponent.cpp
--- a/Engine/Plugins/Experimental/Water/Source/Runtime/Private/WaterBodyMeshComponent.cpp
+++ b/Engine/Plugins/Experimental/Water/Source/Runtime/Private/WaterBodyMeshComponent.cpp
@@ -113,12 +113,35 @@ void UWaterBodyMeshComponent::FixupCollisionOnBodySetup()
 {
 	TRACE_CPUPROFILER_EVENT_SCOPE(WaterBodyMeshComponent::FixupCollisionOnBodySetup);
 
-	// Fix for a bug that existing where the body setup was not created with bNeverNeedsCookedCollisionData, and also address an issue where bNeverNeedsCookedCollisionData was allowing meshes to generate cooked data.
-	UBodySetup* BodySetup = GetBodySetup();
+	// We had a bug where the body setup was not created with bNeverNeedsCookedCollisionData,
+	// and another where bNeverNeedsCookedCollisionData was allowing meshes to generate cooked data.
+	// This attempts to clean up the resulting mess.
 	UStaticMesh* Mesh = GetStaticMesh();
-	const bool bResetBodySetup = Mesh != nullptr && (BodySetup == nullptr || !BodySetup->bNeverNeedsCookedCollisionData || BodySetup->bHasCookedCollisionData);
-	if (bResetBodySetup)
+	if (!Mesh)
+	{
+		return;
+	}
+
+	UBodySetup* BodySetup = GetBodySetup();
+
+	if (BodySetup && BodySetup->bNeverNeedsCookedCollisionData)
+	{
+		// We could still have cooked collision data attached, but shouldn't.
+		// Quietly remove the cooked collision data without invalidating the GUID,
+		// because the latter leads to cook non-determinism for affected instances.
+		//
+		// This is effectively just stripping cached data, and referring back to
+		// any cooked collision data in the DDC under the same GUID is totally
+		// valid if there were no other modifications that actually truly
+		// invalidate physics data.
+		FGuid OriginalBodySetupGuid = BodySetup->BodySetupGuid;
+		BodySetup->InvalidatePhysicsData();
+		BodySetup->BodySetupGuid = OriginalBodySetupGuid;
+	}
+	else if (!BodySetup || !BodySetup->bNeverNeedsCookedCollisionData)
 	{
+		// We don't have a body setup, or it's set up the wrong
+		// way. Create a new one.
 		Mesh->CreateBodySetup();
 		BodySetup = GetBodySetup();
 		BodySetup->bNeverNeedsCookedCollisionData = true;
diff --git a/Engine/Source/Runtime/Engine/Classes/PhysicsEngine/BodySetup.h b/Engine/Source/Runtime/Engine/Classes/PhysicsEngine/BodySetup.h
--- a/Engine/Source/Runtime/Engine/Classes/PhysicsEngine/BodySetup.h
+++ b/Engine/Source/Runtime/Engine/Classes/PhysicsEngine/BodySetup.h
@@ -193,7 +193,7 @@ class UBodySetup : public UBodySetupCore
 	/** Flag used to know if we have failed to create physics meshes. Note that this is not the inverse of bCreatedPhysicsMeshes which is true even on failure */
 	uint8 bFailedToCreatePhysicsMeshes:1;
 
-	/** Indicates whether this setup has any cooked collision data. */
+	/** Indicates whether this setup has any cooked collision data. This is set when writing to or reading from cooked data, but not authoritative for assets in the editor. */
 	uint8 bHasCookedCollisionData:1;
 
 	/** Indicates that we will never use convex or trimesh shapes. This is an optimization to skip checking for binary data. */
