Add UBakedShallowWaterSimulationComponent, which is a simple container for CPU accessible grid based water simulation data
UWaterBodyComponent - weak ptr to baked shallow water sim data component (owned by UShallowWaterRiverComponent)
Add override toggle to use/disable shallow water sims for physics

diff --git a/Engine/Plugins/Experimental/Water/Source/Runtime/Private/BakedShallowWaterSimulationComponent.cpp b/Engine/Plugins/Experimental/Water/Source/Runtime/Private/BakedShallowWaterSimulationComponent.cpp
new file mode 100644
--- /dev/null
+++ b/Engine/Plugins/Experimental/Water/Source/Runtime/Private/BakedShallowWaterSimulationComponent.cpp
@@ -0,0 +1,89 @@
+// Copyright Epic Games, Inc. All Rights Reserved.
+#include "BakedShallowWaterSimulationComponent.h"
+
+void FShallowWaterSimulationGrid::SampleShallowWaterSimulationAtIndex(const FVector2D &QueryFloatIndex, FVector& OutWaterVelocity, float& OutWaterHeight, float& OutWaterDepth) const
+{
+	OutWaterHeight = 0;
+	OutWaterDepth = 0;
+	OutWaterVelocity = FVector(0, 0, 0);
+
+	if (QueryFloatIndex.X > 0 && QueryFloatIndex.X < NumCells.X - 1 && QueryFloatIndex.Y > 0 && QueryFloatIndex.Y < NumCells.Y - 1)
+	{
+		const FIntVector2 BaseIndex = FIntVector2(QueryFloatIndex.X, QueryFloatIndex.Y);
+		const FVector2D LerpValue = QueryFloatIndex - FVector2D(BaseIndex.X, BaseIndex.Y);
+
+		const FVector4 Sample00 = ArrayValues[(BaseIndex.X + 0) + (BaseIndex.Y + 0) * NumCells.X];
+		const FVector4 Sample10 = ArrayValues[(BaseIndex.X + 1) + (BaseIndex.Y + 0) * NumCells.X];
+		const FVector4 Sample01 = ArrayValues[(BaseIndex.X + 0) + (BaseIndex.Y + 1) * NumCells.X];
+		const FVector4 Sample11 = ArrayValues[(BaseIndex.X + 1) + (BaseIndex.Y + 1) * NumCells.X];
+		
+		const FVector4 Sample0 = Sample00 * (1. - LerpValue.X) + Sample10 * (LerpValue.X);
+		const FVector4 Sample1 = Sample01 * (1. - LerpValue.X) + Sample11 * (LerpValue.X);
+
+		const FVector4 Sample = Sample0 * (1. - LerpValue.Y) + Sample1 * (LerpValue.Y);
+
+		OutWaterHeight = Sample.X + Position.Z;
+		OutWaterDepth = Sample.Y;
+		OutWaterVelocity = FVector(Sample.Z, Sample.W, 0);
+	}
+}
+
+void FShallowWaterSimulationGrid::QueryShallowWaterSimulationAtIndex(const FIntVector2 &QueryIndex, FVector& OutWaterVelocity, float& OutWaterHeight, float& OutWaterDepth) const
+{
+	OutWaterHeight = 0;
+	OutWaterDepth = 0;
+	OutWaterVelocity = FVector(0,0,0);
+
+	if (QueryIndex.X >= 0 && QueryIndex.X < NumCells.X && QueryIndex.Y >= 0 && QueryIndex.Y < NumCells.Y)
+	{
+		const FVector4 Sample = ArrayValues[QueryIndex.X + QueryIndex.Y * NumCells.X];
+
+		OutWaterHeight = Sample.X + Position.Z;
+		OutWaterDepth = Sample.Y;
+		OutWaterVelocity = FVector(Sample.Z, Sample.W, 0);
+	}
+}
+
+FVector FShallowWaterSimulationGrid::ComputeShallowWaterSimulationNormalAtPosition(const FVector &QueryPos) const
+{
+	const FVector2D FloatIndexPos = WorldToFloatIndex(QueryPos);
+
+	float HRight;
+	float HUp;
+	float H;
+
+	float Depth;
+	float DepthRight;
+	float DepthUp;
+
+	FVector VelocityIGNORE;
+
+	FVector Normal(0, 0, 1);
+	
+	SampleShallowWaterSimulationAtIndex(FloatIndexPos + FVector2D::ZeroVector, VelocityIGNORE, H, Depth);
+	SampleShallowWaterSimulationAtIndex(FloatIndexPos + FVector2D(1,0), VelocityIGNORE, HRight, DepthRight);
+	SampleShallowWaterSimulationAtIndex(FloatIndexPos + FVector2D(0, 1), VelocityIGNORE, HUp, DepthUp);
+
+	FVector WorldPos = FloatIndexToWorld(FloatIndexPos + FVector2D::ZeroVector);
+	FVector WorldPosRight = FloatIndexToWorld(FloatIndexPos + FVector2D(1, 0));
+	FVector WorldPosUp = FloatIndexToWorld(FloatIndexPos + FVector2D(0, 1));
+
+	WorldPos.Z += H;
+	WorldPosRight.Z += HRight;
+	WorldPosUp.Z += HUp;
+
+	const FVector CrossProd = (WorldPosRight - WorldPos).Cross(WorldPosUp - WorldPos);
+	const float CrossProdLength = CrossProd.Length();
+
+	// default to upward facing normal
+	if (CrossProdLength > SMALL_NUMBER)
+	{
+		Normal = CrossProd / CrossProdLength;
+	}
+
+	return Normal;
+}
+
+UBakedShallowWaterSimulationComponent::UBakedShallowWaterSimulationComponent(const FObjectInitializer& ObjectInitializer)
+	: Super(ObjectInitializer)
+{}
\ No newline at end of file
diff --git a/Engine/Plugins/Experimental/Water/Source/Runtime/Public/BakedShallowWaterSimulationComponent.h b/Engine/Plugins/Experimental/Water/Source/Runtime/Public/BakedShallowWaterSimulationComponent.h
new file mode 100644
--- /dev/null
+++ b/Engine/Plugins/Experimental/Water/Source/Runtime/Public/BakedShallowWaterSimulationComponent.h
@@ -0,0 +1,119 @@
+// Copyright Epic Games, Inc. All Rights Reserved.
+
+#pragma once
+
+#include "CoreMinimal.h"
+#include "Math/Float16Color.h"
+#include "Components/PrimitiveComponent.h"
+#include "BakedShallowWaterSimulationComponent.generated.h"
+
+USTRUCT(BlueprintType)
+struct WATER_API FShallowWaterSimulationGrid
+{
+	GENERATED_BODY()
+
+public:
+	FShallowWaterSimulationGrid(const TArray<FVector4> &InArrayValues, const FIntVector2 &InNumCells, const FVector &InPosition, const FVector2D &InSize)
+		: ArrayValues(InArrayValues), 
+		NumCells(InNumCells), 
+		Position(InPosition), 
+		Size(InSize)
+	{
+		Dx = Size / FVector2D(NumCells.X, NumCells.Y);
+	}
+	FShallowWaterSimulationGrid()
+	{				
+		NumCells = { 0, 0 };
+		Position = { 0.0, 0.0, 0.0 };		
+		Size = { 0.0, 0.0 };
+		Dx = { 0.0, 0.0 };
+	};
+
+	/** Check for a valid sim domain */
+	bool IsValid() const
+	{
+		return NumCells.X > 0 && NumCells.Y > 0 && ArrayValues.Num() == NumCells.X * NumCells.Y && Size.X > SMALL_NUMBER && Size.Y > SMALL_NUMBER;
+	}
+
+	/** Linear interpolate sim data at an index */
+	void SampleShallowWaterSimulationAtIndex(const FVector2D &QueryFloatIndex, FVector& OutWaterVelocity, float& OutWaterHeight, float& OutWaterDepth) const;
+	
+	/** Query sim data at an index */
+	void QueryShallowWaterSimulationAtIndex(const FIntVector2 &QueryIndex, FVector& OutWaterVelocity, float& OutWaterHeight, float& OutWaterDepth) const;
+
+	/** Compute normal at a world position */
+	FVector ComputeShallowWaterSimulationNormalAtPosition(const FVector &QueryPos) const;
+	
+	/** Query shallow water grid cell values at a world position */
+	void FORCEINLINE QueryShallowWaterSimulationAtPosition(const FVector &QueryPos, FVector& OutWaterVelocity, float& OutWaterHeight, float& OutWaterDepth) const
+	{
+		const FIntVector2 IndexPos = WorldToIndex(QueryPos);
+		QueryShallowWaterSimulationAtIndex(IndexPos, OutWaterVelocity, OutWaterHeight, OutWaterDepth);
+	}
+
+	/** Linear interpolate for shallow water grid cell values at a world position */
+	void FORCEINLINE SampleShallowWaterSimulationAtPosition(const FVector &QueryPos, FVector& OutWaterVelocity, float& OutWaterHeight, float& OutWaterDepth) const
+	{
+		FVector2D FloatIndex = WorldToFloatIndex(QueryPos);
+		SampleShallowWaterSimulationAtIndex(FloatIndex, OutWaterVelocity, OutWaterHeight, OutWaterDepth);
+	}
+
+	/** Convert world space to integer index */
+	FIntVector2 FORCEINLINE WorldToIndex(const FVector &WorldPos) const
+	{
+		FVector2D FloatIndex = WorldToFloatIndex(WorldPos);
+		return FIntVector2(FloatIndex.X, FloatIndex.Y);
+	}
+
+	/** Convert world space to float index */
+	FVector2D FORCEINLINE WorldToFloatIndex(const FVector &WorldPos) const
+	{
+		const FVector LocalPos = WorldPos - Position;
+		const FVector2D UnitPos = FVector2D(LocalPos.X, LocalPos.Y) / Size + .5;
+		return UnitPos * FVector2D(NumCells.X, NumCells.Y) - .5;
+	}
+
+	/** Convert a float index to world space */
+	FVector FORCEINLINE FloatIndexToWorld(const FVector2D &Index) const
+	{
+		const FVector2D UnitPos = (Index + .5) / FVector2D(NumCells.X, NumCells.Y);
+		const FVector2D LocalPos = (UnitPos - .5) * Size;
+		const FVector LocalPos3D = FVector(LocalPos.X, LocalPos.Y, 0);
+		return LocalPos3D + Position;
+	}
+
+	/** Convert an integer index to world space */
+	FVector FORCEINLINE IndexToWorld(const FIntVector2 &Index) const
+	{
+		return FloatIndexToWorld(FVector2D(Index.X, Index.Y));
+	}
+
+	/** Raw grid data */
+	UPROPERTY()
+	TArray<FVector4> ArrayValues;
+	
+	/** Number of cells */
+	UPROPERTY()
+	FIntVector2 NumCells;
+	
+	/** World space center */
+	UPROPERTY()
+	FVector Position;
+	
+	/** World space size */
+	UPROPERTY()
+	FVector2D Size;
+
+	/** World space grid cell size */
+	UPROPERTY()
+	FVector2D Dx;
+};
+
+UCLASS(meta = (BlueprintSpawnableComponent), MinimalAPI)
+class UBakedShallowWaterSimulationComponent : public UPrimitiveComponent
+{
+	GENERATED_UCLASS_BODY()
+
+	UPROPERTY()
+	FShallowWaterSimulationGrid SimulationData;
+};
\ No newline at end of file
diff --git a/Engine/Plugins/Experimental/Water/Source/Runtime/Public/WaterBodyComponent.h b/Engine/Plugins/Experimental/Water/Source/Runtime/Public/WaterBodyComponent.h
--- a/Engine/Plugins/Experimental/Water/Source/Runtime/Public/WaterBodyComponent.h
+++ b/Engine/Plugins/Experimental/Water/Source/Runtime/Public/WaterBodyComponent.h
@@ -9,6 +9,7 @@
 #include "WaterCurveSettings.h"
 #include "WaterBodyStaticMeshSettings.h"
 #include "WaterSplineMetadata.h"
+#include "BakedShallowWaterSimulationComponent.h"
 #include "WaterZoneActor.h"
 #include "WaterBodyTypes.h"
 
@@ -425,6 +426,15 @@ public:
 	UFUNCTION(BlueprintCallable, Category=Water, meta=(Deprecated = "5.2"))
 	void OnWaterBodyChanged(bool bShapeOrPositionChanged, bool bWeightmapSettingsChanged = false, bool bUserTriggeredChanged = false) {}
 
+	/** Get the baked shallow water simulation for this water body */
+	UBakedShallowWaterSimulationComponent* GetBakedShallowWaterSimulation() { return BakedShallowWaterSim.Get(); }
+	
+	/** Set the baked shallow water simulation for this water body */
+	void SetBakedShallowWaterSimulation(TObjectPtr<UBakedShallowWaterSimulationComponent> BakedSim) { BakedShallowWaterSim = BakedSim; }
+
+	/** Query for if the baked shallow water sim is valid for use */
+	bool UseBakedShallowWaterSimulationForBuoyancy() const { return bUseBakedSimForPhysics && BakedShallowWaterSim.IsValid() && BakedShallowWaterSim->SimulationData.IsValid(); }
+
 protected:
 	//~ Begin UActorComponent interface.
 	virtual bool IsHLODRelevant() const override;
@@ -676,6 +686,14 @@ protected:
 	UPROPERTY(Category = Water, EditAnywhere, AdvancedDisplay)
 	double FixedWaterDepth = 512.0;
 
+	/**  Baked simulation data for this water body, owned by a UShallowWaterRiverComponent */
+	UPROPERTY(Category = BakedSimulation, AdvancedDisplay, VisibleAnywhere)
+	TWeakObjectPtr<UBakedShallowWaterSimulationComponent> BakedShallowWaterSim;
+
+	/**  Override to disable use of the baked shallow water simulation for collisons and other uses */
+	UPROPERTY(Category = BakedSimulation, AdvancedDisplay, EditAnywhere)
+	bool bUseBakedSimForPhysics = true;
+
 #if WITH_EDITORONLY_DATA
 	UPROPERTY()
 	TArray<TLazyObjectPtr<AWaterBodyIsland>> Islands_DEPRECATED;
