Limit the Water DI to no longer support querying depth related information as it requires accessing data that is unsafe from the concurrent task.  Two new properties have been added to the DI bEvaluateSystemDepth & bEvaluateSystemDepthPerFrame which indicates that the depth should be sampled from the water body on the game thread using the system instance.  Also, legacy behavior can be restored through the cvar fx.Niagara.Water.DepthQuerySupported.

diff --git a/Engine/Plugins/Experimental/Water/Source/Runtime/Private/NiagaraDataInterfaceWater.cpp b/Engine/Plugins/Experimental/Water/Source/Runtime/Private/NiagaraDataInterfaceWater.cpp
--- a/Engine/Plugins/Experimental/Water/Source/Runtime/Private/NiagaraDataInterfaceWater.cpp
+++ b/Engine/Plugins/Experimental/Water/Source/Runtime/Private/NiagaraDataInterfaceWater.cpp
@@ -13,11 +13,24 @@
 
 #define LOCTEXT_NAMESPACE "NiagaraDataInterfaceWater"
 
+// this cvar controls whether we'll allow depth queries into water bodies from the worker task.  This has been shown to be unsafe
+// because the water body and the underlying assets (landscape proxies) could be in flux on the gamethread as assets stream in.
+static int GNiagaraWaterDepthQuerySupported = 0;
+static FAutoConsoleVariableRef CVarNiagaraSoloTickEarly(
+	TEXT("fx.Niagara.Water.DepthQuerySupported"),
+	GNiagaraWaterDepthQuerySupported,
+	TEXT("When enabled water body queries will include unsafe access to the depth."),
+	ECVF_Default
+);
+
 struct FNDIWater_InstanceData
 {
 	bool								bFindClosestBody = false;
+	bool								bEvaluateSystemDepth = true;
+	bool								bEvaluateSystemDepthPerFrame = false;
 	TWeakObjectPtr<UWaterBodyComponent>	WaterBodyComponent;
 	uint32								WaterBodyChangeId = 0;
+	float								SystemInstanceWaterDepth = 0.0f;
 	FNiagaraLWCConverter				LWCConverter;
 };
 
@@ -106,15 +119,15 @@ namespace NDIWaterPrivate
 			const bool bIncludeDepth = InIncludeDepth.GetAndAdvance();
 			const bool bIncludeWaves = InIncludeWaves.GetAndAdvance();
 			const bool bSimpleWaves = InSimpleWaves.GetAndAdvance();
+			const bool bDoDepthQuery = bIncludeDepth && GNiagaraWaterDepthQuerySupported;
 
 			if (bExecuteQuery)
 			{
 				const EWaterBodyQueryFlags QueryFlags =
 					EWaterBodyQueryFlags::ComputeLocation |
 					EWaterBodyQueryFlags::ComputeNormal |
-					EWaterBodyQueryFlags::ComputeImmersionDepth |
 					EWaterBodyQueryFlags::ComputeVelocity |
-					(bIncludeDepth ? EWaterBodyQueryFlags::ComputeDepth : EWaterBodyQueryFlags::None) |
+					(bDoDepthQuery ? EWaterBodyQueryFlags::ComputeDepth : EWaterBodyQueryFlags::None) |
 					(bIncludeWaves ? EWaterBodyQueryFlags::IncludeWaves : EWaterBodyQueryFlags::None) |
 					(bIncludeWaves && bSimpleWaves ? EWaterBodyQueryFlags::SimpleWaves : EWaterBodyQueryFlags::None);
 
@@ -123,10 +136,16 @@ namespace NDIWaterPrivate
 				const FVector3f WaterPlaneLocation = InstData->LWCConverter.ConvertWorldToSimulationPosition(QueryResult.GetWaterPlaneLocation());
 				const FVector3f WaterSurfacePosition = InstData->LWCConverter.ConvertWorldToSimulationPosition(QueryResult.GetWaterSurfaceLocation());
 
+				float DepthValue = bIncludeDepth ? InstData->SystemInstanceWaterDepth : 0.0f;
+				if (bDoDepthQuery)
+				{
+					DepthValue = QueryResult.GetWaterSurfaceDepth();
+				}
+
 				OutWaterPlanePosition.SetAndAdvance(WaterPlaneLocation);
 				OutWaterPlaneNormal.SetAndAdvance(FVector3f(QueryResult.GetWaterPlaneNormal()));
 				OutWaterSurfacePosition.SetAndAdvance(WaterSurfacePosition);
-				OutWaterDepth.SetAndAdvance(bIncludeDepth ? QueryResult.GetWaterSurfaceDepth() : 0.0f);
+				OutWaterDepth.SetAndAdvance(DepthValue);
 				OutWaterVelocity.SetAndAdvance(FVector3f(QueryResult.GetVelocity()));
 				OutInExclusionVolume.SetAndAdvance(QueryResult.IsInExclusionVolume());
 			}
@@ -275,6 +294,8 @@ bool UNiagaraDataInterfaceWater::Equals(const UNiagaraDataInterface* Other) cons
 	const UNiagaraDataInterfaceWater* OtherTyped = CastChecked<const UNiagaraDataInterfaceWater>(Other);
 	return
 		OtherTyped->bFindWaterBodyOnSpawn == bFindWaterBodyOnSpawn &&
+		OtherTyped->bEvaluateSystemDepth == bEvaluateSystemDepth &&
+		OtherTyped->bEvaluateSystemDepthPerFrame == bEvaluateSystemDepthPerFrame &&
 		OtherTyped->SourceBodyComponent == SourceBodyComponent;
 }
 
@@ -287,6 +308,8 @@ bool UNiagaraDataInterfaceWater::CopyToInternal(UNiagaraDataInterface* Destinati
 
 	UNiagaraDataInterfaceWater* OtherTyped = CastChecked<UNiagaraDataInterfaceWater>(Destination);
 	OtherTyped->bFindWaterBodyOnSpawn = bFindWaterBodyOnSpawn;
+	OtherTyped->bEvaluateSystemDepth = bEvaluateSystemDepth;
+	OtherTyped->bEvaluateSystemDepthPerFrame = bEvaluateSystemDepthPerFrame;
 	OtherTyped->SourceBodyComponent = SourceBodyComponent;
 
 	return true;
@@ -301,6 +324,8 @@ bool UNiagaraDataInterfaceWater::InitPerInstanceData(void* PerInstanceData, FNia
 {
 	FNDIWater_InstanceData* InstData = new (PerInstanceData) FNDIWater_InstanceData();
 	InstData->bFindClosestBody		= bFindWaterBodyOnSpawn;
+	InstData->bEvaluateSystemDepth	= bEvaluateSystemDepth;
+	InstData->bEvaluateSystemDepthPerFrame = bEvaluateSystemDepthPerFrame;
 	InstData->WaterBodyComponent	= nullptr;
 	InstData->WaterBodyChangeId		= SourceBodyChangeId - 1;
 	InstData->LWCConverter			= SystemInstance->GetLWCConverter();
@@ -319,6 +344,8 @@ bool UNiagaraDataInterfaceWater::PerInstanceTick(void* PerInstanceData, FNiagara
 	check(SystemInstance);
 	FNDIWater_InstanceData* InstData = static_cast<FNDIWater_InstanceData*>(PerInstanceData);
 
+	bool bCalcDepth = bEvaluateSystemDepth && bEvaluateSystemDepthPerFrame;
+
 	// Do we need to update the water body component?
 	if (InstData->WaterBodyChangeId != SourceBodyChangeId)
 	{
@@ -335,6 +362,20 @@ bool UNiagaraDataInterfaceWater::PerInstanceTick(void* PerInstanceData, FNiagara
 		}
 		InstData->bFindClosestBody = false;
 		InstData->WaterBodyChangeId = SourceBodyChangeId;
+
+		bCalcDepth = bEvaluateSystemDepth;
+	}
+
+	if (bCalcDepth)
+	{
+		if (UWaterBodyComponent* WaterBodyComponent = InstData->WaterBodyComponent.Get())
+		{
+			const FVector QueryLocation = SystemInstance->GetWorldTransform().GetTranslation();
+			const EWaterBodyQueryFlags QueryFlags = EWaterBodyQueryFlags::ComputeDepth;
+			const FWaterBodyQueryResult QueryResult = WaterBodyComponent->QueryWaterInfoClosestToWorldLocation(QueryLocation, QueryFlags);
+
+			InstData->SystemInstanceWaterDepth = QueryResult.IsInExclusionVolume() ? 0.0f : QueryResult.GetWaterSurfaceDepth();
+		}
 	}
 
 	return false;
@@ -403,17 +444,29 @@ void UNiagaraDataInterfaceWater::GetWaterDataAtPoint(FVectorVMExternalFunctionCo
 		if (Component != nullptr)
 		{
 			FVector QueryPos = InstData->LWCConverter.ConvertSimulationPositionToWorld(WorldPos.GetAndAdvance());
-			QueryResult = Component->QueryWaterInfoClosestToWorldLocation(QueryPos,
-				EWaterBodyQueryFlags::ComputeLocation
+			EWaterBodyQueryFlags QueryFlags = EWaterBodyQueryFlags::ComputeLocation
 				| EWaterBodyQueryFlags::ComputeVelocity
 				| EWaterBodyQueryFlags::ComputeNormal
-				| EWaterBodyQueryFlags::ComputeDepth
-				| EWaterBodyQueryFlags::IncludeWaves);
+				| EWaterBodyQueryFlags::IncludeWaves;
+
+			if (GNiagaraWaterDepthQuerySupported)
+			{
+				QueryFlags |= EWaterBodyQueryFlags::ComputeDepth;
+			}
+
+			QueryResult = Component->QueryWaterInfoClosestToWorldLocation(QueryPos, QueryFlags);
+
 			bIsValid = !QueryResult.IsInExclusionVolume();
 		}
 
+		float DepthValue = InstData->SystemInstanceWaterDepth;
+		if (GNiagaraWaterDepthQuerySupported)
+		{
+			DepthValue = QueryResult.GetWaterSurfaceDepth();
+		}
+
 		OutHeight.SetAndAdvance(bIsValid ? QueryResult.GetWaveInfo().Height : 0.0f);
-		OutDepth.SetAndAdvance(bIsValid ? QueryResult.GetWaterSurfaceDepth() : 0.0f);
+		OutDepth.SetAndAdvance(bIsValid ? DepthValue : 0.0f);
 		OutVelocity.SetAndAdvance(bIsValid ? FVector3f(QueryResult.GetVelocity()) : FVector3f::ZeroVector);		// LWC_TODO: Precision loss
 
 		// Note we assume X and Y are in water by the time this is queried
diff --git a/Engine/Plugins/Experimental/Water/Source/Runtime/Public/NiagaraDataInterfaceWater.h b/Engine/Plugins/Experimental/Water/Source/Runtime/Public/NiagaraDataInterfaceWater.h
--- a/Engine/Plugins/Experimental/Water/Source/Runtime/Public/NiagaraDataInterfaceWater.h
+++ b/Engine/Plugins/Experimental/Water/Source/Runtime/Public/NiagaraDataInterfaceWater.h
@@ -55,7 +55,15 @@ private:
 	UPROPERTY(EditAnywhere, Category = "Water")
 	bool bFindWaterBodyOnSpawn = false;
 
-	UPROPERTY(EditAnywhere, Category = "Water") 
+	/** When enabled the owning system instance position will be used to sample the depth of the water. */
+	UPROPERTY(EditAnywhere, Category = "Water")
+	bool bEvaluateSystemDepth = true;
+
+	/** If bEvaluateSystemDepth is enabled the depth will be updated each frame. */
+	UPROPERTY(EditAnywhere, Category = "Water", meta = (EditCondition = "bEvaluateSystemDepth"))
+	bool bEvaluateSystemDepthPerFrame = true;
+
+	UPROPERTY(EditAnywhere, Category = "Water")
 	TObjectPtr<UWaterBodyComponent> SourceBodyComponent;
 
 	uint32 SourceBodyChangeId = 0;
