Added multiple views (split screen) support to SingleLayerWater by using a different slice of WaterInfoTexture and a different water quadtree per view.

#changelist validated
#virtualized

diff --git a/Engine/Plugins/Experimental/Water/Shaders/Private/WaterDataFunctions.ush b/Engine/Plugins/Experimental/Water/Shaders/Private/WaterDataFunctions.ush
--- a/Engine/Plugins/Experimental/Water/Shaders/Private/WaterDataFunctions.ush
+++ b/Engine/Plugins/Experimental/Water/Shaders/Private/WaterDataFunctions.ush
@@ -6,6 +6,7 @@
 #define PER_WAVE_DATA_SIZE 2
 #define PER_WATERZONE_DATA_SIZE 2
 #define PER_WATERBODY_DATA_SIZE 2
+#define PER_WATERZONEVIEW_DATA_SIZE 1
 
 /** Struct containing the decoded water body data  */	
 struct FWaterBodyData
@@ -44,6 +45,11 @@ struct FWaveParams
 	float Steepness;
 };
 
+struct FWaterZoneViewData
+{
+	float2 Location;
+};
+
 int GetWaterZoneDataOffset()
 {
 	return (int)View.WaterData[0].x;
@@ -54,9 +60,29 @@ int GetWaveDataOffset()
 	return (int)View.WaterData[0].y;
 }
 
-int GetWaterInfoTextureViewIndex()
+int GetWaterZoneViewDataOffset()
+{
+	return (int)View.WaterData[0].z;
+}
+
+int GetNumWaterViewData()
+{
+	return (int)View.WaterData[0].w;
+}
+
+bool IsWaterLocalTessellationEnabled(int InWaterZoneIndex)
+{
+	const int WaterZoneDataBufferOffset = GetWaterZoneDataOffset();
+	const int WaterZoneOffset = WaterZoneDataBufferOffset + (InWaterZoneIndex * PER_WATERZONE_DATA_SIZE);
+
+	return View.WaterData[WaterZoneOffset + 1].y > 0;
+}
+
+int GetWaterInfoTextureViewIndex(int InWaterZoneIndex)
 {
-	return View.WaterInfoTextureViewIndex;
+	int Index = IsWaterLocalTessellationEnabled(InWaterZoneIndex) ? View.WaterInfoTextureViewIndex : 0;
+	
+	return Index;
 }
 
 /** Function to decode the raw data from the water indirection buffer */
@@ -105,26 +131,45 @@ FWaveParams GetWaveDataNew(int InWaveIndex, FWaterBodyData InWaterBodyData)
 	return DecodeWaveParams(Data0, Data1);
 }
 
-FWaterZoneData DecodeWaterZoneData(float4 Data0, float4 Data1)
+FWaterZoneData DecodeWaterZoneData(float2 Location, float4 Data0, float4 Data1)
 {
 	FWaterZoneData OutWaterZoneData = (FWaterZoneData)0;
-	OutWaterZoneData.Location = Data0.xy;
-	OutWaterZoneData.Extent = Data0.zw;
-	OutWaterZoneData.HeightExtent = Data1.xy;
-	OutWaterZoneData.GroundZMin = Data1.z;
+	OutWaterZoneData.Location = Location;
+	OutWaterZoneData.Extent = Data0.xy;
+	OutWaterZoneData.HeightExtent = Data0.zw;
+	OutWaterZoneData.GroundZMin = Data1.x;
 
 	return OutWaterZoneData;
 }
 
+FWaterZoneViewData GetWaterZoneViewData(int InWaterZoneIndex)
+{
+	const int NumViews = GetNumWaterViewData();
+
+	const int ViewIndex = GetWaterInfoTextureViewIndex(InWaterZoneIndex);
+
+	const int WaterZoneViewDataBufferOffset = GetWaterZoneViewDataOffset();
+	const int WaterZoneViewOffset = WaterZoneViewDataBufferOffset + (InWaterZoneIndex * NumViews + ViewIndex) * PER_WATERZONEVIEW_DATA_SIZE;
+
+	FWaterZoneViewData WaterZoneViewData = (FWaterZoneViewData)0;
+	WaterZoneViewData.Location = View.WaterData[WaterZoneViewOffset].xy;
+
+	return WaterZoneViewData;
+}
+
 FWaterZoneData GetWaterZoneData(int InWaterZoneIndex)
 {
 	const int WaterZoneDataBufferOffset = GetWaterZoneDataOffset();
 	const int WaterZoneOffset = WaterZoneDataBufferOffset + (InWaterZoneIndex * PER_WATERZONE_DATA_SIZE);
 
+	FWaterZoneViewData WaterZoneViewData = GetWaterZoneViewData(InWaterZoneIndex);
+
+	const float2 ZoneViewLocation = WaterZoneViewData.Location;
+	
 	const float4 Data0 = View.WaterData[WaterZoneOffset];
 	const float4 Data1 = View.WaterData[WaterZoneOffset + 1];
 
-	return DecodeWaterZoneData(Data0, Data1);
+	return DecodeWaterZoneData(ZoneViewLocation, Data0, Data1);
 }
 
 FWaterZoneData GetWaterZoneDataForWaterBody(int InWaterBodyIndex)
diff --git a/Engine/Plugins/Experimental/Water/Source/Runtime/Private/WaterInfoRendering.cpp b/Engine/Plugins/Experimental/Water/Source/Runtime/Private/WaterInfoRendering.cpp
--- a/Engine/Plugins/Experimental/Water/Source/Runtime/Private/WaterInfoRendering.cpp
+++ b/Engine/Plugins/Experimental/Water/Source/Runtime/Private/WaterInfoRendering.cpp
@@ -712,7 +712,8 @@ static FSceneRenderer* CreateWaterInfoSceneRenderer(const FCreateWaterInfoSceneR
 
 void UpdateWaterInfoRendering(
 	FSceneInterface* Scene,
-	const WaterInfo::FRenderingContext& Context)
+	const WaterInfo::FRenderingContext& Context,
+	const FVector& WaterInfoCenter)
 {
 	TRACE_CPUPROFILER_EVENT_SCOPE(WaterInfo::UpdateWaterInfoRendering);
 
@@ -731,7 +732,7 @@ void UpdateWaterInfoRendering(
 	}
 	const FVector ZoneExtent = Context.ZoneToRender->GetDynamicWaterInfoExtent();
 
-	FVector ViewLocation = Context.ZoneToRender->GetDynamicWaterInfoCenter();
+	FVector ViewLocation = WaterInfoCenter;
 	ViewLocation.Z = Context.CaptureZ;
 
 	const FBox2D CaptureBounds(FVector2D(ViewLocation - ZoneExtent), FVector2D(ViewLocation + ZoneExtent));
@@ -822,77 +823,72 @@ void UpdateWaterInfoRendering(
 		});
 }
 
-void UpdateWaterInfoRendering2(FSceneView& InView, const TWeakObjectPtrKeyMap<AWaterZone, UE::WaterInfo::FRenderingContext> WaterInfoContexts)
+void UpdateWaterInfoRendering2(FSceneView& InView, const FRenderingContext& Context, int32 RenderTargetArrayLayer, const FVector& WaterInfoCenter)
 {
 	TRACE_CPUPROFILER_EVENT_SCOPE(WaterInfo::UpdateWaterInfoRendering2);
 
 	InView.WaterInfoTextureRenderingParams.Reset();
-	for (const TPair<TWeakObjectPtr<AWaterZone>, UE::WaterInfo::FRenderingContext>& Pair : WaterInfoContexts)
+
+	if (!IsValid(Context.TextureRenderTarget))
 	{
-		if (Pair.Key.IsValid())
-		{
-			const UE::WaterInfo::FRenderingContext& Context(Pair.Value);
+		return;
+	}
 
-			if (!IsValid(Context.TextureRenderTarget))
-			{
-				continue;
-			}
-			const FVector ZoneExtent = Context.ZoneToRender->GetDynamicWaterInfoExtent();
-
-			FVector ViewLocation = Context.ZoneToRender->GetDynamicWaterInfoCenter();
-			ViewLocation.Z = Context.CaptureZ;
-
-			const FBox2D CaptureBounds(FVector2D(ViewLocation - ZoneExtent), FVector2D(ViewLocation + ZoneExtent));
-
-			// Zone rendering always happens facing towards negative z.
-			const FVector LookAt = ViewLocation - FVector(0.f, 0.f, 1.f);
-
-			FSceneView::FWaterInfoTextureRenderingParams RenderingParams;
-			RenderingParams.RenderTarget = Context.TextureRenderTarget->GameThread_GetRenderTargetResource();
-			RenderingParams.ViewLocation = ViewLocation;
-			RenderingParams.ViewRotationMatrix = FLookAtMatrix(ViewLocation, LookAt, FVector(0.f, -1.f, 0.f));
-			RenderingParams.ViewRotationMatrix = RenderingParams.ViewRotationMatrix.RemoveTranslation();
-			RenderingParams.ViewRotationMatrix.RemoveScaling();
-			RenderingParams.ProjectionMatrix = BuildOrthoMatrix(ZoneExtent.X, ZoneExtent.Y);
-			RenderingParams.CaptureZ = ViewLocation.Z;
-			RenderingParams.WaterHeightExtents = Context.ZoneToRender->GetWaterHeightExtents();
-			RenderingParams.GroundZMin = Context.ZoneToRender->GetGroundZMin();
-			RenderingParams.VelocityBlurRadius = Context.ZoneToRender->GetVelocityBlurRadius();
-			RenderingParams.WaterZoneExtents = ZoneExtent;
-
-			if (Context.GroundPrimitiveComponents.Num() > 0)
+	const FVector ZoneExtent = Context.ZoneToRender->GetDynamicWaterInfoExtent();
+
+	FVector ViewLocation = WaterInfoCenter;
+	ViewLocation.Z = Context.CaptureZ;
+
+	const FBox2D CaptureBounds(FVector2D(ViewLocation - ZoneExtent), FVector2D(ViewLocation + ZoneExtent));
+
+	// Zone rendering always happens facing towards negative z.
+	const FVector LookAt = ViewLocation - FVector(0.f, 0.f, 1.f);
+
+	FSceneView::FWaterInfoTextureRenderingParams RenderingParams;
+	RenderingParams.RenderTarget = Context.TextureRenderTarget->GameThread_GetRenderTargetResource();
+	RenderingParams.ViewLocation = ViewLocation;
+	RenderingParams.ViewRotationMatrix = FLookAtMatrix(ViewLocation, LookAt, FVector(0.f, -1.f, 0.f));
+	RenderingParams.ViewRotationMatrix = RenderingParams.ViewRotationMatrix.RemoveTranslation();
+	RenderingParams.ViewRotationMatrix.RemoveScaling();
+	RenderingParams.ProjectionMatrix = BuildOrthoMatrix(ZoneExtent.X, ZoneExtent.Y);
+	RenderingParams.CaptureZ = ViewLocation.Z;
+	RenderingParams.WaterHeightExtents = Context.ZoneToRender->GetWaterHeightExtents();
+	RenderingParams.GroundZMin = Context.ZoneToRender->GetGroundZMin();
+	RenderingParams.VelocityBlurRadius = Context.ZoneToRender->GetVelocityBlurRadius();
+	RenderingParams.WaterZoneExtents = ZoneExtent;
+	RenderingParams.RenderTargetArrayLayer = RenderTargetArrayLayer;
+
+	if (Context.GroundPrimitiveComponents.Num() > 0)
+	{
+		RenderingParams.TerrainComponentIds.Reserve(Context.GroundPrimitiveComponents.Num());
+		for (TWeakObjectPtr<UPrimitiveComponent> GroundPrimComp : Context.GroundPrimitiveComponents)
+		{
+			if (GroundPrimComp.IsValid())
 			{
-				RenderingParams.TerrainComponentIds.Reserve(Context.GroundPrimitiveComponents.Num());
-				for (TWeakObjectPtr<UPrimitiveComponent> GroundPrimComp : Context.GroundPrimitiveComponents)
-				{
-					if (GroundPrimComp.IsValid())
-					{
-						RenderingParams.TerrainComponentIds.Add(GroundPrimComp.Get()->GetPrimitiveSceneId());
-					}
-				}
+				RenderingParams.TerrainComponentIds.Add(GroundPrimComp.Get()->GetPrimitiveSceneId());
 			}
-			if (Context.WaterBodies.Num() > 0)
+		}
+	}
+	if (Context.WaterBodies.Num() > 0)
+	{
+		RenderingParams.WaterBodyComponentIds.Reserve(Context.WaterBodies.Num());
+		RenderingParams.DilatedWaterBodyComponentIds.Reserve(Context.WaterBodies.Num());
+		for (const TWeakObjectPtr<UWaterBodyComponent> WaterBodyToRenderPtr : Context.WaterBodies)
+		{
+			if (UWaterBodyComponent* WaterBodyToRender = WaterBodyToRenderPtr.Get())
 			{
-				RenderingParams.WaterBodyComponentIds.Reserve(Context.WaterBodies.Num());
-				RenderingParams.DilatedWaterBodyComponentIds.Reserve(Context.WaterBodies.Num());
-				for (const TWeakObjectPtr<UWaterBodyComponent> WaterBodyToRenderPtr : Context.WaterBodies)
+				// Perform our own simple culling based on the known Capture bounds:
+				const FBox WaterBodyBounds = WaterBodyToRender->Bounds.GetBox();
+				if (CaptureBounds.Intersect(FBox2D(FVector2D(WaterBodyBounds.Min), FVector2D(WaterBodyBounds.Max))))
 				{
-					if (UWaterBodyComponent* WaterBodyToRender = WaterBodyToRenderPtr.Get())
-					{
-						// Perform our own simple culling based on the known Capture bounds:
-						const FBox WaterBodyBounds = WaterBodyToRender->Bounds.GetBox();
-						if (CaptureBounds.Intersect(FBox2D(FVector2D(WaterBodyBounds.Min), FVector2D(WaterBodyBounds.Max))))
-						{
-							RenderingParams.WaterBodyComponentIds.Add(WaterBodyToRender->GetWaterInfoMeshComponent()->GetPrimitiveSceneId());
-							RenderingParams.DilatedWaterBodyComponentIds.Add(WaterBodyToRender->GetDilatedWaterInfoMeshComponent()->GetPrimitiveSceneId());
-						}
-					}
+					RenderingParams.WaterBodyComponentIds.Add(WaterBodyToRender->GetWaterInfoMeshComponent()->GetPrimitiveSceneId());
+					RenderingParams.DilatedWaterBodyComponentIds.Add(WaterBodyToRender->GetDilatedWaterInfoMeshComponent()->GetPrimitiveSceneId());
 				}
 			}
-
-			InView.WaterInfoTextureRenderingParams.Add(MoveTemp(RenderingParams));
 		}
 	}
+
+	InView.WaterInfoTextureRenderingParams.Add(MoveTemp(RenderingParams));
 }
 
 
@@ -989,6 +985,7 @@ public:
 		FinalizeWaterInfo(GraphBuilder, *Views[0]->Family, *Views[0], MergeTargetTexture, FinalizedTexture, Params);
 
 		FRDGTextureRef WaterInfoTexture = RegisterExternalTexture(GraphBuilder, WaterInfoRenderTarget->GetRenderTargetTexture(), TEXT("WaterInfoTexture"));
+		GraphBuilder.UseInternalAccessMode(WaterInfoTexture);
 		FRHICopyTextureInfo CopyInfo;
 		CopyInfo.DestSliceIndex = RenderTargetArraySlice;
 		AddCopyTexturePass(GraphBuilder, FinalizedTexture, WaterInfoTexture, CopyInfo);
@@ -1053,7 +1050,9 @@ static TMap<uint32, int32> GatherLandscapeLODOverrides(const UWorld* World, cons
 void UpdateWaterInfoRendering_CustomRenderPass(
 	FSceneInterface* Scene,
 	const FSceneViewFamily& ViewFamily,
-	const WaterInfo::FRenderingContext& Context)
+	const WaterInfo::FRenderingContext& Context,
+	int32 TextureArraySlice,
+	const FVector& WaterInfoCenter)
 {
 	TRACE_CPUPROFILER_EVENT_SCOPE(WaterInfo::UpdateWaterInfoRendering_CustomRenderPass);
 
@@ -1074,7 +1073,7 @@ void UpdateWaterInfoRendering_CustomRenderPass(
 
 	const FVector ZoneExtent = Context.ZoneToRender->GetDynamicWaterInfoExtent();
 
-	FVector ViewLocation = Context.ZoneToRender->GetDynamicWaterInfoCenter();
+	FVector ViewLocation = WaterInfoCenter;
 	ViewLocation.Z = Context.CaptureZ;
 
 	const FBox2D CaptureBounds(FVector2D(ViewLocation - ZoneExtent), FVector2D(ViewLocation + ZoneExtent));
@@ -1155,6 +1154,7 @@ void UpdateWaterInfoRendering_CustomRenderPass(
 	DilationPass->ColorPass = ColorPass;
 	DilationPass->WaterInfoRenderTarget = Context.TextureRenderTarget->GameThread_GetRenderTargetResource();
 	DilationPass->Params = Params;
+	DilationPass->RenderTargetArraySlice = TextureArraySlice;
 	if (bPerformRenderCapture)
 	{
 		// End a render capture when this pass runs :
diff --git a/Engine/Plugins/Experimental/Water/Source/Runtime/Private/WaterInfoRendering.h b/Engine/Plugins/Experimental/Water/Source/Runtime/Private/WaterInfoRendering.h
--- a/Engine/Plugins/Experimental/Water/Source/Runtime/Private/WaterInfoRendering.h
+++ b/Engine/Plugins/Experimental/Water/Source/Runtime/Private/WaterInfoRendering.h
@@ -30,16 +30,21 @@ struct FRenderingContext
 	
 void UpdateWaterInfoRendering(
 	FSceneInterface* Scene,
-	const FRenderingContext& Context);
+	const FRenderingContext& Context,
+	const FVector& WaterInfoCenter);
 
 void UpdateWaterInfoRendering2(
 	FSceneView& InView, 
-	const TWeakObjectPtrKeyMap<AWaterZone, UE::WaterInfo::FRenderingContext> WaterInfoContexts);
+	const FRenderingContext& Context,
+	int32 RenderTargetArrayLayer,
+	const FVector& WaterInfoCenter);
 
 void UpdateWaterInfoRendering_CustomRenderPass(
 	FSceneInterface* Scene,
 	const FSceneViewFamily& ViewFamily,
-	const FRenderingContext& Context);
+	const FRenderingContext& Context,
+	int32 TextureArraySlice,
+	const FVector& WaterInfoCenter);
 
 const FName& GetWaterInfoDepthPassName();
 const FName& GetWaterInfoColorPassName();
diff --git a/Engine/Plugins/Experimental/Water/Source/Runtime/Private/WaterMeshComponent.cpp b/Engine/Plugins/Experimental/Water/Source/Runtime/Private/WaterMeshComponent.cpp
--- a/Engine/Plugins/Experimental/Water/Source/Runtime/Private/WaterMeshComponent.cpp
+++ b/Engine/Plugins/Experimental/Water/Source/Runtime/Private/WaterMeshComponent.cpp
@@ -185,15 +185,6 @@ FMaterialRelevance UWaterMeshComponent::GetWaterMaterialRelevance(ERHIFeatureLev
 	return Result;
 }
 
-void UWaterMeshComponent::SetDynamicWaterMeshCenter(const FVector2D& NewCenter)
-{
-	if (!DynamicWaterMeshCenter.Equals(NewCenter))
-	{
-		DynamicWaterMeshCenter = NewCenter;
-		MarkWaterMeshGridDirty();
-	}
-}
-
 FVector2D UWaterMeshComponent::GetGlobalWaterMeshCenter() const
 {
 	const float LODCountBiasFactor = FMath::Pow(2.0f, (float)CVarWaterMeshLODCountBias.GetValueOnGameThread());
diff --git a/Engine/Plugins/Experimental/Water/Source/Runtime/Private/WaterMeshSceneProxy.cpp b/Engine/Plugins/Experimental/Water/Source/Runtime/Private/WaterMeshSceneProxy.cpp
--- a/Engine/Plugins/Experimental/Water/Source/Runtime/Private/WaterMeshSceneProxy.cpp
+++ b/Engine/Plugins/Experimental/Water/Source/Runtime/Private/WaterMeshSceneProxy.cpp
@@ -18,9 +18,7 @@
 #include "RenderGraphBuilder.h"
 #include "RenderGraphUtils.h"
 #include "StereoRenderUtils.h"
-
-// TODO: Remove this once FWaterViewExtension actually supports managing multiple local quadtrees
-#define FORCE_SINGLE_WATER_QUADTREE 1
+#include "WaterSubsystem.h"
 
 DECLARE_STATS_GROUP(TEXT("Water Mesh"), STATGROUP_WaterMesh, STATCAT_Advanced);
 
@@ -234,16 +232,12 @@ FWaterMeshSceneProxy::FWaterMeshSceneProxy(UWaterMeshComponent* Component)
 
 	// In case we don't need to support multiple local quadtrees, create a global quadtree right away and initialize occlusion culling bounds.
 	// Local quadtrees are incompatible with CPU-based occlusion culling as the number of quadtrees can change during the lifetime of the scene proxy.
-#if !FORCE_SINGLE_WATER_QUADTREE
 	if (!bIsLocalOnlyTessellationEnabled)
-#endif
 	{
 		FViewWaterQuadTree& ViewWaterQuadTree = ViewQuadTrees.Add(INDEX_NONE);
-#if FORCE_SINGLE_WATER_QUADTREE
-		ViewWaterQuadTree.Update(WaterQuadTreeBuilder, bIsLocalOnlyTessellationEnabled ? Component->GetDynamicWaterMeshCenter() : Component->GetGlobalWaterMeshCenter());
-#else
+
 		ViewWaterQuadTree.Update(WaterQuadTreeBuilder, Component->GetGlobalWaterMeshCenter());
-#endif
+
 		const FWaterQuadTree& WaterQuadTree = ViewWaterQuadTree.GetWaterQuadTree();
 
 		// Always do CPU occlusion queries, even if this is a GPU quadtree. The GPU quadtree still potentially uses the far mesh which is CPU driven.
@@ -257,12 +251,18 @@ FWaterMeshSceneProxy::FWaterMeshSceneProxy(UWaterMeshComponent* Component)
 			OcclusionCullingBounds[0] = FBox(FVector(WaterQuadTree.GetTileRegion().Min, ViewWaterQuadTree.GetMinHeight()), FVector(WaterQuadTree.GetTileRegion().Max, ViewWaterQuadTree.GetMaxHeight()));
 		}
 	}
-#if !FORCE_SINGLE_WATER_QUADTREE
 	else
 	{
 		EmptyOcclusionCullingBounds.Add(Component->Bounds);
+
+		// Normally the quadtrees are created and updated by the WaterViewExtension. 
+		// It can happen that a SceneProxy tries to render before the WaterView updates it, which caused some water flickering.
+		// By forcing the creation here, using the quadtree update information stored by the water view extension, we make sure the water data is always valid.
+		if (FWaterViewExtension* WaterViewExtension = UWaterSubsystem::GetWaterViewExtension(Component->GetWorld()))
+		{
+			WaterViewExtension->CreateSceneProxyQuadtrees(this);
+		}
 	}
-#endif
 
 	// When using a GPU quadtree, this scene proxy allocates pooled buffers in GetDynamicMeshElements. AllocatePooledBuffer must be called on the renderthread.
 	// The callback on GWaterMeshGPUWork seems to be invoked after all GetDynamicMeshElements tasks finish, so there should be no race condition there.
@@ -709,11 +709,7 @@ void FWaterMeshSceneProxy::GetDynamicMeshElements(const TArray<const FSceneView*
 const TArray<FBoxSphereBounds>* FWaterMeshSceneProxy::GetOcclusionQueries(const FSceneView* View) const
 {
 	// Local tessellation/multiple dynamic quadtrees are incompatible with CPU occlusion culling
-	if (
-#if !FORCE_SINGLE_WATER_QUADTREE
-		!bIsLocalOnlyTessellationEnabled && 
-#endif
-		CVarWaterMeshOcclusionCulling.GetValueOnRenderThread() != 0)
+	if (!bIsLocalOnlyTessellationEnabled && CVarWaterMeshOcclusionCulling.GetValueOnRenderThread() != 0)
 	{
 		return &OcclusionCullingBounds;
 	}
@@ -857,20 +853,25 @@ int32 FWaterMeshSceneProxy::FindBestQuadTreeForView(const FSceneView* View) cons
 	// Otherwise just use the closest quadtree.
 	else
 	{
-		const FVector2D ViewPosition2D = FVector2D(View->ViewLocation);
-		double ClosestDistSquared = DBL_MAX;
-		int32 ClosestQuadTreeKey = INDEX_NONE;
-		for (const auto& Pair : ViewQuadTrees)
+		return FindBestQuadTreeForViewLocation(View);
+	}
+}
+
+int32 FWaterMeshSceneProxy::FindBestQuadTreeForViewLocation(const FSceneView* View) const
+{
+	const FVector2D ViewPosition2D = FVector2D(View->ViewLocation);
+	double ClosestDistSquared = DBL_MAX;
+	int32 ClosestQuadTreeKey = INDEX_NONE;
+	for (const auto& Pair : ViewQuadTrees)
+	{
+		const double DistSquared = FVector2D::DistSquared(Pair.Value.GetWaterQuadTree().GetTileRegion().GetCenter(), ViewPosition2D);
+		if (DistSquared < ClosestDistSquared)
 		{
-			const double DistSquared = FVector2D::DistSquared(Pair.Value.GetWaterQuadTree().GetTileRegion().GetCenter(), ViewPosition2D);
-			if (DistSquared < ClosestDistSquared)
-			{
-				ClosestDistSquared = DistSquared;
-				ClosestQuadTreeKey = Pair.Key;
-			}
+			ClosestDistSquared = DistSquared;
+			ClosestQuadTreeKey = Pair.Key;
 		}
-		return ClosestQuadTreeKey;
 	}
+	return ClosestQuadTreeKey;
 }
 
 TArray<int32, TInlineAllocator<8>> FWaterMeshSceneProxy::GetViewToQuadTreeMapping(const TArray<const FSceneView*>& Views, uint32 VisibilityMap) const
@@ -1095,9 +1096,7 @@ void FWaterMeshSceneProxy::GetDynamicRayTracingInstances(FRayTracingInstanceColl
 
 bool FWaterMeshSceneProxy::CreateViewWaterQuadTree(int32 Key, const FVector2D& CenterPosition)
 {
-#if !FORCE_SINGLE_WATER_QUADTREE
 	if (bIsLocalOnlyTessellationEnabled)
-#endif
 	{
 		if (!ViewQuadTrees.Contains(Key))
 		{
@@ -1111,9 +1110,7 @@ bool FWaterMeshSceneProxy::CreateViewWaterQuadTree(int32 Key, const FVector2D& C
 
 bool FWaterMeshSceneProxy::UpdateViewWaterQuadTree(int32 Key, const FVector2D& CenterPosition)
 {
-#if !FORCE_SINGLE_WATER_QUADTREE
 	if (bIsLocalOnlyTessellationEnabled)
-#endif
 	{
 		if (FViewWaterQuadTree* ViewWaterQuadTree = ViewQuadTrees.Find(Key))
 		{
@@ -1126,9 +1123,7 @@ bool FWaterMeshSceneProxy::UpdateViewWaterQuadTree(int32 Key, const FVector2D& C
 
 void FWaterMeshSceneProxy::DestroyViewWaterQuadTree(int32 Key)
 {
-#if !FORCE_SINGLE_WATER_QUADTREE
 	if (bIsLocalOnlyTessellationEnabled)
-#endif
 	{
 		ViewQuadTrees.Remove(Key);
 	}
diff --git a/Engine/Plugins/Experimental/Water/Source/Runtime/Private/WaterViewExtension.cpp b/Engine/Plugins/Experimental/Water/Source/Runtime/Private/WaterViewExtension.cpp
--- a/Engine/Plugins/Experimental/Water/Source/Runtime/Private/WaterViewExtension.cpp
+++ b/Engine/Plugins/Experimental/Water/Source/Runtime/Private/WaterViewExtension.cpp
@@ -10,6 +10,9 @@
 #include "WaterBodyManager.h"
 #include "WaterSubsystem.h"
 #include "RHIResourceUtils.h"
+#include "WaterMeshSceneProxy.h"
+#include "Engine/GameInstance.h"
+#include "WaterModule.h"
 
 static TAutoConsoleVariable<bool> CVarLocalTessellationFreeze(
 	TEXT("r.Water.WaterMesh.LocalTessellation.Freeze"),
@@ -33,6 +36,12 @@ static TAutoConsoleVariable<int32> CVarWaterInfoRenderMethod(
 	FConsoleVariableDelegate::CreateStatic(OnCVarWaterInfoSceneProxiesValueChanged),
 	ECVF_Default | ECVF_RenderThreadSafe);
 
+static TAutoConsoleVariable<bool> CVarDrawPerViewDebugInfo(
+	TEXT("r.Water.WaterInfo.DrawPerViewDebugInfo"),
+	false,
+	TEXT("Enable this to draw WaterZoneInfo debug information per view."),
+	ECVF_Default);
+
 // ----------------------------------------------------------------------------------
 
 FWaterMeshGPUWork GWaterMeshGPUWork;
@@ -54,6 +63,10 @@ void FWaterViewExtension::Initialize()
 	{
 		GerstnerWaterWaveSubsystem->Register(this);
 	}
+
+	CurrentNumViews = 0;
+
+	QuadTreeKeyLocationMap.Empty();
 }
 
 void FWaterViewExtension::Deinitialize()
@@ -68,6 +81,10 @@ void FWaterViewExtension::Deinitialize()
 	{
 		GerstnerWaterWaveSubsystem->Unregister(this);
 	}
+
+	CurrentNumViews = 0;
+
+	QuadTreeKeyLocationMap.Empty();
 }
 
 void FWaterViewExtension::UpdateGPUBuffers()
@@ -100,15 +117,15 @@ void FWaterViewExtension::UpdateGPUBuffers()
 
 		struct FWaterZoneData
 		{
-			FVector2f Location;
 			FVector2f Extent;
 			FVector2f HeightExtent;
 			float GroundZMin;
+			float bIsLocalOnlyTessellation;
 
-			float _Padding; // Unused
+			float _Padding[2]; // Unused
 
-			FWaterZoneData(const FVector2f& InLocation, const FVector2f& InExtent, const FVector2f& InHeightExtent, float InGroundZMin)
-				: Location(InLocation), Extent(InExtent), HeightExtent(InHeightExtent), GroundZMin(InGroundZMin) {}
+			FWaterZoneData(const FVector2f& InExtent, const FVector2f& InHeightExtent, float InGroundZMin, float bInIsLocalOnlyTessellation)
+				: Extent(InExtent), HeightExtent(InHeightExtent), GroundZMin(InGroundZMin), bIsLocalOnlyTessellation(bInIsLocalOnlyTessellation) {}
 		};
 		static_assert(sizeof(FWaterZoneData) == 2 * sizeof(FVector4f));
 
@@ -126,6 +143,17 @@ void FWaterViewExtension::UpdateGPUBuffers()
 		};
 		static_assert(sizeof(FGerstnerWaveData) == 2 * sizeof(FVector4f));
 
+		struct FWaterZoneViewData
+		{
+			FVector2f Location;
+
+			FVector2f _Padding; // Unused
+
+			FWaterZoneViewData(const FVector2f& InLocation)
+				: Location(InLocation) {}
+		};
+		static_assert(sizeof(FWaterZoneViewData) == sizeof(FVector4f));
+
 		// Water Body Data Buffer layout:
 		// -------------------------------------------------------------------------------
 		// || WaterZoneIndex | WaveDataIndex | NumWaves | (Other members) ||   ...   ||
@@ -133,28 +161,49 @@ void FWaterViewExtension::UpdateGPUBuffers()
 		//
 		// Water Aux Data Buffer layout:
 		// -----------------------------------------------------------------------------
-		// ||| WaterZone Data | ... || GerstnerWaveData | ... |||
+		// ||| WaterZone Data | ... || WaterZoneView Data | ... || GerstnerWaveData | ... |||
 		// -----------------------------------------------------------------------------
 		//
 
 		TArray<FWaterZoneData> WaterZoneData;
 
 		{
-			WaterBodyManager->ForEachWaterZone([&WaterZoneData](AWaterZone* WaterZone)
+			WaterBodyManager->ForEachWaterZone([&WaterZoneData, this](AWaterZone* WaterZone)
 			{
-				// #todo_water: LWC
-				const FVector2f ZoneLocation = FVector2f(FVector2D(WaterZone->GetDynamicWaterInfoCenter()));
-
 				const FVector2f ZoneExtent = FVector2f(FVector2D(WaterZone->GetDynamicWaterInfoExtent()));
 				const FVector2f WaterHeightExtents = WaterZone->GetWaterHeightExtents();
 				const float GroundZMin = WaterZone->GetGroundZMin();
+				const float bIsLocalOnlyTessellation = WaterZone->IsLocalOnlyTessellationEnabled() ? 1.0f : -1.0f;
 
-				WaterZoneData.Emplace(ZoneLocation, ZoneExtent, WaterHeightExtents, GroundZMin);
+				WaterZoneData.Emplace(ZoneExtent, WaterHeightExtents, GroundZMin, bIsLocalOnlyTessellation);
 
 				return true;
 			});
 		}
 
+		// We store views packed per zone:
+		// ||| Zone0View0 | Zone0View1 | ... | Zone0ViewN ||...|| ZoneNView0 | ZoneNView1 | ... | ZoneNViewN |||
+		TArray<FWaterZoneViewData> WaterZoneViewData;
+		{
+			WaterBodyManager->ForEachWaterZone([&WaterZoneViewData, this](AWaterZone* WaterZone)
+			{
+				FWaterZoneInfo* WaterZoneInfo = WaterZoneInfos.Find(WaterZone);
+
+				check(WaterZoneInfo != nullptr);
+
+				if (WaterZoneInfo != nullptr)
+				{
+					for (const FWaterZoneInfo::FWaterZoneViewInfo& ViewInfos : WaterZoneInfo->ViewInfos)
+					{
+						// #todo_water: LWC
+						FVector2f ZoneViewLocation = FVector2f(FVector2D(ViewInfos.Center));
+						WaterZoneViewData.Emplace(ZoneViewLocation);
+					}
+				}
+				
+				return true;
+			});
+		}
 
 		TArray<FWaterBodyData> WaterBodyData;
 		TArray<FGerstnerWaveData> WaveData;
@@ -226,11 +275,11 @@ void FWaterViewExtension::UpdateGPUBuffers()
 		TArray<FVector4f> WaterBodyDataBuffer;
 		TArray<FVector4f> WaterAuxDataBuffer;
 
-		// The first element of the WaterDataBuffer contains the offsets to each of the sub-buffers.
+		// The first element of the WaterDataBuffer contains the offsets to each of the sub-buffers and the number of view data.
 		// X = WaterZoneDataOffset
 		// Y = WaterWaveDataOffset
-		// Z = Unused
-		// W = Unused
+		// Z = WaterViewDataOffset
+		// W = NumWaterViewData
 		WaterAuxDataBuffer.AddZeroed();
 
 		// Transform the individual arrays into the single buffer:
@@ -251,6 +300,7 @@ void FWaterViewExtension::UpdateGPUBuffers()
 			AppendDataToFloat4Buffer(WaterBodyDataBuffer, WaterBodyData);
 
 			const int32 ZoneDataOffset = AppendDataToFloat4Buffer(WaterAuxDataBuffer, WaterZoneData);
+			const int32 ZoneViewDataOffset = AppendDataToFloat4Buffer(WaterAuxDataBuffer, WaterZoneViewData);
 			const int32 WaveDataOffset = AppendDataToFloat4Buffer(WaterAuxDataBuffer, WaveData);
 
 			// Store the offsets to each sub-buffer in the first entry.
@@ -258,8 +308,8 @@ void FWaterViewExtension::UpdateGPUBuffers()
 			FVector4f& OffsetData = WaterAuxDataBuffer[0];
 			OffsetData.X = ZoneDataOffset;
 			OffsetData.Y = WaveDataOffset;
-			OffsetData.Z = 0.0f;
-			OffsetData.W = 0.0f;
+			OffsetData.Z = ZoneViewDataOffset;
+			OffsetData.W = CurrentNumViews;
 		}
 
 		if (WaterBodyDataBuffer.Num() == 0)
@@ -295,11 +345,110 @@ void FWaterViewExtension::UpdateGPUBuffers()
 	}
 }
 
+int32 FWaterViewExtension::GetOrAddViewindex(const FSceneView& InView)
+{
+	const int32 ViewPlayerIndex = (InView.PlayerIndex != INDEX_NONE) ? InView.PlayerIndex : 0;
+
+	return ViewPlayerIndices.AddUnique(ViewPlayerIndex);
+}
+
+int32 FWaterViewExtension::GetViewIndex(int32 PlayerIndex) const
+{
+	int32 OutIndex = INDEX_NONE;
+
+	if (ViewPlayerIndices.Find(PlayerIndex, OutIndex))
+	{
+		return OutIndex;
+	}
+
+	return OutIndex;
+}
+
+int32 FWaterViewExtension::GetViewIndex(const FSceneView& InView) const
+{
+	const int32 PlayerIndex = (InView.PlayerIndex != INDEX_NONE) ? InView.PlayerIndex : 0;
+	return GetViewIndex(PlayerIndex);
+}
+
 void FWaterViewExtension::SetupViewFamily(FSceneViewFamily& InViewFamily)
 {
 
 }
 
+void FWaterViewExtension::UpdateViewInfo(AWaterZone* WaterZone, FSceneView& InView)
+{
+	const int32 ViewPlayerIndex = GetViewIndex(InView);
+	check(ViewPlayerIndex != INDEX_NONE);
+
+	check(WaterZone != nullptr);
+
+	const FVector ViewLocation = InView.ViewLocation;
+
+	FWaterZoneInfo* WaterZoneInfo = WaterZoneInfos.Find(WaterZone);
+	if (ensureMsgf(WaterZoneInfo != nullptr, TEXT("We are trying to render a water info texture for a water zone that is not registered!")))
+	{
+		FWaterZoneInfo::FWaterZoneViewInfo& WaterZoneViewInfo = WaterZoneInfo->ViewInfos[ViewPlayerIndex];
+
+		if (WaterZone->IsLocalOnlyTessellationEnabled())
+		{
+			UWaterMeshComponent* WaterMesh = WaterZone->GetWaterMeshComponent();
+			check(WaterMesh);
+
+			const double TileSize = WaterMesh->GetTileSize();
+
+			// UWaterMeshComponent::GetGlobalWaterMeshCenter()  already does some snapping, also taking into account
+			// r.Water.WaterMesh.LODCountBias logic for the current sliding window. Should we apply this here too? 
+			// TODO: Look into the above.
+			FVector SlidingWindowCenter = ViewLocation.GridSnap(TileSize);
+
+			const FVector2D SlidingWindowHalfExtent(WaterZone->GetDynamicWaterInfoExtent() / 2.0);
+
+			WaterZoneViewInfo.Center = SlidingWindowCenter;
+
+			// Trigger the next update when the camera is <UpdateMargin> units away from the border of the current window.
+			const FVector2D UpdateMargin(CVarLocalTessellationUpdateMargin.GetValueOnGameThread());
+
+			// Keep a minimum of <1., 1.> bounds to avoid updating every frame if the update margin is larger than the zone.
+			const FVector2D UpdateExtents = FVector2D::Max(FVector2D(1., 1.), SlidingWindowHalfExtent - UpdateMargin);
+			WaterZoneViewInfo.UpdateBounds.Emplace(FVector2D(SlidingWindowCenter) - UpdateExtents, FVector2D(SlidingWindowCenter) + UpdateExtents);
+
+			// Mark GPU data dirty since we have a new WaterArea parameter and need to push this to water bodies.
+			MarkGPUDataDirty();
+		}
+		else
+		{
+			WaterZoneViewInfo.UpdateBounds.Reset();
+			WaterZoneViewInfo.Center = WaterZone->GetActorLocation();
+			MarkGPUDataDirty();
+		}
+	}
+}
+
+void FWaterViewExtension::RenderWaterInfoTexture(FSceneViewFamily& InViewFamily, FSceneView& InView, const FWaterZoneInfo* WaterZoneInfo, FSceneInterface* Scene, const FVector& ZoneCenter)
+{
+	const int32 WaterInfoRenderMethod = CVarWaterInfoRenderMethod.GetValueOnGameThread();
+
+	int32 ViewPlayerIndex = GetViewIndex(InView);
+	check(ViewPlayerIndex != INDEX_NONE);
+
+	const UE::WaterInfo::FRenderingContext& Context(WaterZoneInfo->RenderContext);
+	// Old method of rendering the water info texture; uses scene captures
+	if (WaterInfoRenderMethod == 0)
+	{
+		UE::WaterInfo::UpdateWaterInfoRendering(Scene, Context, ZoneCenter);
+	}
+	// Render the water info texture using custom render pass method
+	else if (WaterInfoRenderMethod == 2)
+	{
+		UE::WaterInfo::UpdateWaterInfoRendering_CustomRenderPass(Scene, InViewFamily, Context, ViewPlayerIndex, ZoneCenter);
+	}
+	// Rendering is done in a separate pass when rendering the main view
+	else if (WaterInfoRenderMethod == 1)
+	{
+		UE::WaterInfo::UpdateWaterInfoRendering2(InView, Context, ViewPlayerIndex, ZoneCenter);
+	}
+}
+
 void FWaterViewExtension::SetupView(FSceneViewFamily& InViewFamily, FSceneView& InView)
 {
 	if (CVarLocalTessellationFreeze.GetValueOnGameThread())
@@ -316,6 +465,13 @@ void FWaterViewExtension::SetupView(FSceneViewFamily& InViewFamily, FSceneView&
 		return;
 	}
 
+	int32 NumViews = 1;
+
+	if (WorldPtr->GetGameInstance())
+	{
+		NumViews = WorldPtr->GetGameInstance()->GetLocalPlayers().Num();
+	}
+
 	// Prevent re-entrancy. 
 	// Since the water info render will update the view extensions we could end up with a re-entrant case.
 	static bool bUpdatingWaterInfo = false;
@@ -326,110 +482,233 @@ void FWaterViewExtension::SetupView(FSceneViewFamily& InViewFamily, FSceneView&
 	bUpdatingWaterInfo = true;
 	ON_SCOPE_EXIT { bUpdatingWaterInfo = false; };
 
-	const int32 WaterInfoRenderMethod = CVarWaterInfoRenderMethod.GetValueOnGameThread();
+	int32 ViewPlayerIndex = GetOrAddViewindex(InView);
+
 	FSceneInterface* Scene = WorldPtr.Get()->Scene;
 	check(Scene != nullptr);
 
-	for (const TPair<TWeakObjectPtr<AWaterZone>, UE::WaterInfo::FRenderingContext>& Pair : WaterInfoContextsToRender)
+	// if a texture rebuild is pending, we need to wait for that to be done by the WaterZone, so skip any view update while that is completed.
+	if (bWaterInfoTextureRebuildPending)
 	{
-		TWeakObjectPtr<AWaterZone> WaterZonePtr = Pair.Key;
-		if (AWaterZone* WaterZone = WaterZonePtr.Get())
-		{
-			check(WaterZone != nullptr);
+		ViewPlayerIndices.Empty();
+		NonDataViewsQuadtreeKeys.Empty();
+		return;
+	}
+
+	// Don't dirty the water info texture when we're rendering from a scene capture. Due to the frame delay after marking the texture as dirty, scene captures wouldn't have the right texture anyways.
+	// #todo_water [roey]: Once we have no frame-delay for updating the texture and lesser performance impact, we can re-enable updates within scene captures.
+	bool bShouldHaveWaterZoneViewData = !InView.bIsSceneCapture && !InView.bIsSceneCaptureCube && !InView.bIsReflectionCapture && !InView.bIsPlanarReflection && !InView.bIsVirtualTexture;
 
-			FWaterZoneInfo* WaterZoneInfo = WaterZoneInfos.Find(WaterZone);
-			if (ensureMsgf(WaterZoneInfo != nullptr, TEXT("We are trying to rendering a water info texture for a water zone that is not registered!")))
+	if (bShouldHaveWaterZoneViewData)
+	{
+		if (CurrentNumViews != NumViews)
+		{
+			for (AWaterZone* WaterZone : TActorRange<AWaterZone>(WorldPtr.Get()))
 			{
-				if (WaterZone->IsLocalOnlyTessellationEnabled())
+				if (WaterZone->HasActorRegisteredAllComponents())
 				{
-					UWaterMeshComponent* WaterMesh= WaterZone->GetWaterMeshComponent();
-					check(WaterMesh);
+					FWaterZoneInfo& WaterZoneInfo = WaterZoneInfos.FindChecked(WaterZone);
+
+					// make sure that if !IsLocalOnlyTessellationEnabled we only create a single WaterInfoTexture slice.
+					WaterZone->WaterInfoTextureArrayNumSlices = WaterZone->IsLocalOnlyTessellationEnabled() ? NumViews : 1;
+					// Mark for rebuild the WaterZone if the size changes
+					WaterZone->MarkForRebuild(EWaterZoneRebuildFlags::UpdateWaterInfoTexture);
+
+					// init per view info
+
+					WaterZoneInfo.ViewInfos.Empty();
 
-					const double TileSize = WaterMesh->GetTileSize();
+					for (int i = 0; i < NumViews; ++i)
+					{
+						WaterZoneInfo.ViewInfos.Emplace(FWaterZoneInfo::FWaterZoneViewInfo());
+					}
+				}
+
+				UE_LOG(LogWater, Verbose, TEXT("Number of views changed. Water Zone (%s) ViewInfos was reset."), *GetNameSafe(WaterZone));
+			}
+
+			CurrentNumViews = NumViews;
+			bWaterInfoTextureRebuildPending = true;
+
+			return;
+		}
+
+		// Check if the view location is no longer within the current update bounds of a water zone and if so, queue an update for it.
+		for (AWaterZone* WaterZone : TActorRange<AWaterZone>(WorldPtr.Get()))
+		{
+			if (WaterZone->HasActorRegisteredAllComponents())
+			{
+				FWaterZoneInfo& WaterZoneInfo = WaterZoneInfos.FindChecked(WaterZone);
 
-					const FVector WaterMeshCenter(ViewLocation.GridSnap(TileSize));
+				if (WaterZone->WaterInfoTextureArray.Get() == nullptr)
+				{
+					continue;
+				}
 
-					const FVector2D WaterInfoHalfExtent(WaterZone->GetDynamicWaterInfoExtent() / 2.0);
-					const FBox2D WaterInfoBounds(FVector2D(WaterMeshCenter) - WaterInfoHalfExtent, FVector2D(WaterMeshCenter) + WaterInfoHalfExtent);
+				if (CVarDrawPerViewDebugInfo.GetValueOnGameThread())
+				{
+					DrawDebugInfo(InView, WaterZone);
+				}
 
-					WaterZone->SetLocalTessellationCenter(WaterMeshCenter);
+				FWaterZoneInfo::FWaterZoneViewInfo& WaterZoneViewInfo = WaterZoneInfo.ViewInfos[ViewPlayerIndex];
 
-					// Trigger the next update when the camera is <UpdateMargin> units away from the border of the current window.
-					const FVector2D UpdateMargin(CVarLocalTessellationUpdateMargin.GetValueOnGameThread());
+				bool bBoundsUpdateNeeded = WaterZone->IsLocalOnlyTessellationEnabled() && (!WaterZoneViewInfo.UpdateBounds.IsSet() || !WaterZoneViewInfo.UpdateBounds->IsInside(FVector2D(ViewLocation)));
+				bBoundsUpdateNeeded |= WaterZoneViewInfo.bIsDirty;
 
-					// Keep a minimum of <1., 1.> bounds to avoid updating every frame if the update margin is larger than the zone.
-					const FVector2D UpdateExtents = FVector2D::Max(FVector2D(1., 1.), WaterInfoHalfExtent - UpdateMargin);
-					WaterZoneInfo->UpdateBounds.Emplace(FVector2D(WaterMeshCenter) - UpdateExtents, FVector2D(WaterMeshCenter) + UpdateExtents);
+				if (bBoundsUpdateNeeded)
+				{
+					UpdateViewInfo(WaterZone, InView);
 
-					const FVector2D WaterQuadTreeHalfExtent = WaterMesh->GetExtentInTiles() * TileSize;
-					const FVector2D WaterQuadTreeCenter = WaterMesh->GetDynamicWaterMeshCenter();
-					const FBox2D WaterQuadTreeBounds(WaterQuadTreeCenter - WaterQuadTreeHalfExtent, WaterQuadTreeCenter + WaterQuadTreeHalfExtent);
+					UE_LOG(LogWater, Verbose, TEXT("Water Zone (%s) ViewInfo for view %d updated."), *GetNameSafe(WaterZone), ViewPlayerIndex);
 
-					// If the new water info bounds would be outside the water mesh, recenter and rebuild the water mesh
-					if (!WaterQuadTreeBounds.IsInside(WaterInfoBounds.ExpandBy(WaterInfoBounds.GetExtent())))
+					// make sure that if !IsLocalOnlyTessellationEnabled, we only update the WaterInfoTexture for a single view
+					if (WaterZone->IsLocalOnlyTessellationEnabled() || ViewPlayerIndex == 0)
 					{
-						WaterMesh->SetDynamicWaterMeshCenter(FVector2D(WaterMeshCenter));
+						RenderWaterInfoTexture(InViewFamily, InView, &WaterZoneInfo, Scene, WaterZoneViewInfo.Center);
 					}
-
-					// Mark GPU data dirty since we have a new WaterArea parameter and need to push this to water bodies.
-					MarkGPUDataDirty();
 				}
-				else
+
+				UWaterMeshComponent* WaterMeshComponent = WaterZone->GetWaterMeshComponent();
+				check(WaterMeshComponent != nullptr);
+
+				FWaterMeshSceneProxy*  SceneProxy = (FWaterMeshSceneProxy*) WaterMeshComponent->GetSceneProxy();
+				check(SceneProxy != nullptr);
+
+				// push a quadtree update
+				if (WaterZone->IsLocalOnlyTessellationEnabled() && (bBoundsUpdateNeeded || SceneProxy != WaterZoneViewInfo.OldSceneProxy))
 				{
-					WaterZoneInfo->UpdateBounds.Reset();
+					FScopeLock AddQuadtreeUpdate(&QuadtreeUpdateLock);
+
+					FQuadtreeUpdateInfo QuadtreeUpdate;
+					QuadtreeUpdate.SceneProxy = (FWaterMeshSceneProxy*)SceneProxy;
+					QuadtreeUpdate.Location = WaterZoneViewInfo.UpdateBounds->GetCenter();
+					// we use the actual PlayerIndex instead of ViewPlayerIndex, since the latter can change if more views are added while a quadtree already exists
+					QuadtreeUpdate.Key = InView.PlayerIndex;
+
+					QuadtreeUpdates.Add(QuadtreeUpdate);
+
+
+					if (!QuadTreeKeyLocationMap.Contains(QuadtreeUpdate.Key))
+					{
+						QuadTreeKeyLocationMap.Add(QuadtreeUpdate.Key, QuadtreeUpdate.Location);
+					}
+					else
+					{
+						QuadTreeKeyLocationMap[QuadtreeUpdate.Key] = QuadtreeUpdate.Location;
+					}
+
+					UE_LOG(LogWater, Verbose, TEXT("Water Zone (%s) queued a quadtree update for view %d updated."), *GetNameSafe(WaterZone), ViewPlayerIndex);
 				}
-			}
 
-			const UE::WaterInfo::FRenderingContext& Context(Pair.Value);
-			// Old method of rendering the water info texture; uses scene captures
-			if (WaterInfoRenderMethod == 0)
-			{
-				UE::WaterInfo::UpdateWaterInfoRendering(Scene, Context);
-			}
-			// Render the water info texture using custom render pass method
-			else if (WaterInfoRenderMethod == 2)
-			{
-				UE::WaterInfo::UpdateWaterInfoRendering_CustomRenderPass(Scene, InViewFamily, Context);
+				WaterZoneViewInfo.OldSceneProxy = SceneProxy;
+				WaterZoneViewInfo.bIsDirty = false;
 			}
 		}
 	}
-
-	// New method of rendering the water info texture; rendering is done in a separate pass when rendering the main view
-	if (WaterInfoRenderMethod == 1)
+	else
 	{
-		UE::WaterInfo::UpdateWaterInfoRendering2(InView, WaterInfoContextsToRender);
-	}
-
-	WaterInfoContextsToRender.Empty();
+		FScopeLock AddNonDataViewsQuadtreeKeys(&QuadtreeUpdateLock);
 
-	// Don't dirty the water info texture when we're rendering from a scene capture. Due to the frame delay after marking the texture as dirty, scene captures wouldn't have the right texture anyways.
-	// #todo_water [roey]: Once we have no frame-delay for updating the texture and lesser performance impact, we can re-enable updates within scene captures.
-	if (!InView.bIsSceneCapture && !InView.bIsSceneCaptureCube && !InView.bIsReflectionCapture && !InView.bIsPlanarReflection && !InView.bIsVirtualTexture)
-	{
-		// Check if the view location is no longer within the current update bounds of a water zone and if so, queue an update for it.
+		// if !bShouldHaveWaterZoneViewData, we want to store the closest quadtree ID for this view (since it is not
+		// going to have its own quadtree associated with it), so we can assign the WaterInfoTexture index of the view assigned to that quadtree
 		for (AWaterZone* WaterZone : TActorRange<AWaterZone>(WorldPtr.Get()))
 		{
+			if (WaterZone->WaterInfoTextureArray.Get() == nullptr)
+			{
+				continue;
+			}
+
 			if (WaterZone->HasActorRegisteredAllComponents())
 			{
-				FWaterZoneInfo& WaterZoneInfo = WaterZoneInfos.FindChecked(WaterZone);
 				if (WaterZone->IsLocalOnlyTessellationEnabled())
 				{
-					if (!WaterZoneInfo.UpdateBounds.IsSet() || !WaterZoneInfo.UpdateBounds->IsInside(FVector2D(ViewLocation)))
-					{
-						WaterZone->MarkForRebuild(EWaterZoneRebuildFlags::UpdateWaterInfoTexture);
-					}
+					UWaterMeshComponent* WaterMeshComponent = WaterZone->GetWaterMeshComponent();
+					check(WaterMeshComponent != nullptr);
+
+					FWaterMeshSceneProxy* SceneProxy = (FWaterMeshSceneProxy*)WaterMeshComponent->GetSceneProxy();
+					check(SceneProxy != nullptr);
+
+					NonDataViewsQuadtreeKeys.Add(InView.State, SceneProxy->FindBestQuadTreeForViewLocation(&InView));
+
+					UE_LOG(LogWater, Verbose, TEXT("Water Zone (%s) queued a search for the closest quadtree for a View (0x%p) which has no WaterInfo."), *GetNameSafe(WaterZone), &InView);
 				}
 			}
-		}
+		}	
 	}
 
 	// The logic in UpdateGPUBuffers() used to be done in SetupViewFamily(). However, SetupView() (which is responsible for water info rendering) potentially modifies the WaterZone but is called after SetupViewFamily().
 	// This can lead to visual artifacts due to outdated data in the GPU buffers.
+	
 	UpdateGPUBuffers();
+
+	// this is needed because when switching back from PIE to Editor, there isn't a proper initialization of
+	// the editor WaterViewExtension, since during PIE it is kept alive and not updated. Because of that when
+	// PIE ends and we start updating the editor view extension, it can contain out of date values which can cause
+	// water rendering artifacts until stepping out of bounds forces the first update. To fix that we keep
+	// track if the last ViewExtension that was updated changed, so we can reset CurrentNumViews which forces an update on the next frame.
+	// TODO: look into any existing callbacks which would allow to achieve this in a cleaner way.
+	static FWaterViewExtension* LastWaterViewExtensionUpdated = nullptr;
+	if (LastWaterViewExtensionUpdated != this)
+	{
+		CurrentNumViews = 0;
+		LastWaterViewExtensionUpdated = this;
+	}
+}
+
+void FWaterViewExtension::DrawDebugInfo(FSceneView& InView, AWaterZone* WaterZone)
+{
+	if (GEngine)
+	{
+		FWaterZoneInfo& WaterZoneInfo = WaterZoneInfos.FindChecked(WaterZone);
+
+		const FVector ViewLocation = InView.ViewLocation;
+		int32 ViewPlayerIndex = GetViewIndex(InView);
+		check(ViewPlayerIndex != INDEX_NONE);
+
+		FWaterZoneInfo::FWaterZoneViewInfo& WaterZoneViewInfo = WaterZoneInfo.ViewInfos[ViewPlayerIndex];
+
+		FVector2D Center = FVector2D::Zero();
+		FVector2D Extents = FVector2D::Zero();
+		FVector SnappedCenter = WaterZoneViewInfo.Center;
+
+		if (WaterZoneViewInfo.UpdateBounds.IsSet())
+		{
+			WaterZoneViewInfo.UpdateBounds.GetValue().GetCenterAndExtents(Center, Extents);
+		}
+
+		FColor Colors[4] = { FColor::Yellow, FColor::Green, FColor::Red, FColor::Blue };
+
+		if (ViewPlayerIndex < 4)
+		{
+			int32 StartingOffset = ViewPlayerIndex * 3;
+			GEngine->AddOnScreenDebugMessage(StartingOffset, 0.01f, Colors[ViewPlayerIndex], FString::Printf(TEXT("View %d Location: %f, %f, %f"), ViewPlayerIndex, ViewLocation.X, ViewLocation.Y, ViewLocation.Z));
+			GEngine->AddOnScreenDebugMessage(StartingOffset + 1, 0.01f, Colors[ViewPlayerIndex], FString::Printf(TEXT("Center: %f, %f; Extents: %f, %f"), Center.X, Center.Y, Extents.X, Extents.Y));
+			GEngine->AddOnScreenDebugMessage(StartingOffset + 2, 0.01f, Colors[ViewPlayerIndex], FString::Printf(TEXT("Snapped Center: %f, %f"), SnappedCenter.X, SnappedCenter.Y));
+		}
+	}
 }
 
 void FWaterViewExtension::PreRenderViewFamily_RenderThread(FRDGBuilder& GraphBuilder, FSceneViewFamily& InViewFamily)
 {
+	// consume quadtree updates
+	{
+		FScopeLock ConsumeQuadtreeUpdates(&QuadtreeUpdateLock);
+
+		for (const FQuadtreeUpdateInfo& QuadtreeUpdate : QuadtreeUpdates)
+		{
+			// TODO: We could add a CreateOrUpdateViewWaterQuadTree to WaterSceneProxy
+			// if the creation fails it means a quadtree for the current view already exists, so we just update it
+			if (!QuadtreeUpdate.SceneProxy->CreateViewWaterQuadTree(QuadtreeUpdate.Key, QuadtreeUpdate.Location))
+			{
+				bool bResult = QuadtreeUpdate.SceneProxy->UpdateViewWaterQuadTree(QuadtreeUpdate.Key, QuadtreeUpdate.Location);
+
+				check(bResult);
+			}
+		}
+
+		QuadtreeUpdates.Empty();
+	}
+	
 }
 
 void FWaterViewExtension::PreRenderView_RenderThread(FRDGBuilder& GraphBuilder, FSceneView& InView)
@@ -440,8 +719,25 @@ void FWaterViewExtension::PreRenderView_RenderThread(FRDGBuilder& GraphBuilder,
 		InView.WaterDataBuffer = WaterGPUData->AuxDataSRV;
 		InView.WaterIndirectionBuffer = WaterGPUData->WaterBodyDataSRV;
 	}
-	// TODO: Look up an actually unique index for the given view.
-	InView.WaterInfoTextureViewIndex = 0;
+
+	{
+		FScopeLock ReadNonDataViewsQuadtreeKeys(&QuadtreeUpdateLock);
+
+		int32 WaterInfoTextureIndex = GetViewIndex(InView);
+
+		// check if this view is one of the queues without its own WaterZoneViewData, if it is assign the 
+		// WaterInfoTextureIndex corresponding to the closest quadtree index we stored in NonDataViewsQuadtreeKeys
+		const int32* QuadtreeKey = NonDataViewsQuadtreeKeys.Find(InView.State);
+		if (QuadtreeKey != nullptr)
+		{
+			// swap the index with the one corresponding to the closest quadtree
+			WaterInfoTextureIndex = *QuadtreeKey;
+		}
+
+		WaterInfoTextureIndex = (WaterInfoTextureIndex != INDEX_NONE) ? WaterInfoTextureIndex : 0;
+
+		InView.WaterInfoTextureViewIndex = WaterInfoTextureIndex;
+	}
 }
 
 void FWaterViewExtension::PreRenderBasePass_RenderThread(FRDGBuilder& GraphBuilder, bool bDepthBufferIsPopulated)
@@ -454,8 +750,23 @@ void FWaterViewExtension::PreRenderBasePass_RenderThread(FRDGBuilder& GraphBuild
 
 void FWaterViewExtension::MarkWaterInfoTextureForRebuild(const UE::WaterInfo::FRenderingContext& RenderContext)
 {
-	WaterInfoContextsToRender.Emplace(RenderContext.ZoneToRender, RenderContext);
 	MarkGPUDataDirty();
+
+	bWaterInfoTextureRebuildPending = false;
+
+	// this should mark dirty the RenderContext.ZoneToRender waterzone info
+
+	FWaterZoneInfo* WaterZoneInfo = WaterZoneInfos.Find(RenderContext.ZoneToRender);
+
+	if (WaterZoneInfo != nullptr)
+	{
+		WaterZoneInfo->RenderContext = RenderContext;
+
+		for (int i = 0; i < WaterZoneInfo->ViewInfos.Num(); ++i)
+		{
+			WaterZoneInfo->ViewInfos[i].bIsDirty = true;
+		}
+	}
 }
 
 void FWaterViewExtension::MarkGPUDataDirty()
@@ -466,12 +777,46 @@ void FWaterViewExtension::MarkGPUDataDirty()
 void FWaterViewExtension::AddWaterZone(AWaterZone* InWaterZone)
 {
 	check(!WaterZoneInfos.Contains(InWaterZone));
-	WaterZoneInfos.Emplace(InWaterZone);
+	FWaterZoneInfo& WaterZoneInfo = WaterZoneInfos.Emplace(InWaterZone);
+
+	// init per view info
+
+	CurrentNumViews = 0;
+	WaterZoneInfo.ViewInfos.Emplace(FWaterZoneInfo::FWaterZoneViewInfo());
+
+	UE_LOG(LogWater, Verbose, TEXT("Water Zone (%s): AddWaterZone was called."), *GetNameSafe(InWaterZone));
 }
 
 void FWaterViewExtension::RemoveWaterZone(AWaterZone* InWaterZone)
 {
 	WaterZoneInfos.FindAndRemoveChecked(InWaterZone);
-	// Remove any WaterInfoContexts that might be trying to rendering for this water zone this frame. They are no longer valid.
-	WaterInfoContextsToRender.Remove(InWaterZone);
+
+	UE_LOG(LogWater, Verbose, TEXT("Water Zone (%s): RemoveWaterZone was called."), *GetNameSafe(InWaterZone));
+}
+
+FVector FWaterViewExtension::GetZoneLocation(AWaterZone* InWaterZone, int32 PlayerIndex) const
+{
+	if (ensure(InWaterZone->HasActorRegisteredAllComponents()))
+	{
+		const FWaterZoneInfo* WaterZoneInfo = WaterZoneInfos.Find(InWaterZone);
+		if (ensure(WaterZoneInfo))
+		{
+			int32 ViewPlayerIndex = GetViewIndex(PlayerIndex);
+
+			if (ensure(ViewPlayerIndex != INDEX_NONE && ViewPlayerIndex < WaterZoneInfo->ViewInfos.Num()))
+			{
+				return WaterZoneInfo->ViewInfos[ViewPlayerIndex].Center;
+			}		
+		}
+	}
+
+	return FVector();
+}
+
+void FWaterViewExtension::CreateSceneProxyQuadtrees(FWaterMeshSceneProxy* SceneProxy)
+{
+	for (auto& Pair : QuadTreeKeyLocationMap)
+	{
+		SceneProxy->CreateViewWaterQuadTree(Pair.Key, Pair.Value);
+	}
 }
diff --git a/Engine/Plugins/Experimental/Water/Source/Runtime/Private/WaterZoneActor.cpp b/Engine/Plugins/Experimental/Water/Source/Runtime/Private/WaterZoneActor.cpp
--- a/Engine/Plugins/Experimental/Water/Source/Runtime/Private/WaterZoneActor.cpp
+++ b/Engine/Plugins/Experimental/Water/Source/Runtime/Private/WaterZoneActor.cpp
@@ -128,6 +128,31 @@ FBox2D AWaterZone::GetZoneBounds2D() const
 	return WaterZoneBounds;
 }
 
+void AWaterZone::GetDynamicWaterInfoBounds(TArray<FBox>& Bounds) const
+{
+	if (FWaterViewExtension* WaterViewExtension = UWaterSubsystem::GetWaterViewExtension(GetWorld()))
+	{
+		FWaterViewExtension::FWaterZoneInfo* WaterZoneInfo = WaterViewExtension->WaterZoneInfos.Find(this);
+
+		if (WaterZoneInfo != nullptr)
+		{
+			for (const FWaterViewExtension::FWaterZoneInfo::FWaterZoneViewInfo& ViewInfo : WaterZoneInfo->ViewInfos)
+			{
+				FVector Center(ViewInfo.Center);
+
+				if (!IsLocalOnlyTessellationEnabled())
+				{
+					Center = GetActorLocation();
+				}
+
+				const FVector HalfExtent(GetDynamicWaterInfoExtent() / 2);
+
+				Bounds.Add(FBox(Center - HalfExtent, Center + HalfExtent));
+			}
+		}
+	}
+}
+
 FBox AWaterZone::GetZoneBounds() const
 {
 	const FBox2D ZoneBounds2D = GetZoneBounds2D();
@@ -218,9 +243,31 @@ void AWaterZone::DeclareConstructClasses(TArray<FTopLevelAssetPath>& OutConstruc
 }
 #endif
 
+void AWaterZone::MarkForRebuild(EWaterZoneRebuildFlags Flags, const FBox2D& UpdateRegion, const FBox& WaterInfoBounds, const UObject* DebugRequestingObject)
+{
+	const FBox2D WaterInfoBounds2D(FVector2D(WaterInfoBounds.Min), FVector2D(WaterInfoBounds.Max));
+
+	// Suppress updates which occur outside the bounds of the water zone.
+	if ((!UpdateRegion.bIsValid || UpdateRegion.Intersect(WaterInfoBounds2D)))
+	{
+		if (EnumHasAnyFlags(Flags, EWaterZoneRebuildFlags::UpdateWaterMesh))
+		{
+			UE_LOG(LogWater, Verbose, TEXT("AWaterZone (%s) UpdateWaterMesh in region {%s} (triggered by %s)"), *GetNameSafe(this), *UpdateRegion.ToString(), *GetNameSafe(DebugRequestingObject));
+			WaterMesh->MarkWaterMeshGridDirty();
+			WaterMesh->MarkRenderStateDirty();
+		}
+		if (EnumHasAnyFlags(Flags, EWaterZoneRebuildFlags::UpdateWaterInfoTexture))
+		{
+			UE_LOG(LogWater, Verbose, TEXT("AWaterZone (%s) UpdateWaterInfoTexture in region {%s} (triggered by %s)"), *GetNameSafe(this), *UpdateRegion.ToString(), *GetNameSafe(DebugRequestingObject));
+			bNeedsWaterInfoRebuild = true;
+		}
+	}
+}
+
 void AWaterZone::MarkForRebuild(EWaterZoneRebuildFlags Flags, const FBox2D& UpdateRegion, const UObject* DebugRequestingObject)
 {
-	const FBox WaterInfoBounds = GetDynamicWaterInfoBounds();
+	// TODO: investigate why optimizing by updating only checking the bounds of the active views dynamic extents is not working.
+	const FBox WaterInfoBounds = GetZoneBounds();
 	const FBox2D WaterInfoBounds2D(FVector2D(WaterInfoBounds.Min), FVector2D(WaterInfoBounds.Max));
 
 	// Suppress updates which occur outside the bounds of the water zone.
@@ -649,7 +696,7 @@ bool AWaterZone::UpdateWaterInfoTexture()
 
 		const ETextureRenderTargetFormat Format = bHalfPrecisionTexture ? ETextureRenderTargetFormat::RTF_RGBA16f : RTF_RGBA32f;
 		UTextureRenderTarget2DArray* OldTexture = WaterInfoTextureArray;
-		WaterInfoTextureArray = FWaterUtils::GetOrCreateTransientRenderTarget2DArray(OldTexture, TEXT("WaterInfoTexture"), RenderTargetResolution, 1, Format);
+		WaterInfoTextureArray = FWaterUtils::GetOrCreateTransientRenderTarget2DArray(OldTexture, TEXT("WaterInfoTexture"), RenderTargetResolution, WaterInfoTextureArrayNumSlices, Format);
 
 		// The water info texture is different, we need to bind the newly created texture to all registered water bodies
 		if (WaterInfoTextureArray != OldTexture)
@@ -681,18 +728,6 @@ bool AWaterZone::UpdateWaterInfoTexture()
 	return true;
 }
 
-FVector AWaterZone::GetDynamicWaterInfoCenter() const
-{
-	if (IsLocalOnlyTessellationEnabled())
-	{
-		return LocalTessellationCenter;
-	}
-	else
-	{
-		return GetActorLocation();
-	}
-}
-
 FVector AWaterZone::GetDynamicWaterInfoExtent() const
 {
 	if (IsLocalOnlyTessellationEnabled())
@@ -706,11 +741,20 @@ FVector AWaterZone::GetDynamicWaterInfoExtent() const
 	}
 }
 
-FBox AWaterZone::GetDynamicWaterInfoBounds() const
+FVector AWaterZone::GetDynamicWaterInfoCenter(int32 PlayerIndex)
 {
-	const FVector WaterInfoCenter(GetDynamicWaterInfoCenter());
-	const FVector WaterInfoHalfExtent(GetDynamicWaterInfoExtent() / 2);
-	return FBox(WaterInfoCenter - WaterInfoHalfExtent, WaterInfoCenter + WaterInfoHalfExtent);
+	UWorld* World = GetWorld();
+
+	check(World != nullptr);
+
+	FVector Center = GetActorLocation();
+
+	if (FWaterViewExtension* WaterViewExtension = UWaterSubsystem::GetWaterViewExtension(World))
+	{
+		Center = WaterViewExtension->GetZoneLocation(this, PlayerIndex);
+	}
+	
+	return Center;
 }
 
 void AWaterZone::OnLevelAddedToWorld(ULevel* InLevel, UWorld* InWorld)
@@ -723,6 +767,16 @@ void AWaterZone::OnLevelRemovedFromWorld(ULevel* InLevel, UWorld* InWorld)
 	OnLevelChanged(InLevel, InWorld);
 }
 
+bool AWaterZone::ContainsActorsAffectingWaterZone(ULevel* InLevel, const FBox& WaterZoneBounds) const
+{
+	const bool bContainsActorsAffectingWaterZone = Algo::AnyOf(InLevel->Actors, [this, &WaterZoneBounds](const AActor* Actor)
+		{
+			return IsAffectingWaterZone(WaterZoneBounds, Actor);
+		});
+
+	return bContainsActorsAffectingWaterZone;
+}
+
 void AWaterZone::OnLevelChanged(ULevel* InLevel, UWorld* InWorld)
 {
 	if ((InLevel == nullptr) || (InWorld != GetWorld()))
@@ -730,11 +784,26 @@ void AWaterZone::OnLevelChanged(ULevel* InLevel, UWorld* InWorld)
 		return;
 	}
 
-	const FBox WaterZoneBounds = GetDynamicWaterInfoBounds();
-	const bool bContainsActorsAffectingWaterZone = Algo::AnyOf(InLevel->Actors, [this, &WaterZoneBounds](const AActor* Actor)
+	TArray<FBox> WaterZoneViewsBounds;
+	GetDynamicWaterInfoBounds(WaterZoneViewsBounds);
+
+	bool bContainsActorsAffectingWaterZone = false;
+
+	if (WaterZoneViewsBounds.Num() > 0)
 	{
-		return IsAffectingWaterZone(WaterZoneBounds, Actor);
-	});
+		for (const FBox& ZoneViewBounds : WaterZoneViewsBounds)
+		{
+			if (ContainsActorsAffectingWaterZone(InLevel, ZoneViewBounds))
+			{
+				bContainsActorsAffectingWaterZone = true;
+				break;
+			}
+		}
+	}
+	else
+	{
+		bContainsActorsAffectingWaterZone = ContainsActorsAffectingWaterZone(InLevel, GetZoneBounds());
+	}
 
 	if (bContainsActorsAffectingWaterZone)
 	{
diff --git a/Engine/Plugins/Experimental/Water/Source/Runtime/Public/WaterMeshComponent.h b/Engine/Plugins/Experimental/Water/Source/Runtime/Public/WaterMeshComponent.h
--- a/Engine/Plugins/Experimental/Water/Source/Runtime/Public/WaterMeshComponent.h
+++ b/Engine/Plugins/Experimental/Water/Source/Runtime/Public/WaterMeshComponent.h
@@ -62,8 +62,6 @@ public:
 
 	FIntPoint GetExtentInTiles() const;
 
-	void SetDynamicWaterMeshCenter(const FVector2D& NewCenter);
-	FVector2D GetDynamicWaterMeshCenter() const { return DynamicWaterMeshCenter; }
 	FVector2D GetGlobalWaterMeshCenter() const;
 
 	bool IsLocalOnlyTessellationEnabled() const;
@@ -107,9 +105,6 @@ private:
 	UPROPERTY(Transient, VisibleAnywhere, Category = Rendering)
 	mutable FIntPoint QuadTreeResolution = FIntPoint::ZeroValue;
 
-	/** The current center of the dynamic water mesh. Updated by the water view extension whenever the view location crosses the update bounds. */
-	FVector2D DynamicWaterMeshCenter = FVector2D::ZeroVector;
-
 	FWaterQuadTreeBuilder WaterQuadTreeBuilder;
 
 	/** Tiles containing water, stored in a quad tree */
diff --git a/Engine/Plugins/Experimental/Water/Source/Runtime/Public/WaterMeshSceneProxy.h b/Engine/Plugins/Experimental/Water/Source/Runtime/Public/WaterMeshSceneProxy.h
--- a/Engine/Plugins/Experimental/Water/Source/Runtime/Public/WaterMeshSceneProxy.h
+++ b/Engine/Plugins/Experimental/Water/Source/Runtime/Public/WaterMeshSceneProxy.h
@@ -164,6 +164,8 @@ public:
 	// Destroys the quadtree associated with Key.
 	void DestroyViewWaterQuadTree(int32 Key);
 
+	int32 FindBestQuadTreeForViewLocation(const FSceneView* View) const;
+
 private:
 
 #if RHI_RAYTRACING
diff --git a/Engine/Plugins/Experimental/Water/Source/Runtime/Public/WaterViewExtension.h b/Engine/Plugins/Experimental/Water/Source/Runtime/Public/WaterViewExtension.h
--- a/Engine/Plugins/Experimental/Water/Source/Runtime/Public/WaterViewExtension.h
+++ b/Engine/Plugins/Experimental/Water/Source/Runtime/Public/WaterViewExtension.h
@@ -4,9 +4,11 @@
 
 #include "SceneViewExtension.h"
 #include "WaterInfoRendering.h"
+#include "Misc/ScopeLock.h"
 
 
 class AWaterZone;
+class FWaterMeshSceneProxy;
 
 class FWaterViewExtension : public FWorldSceneViewExtension
 {
@@ -34,20 +36,32 @@ public:
 	void AddWaterZone(AWaterZone* InWaterZone);
 	void RemoveWaterZone(AWaterZone* InWaterZone);
 
-private:
-	/** Queued Water Info rendering contexts to submit for rendering on the next SetupView call */
-	TWeakObjectPtrKeyMap<AWaterZone, UE::WaterInfo::FRenderingContext> WaterInfoContextsToRender;
+	FVector GetZoneLocation(AWaterZone* InWaterZone, int32 PlayerIndex) const;
+
+	void CreateSceneProxyQuadtrees(FWaterMeshSceneProxy* SceneProxy);
 
 	struct FWaterZoneInfo
 	{
-		/** 
-		 * For each water zone, store the bounds of the tile from which the water zone was last rendered.
-		 * When the view location crosses the bounds, submit a new WaterInfo update to reflect the new active area
+		UE::WaterInfo::FRenderingContext RenderContext;
+
+		/**
+		 * For each water zone, per view: store the bounds of the tile from which the water zone was last rendered.
+		 * When the view location crosses the bounds, submit a new update to reflect the new active area
 		 */
-		TOptional<FBox2D> UpdateBounds = FBox2D(ForceInit);
+		struct FWaterZoneViewInfo
+		{
+			TOptional<FBox2D> UpdateBounds = FBox2D(ForceInit);
+			FVector Center = FVector(ForceInit);
+			FWaterMeshSceneProxy* OldSceneProxy = nullptr;
+			bool bIsDirty = true;
+		};
+		TArray<FWaterZoneViewInfo, TInlineAllocator<4>> ViewInfos;
 	};
 	TWeakObjectPtrKeyMap<AWaterZone, FWaterZoneInfo> WaterZoneInfos;
 
+private:
+	int32 CurrentNumViews = 0;
+
 	struct FWaterGPUResources
 	{
 		FBufferRHIRef WaterBodyDataBuffer;
@@ -59,9 +73,41 @@ private:
 
 	TSharedRef<FWaterGPUResources, ESPMode::ThreadSafe> WaterGPUData;
 
+	TArray<int32, TInlineAllocator<4>> ViewPlayerIndices;
+
+	struct FQuadtreeUpdateInfo
+	{
+		FWaterMeshSceneProxy* SceneProxy;
+		FVector2D Location;
+		int32 Key;
+	};
+
+	TArray<FQuadtreeUpdateInfo, TInlineAllocator<4>> QuadtreeUpdates;
+
+	FCriticalSection QuadtreeUpdateLock;
+
+	TMap<FSceneViewStateInterface*, int32> NonDataViewsQuadtreeKeys;
+
+	bool bWaterInfoTextureRebuildPending = true;
+
 	bool bRebuildGPUData = true;
 
+	// store the locations of every active water mesh scene proxy quad tree based on the key
+	TMap<int32, FVector2D> QuadTreeKeyLocationMap;
+
 	void UpdateGPUBuffers();
+
+	void UpdateViewInfo(AWaterZone* WaterZone, FSceneView& InView);
+	
+	void RenderWaterInfoTexture(FSceneViewFamily& InViewFamily, FSceneView& InView, const FWaterZoneInfo* WaterZoneInfo, FSceneInterface* Scene, const FVector& ZoneCenter);
+
+	// Returns the index in the views array corresponding to InView's PlayerIndex. If the index is not found it adds a new entry.
+	int32 GetOrAddViewindex(const FSceneView& InView);
+	// Returns the index in the views array corresponding to InView's PlayerIndex. INDEX_NONE if it doesn't find it
+	int32 GetViewIndex(int32 PlayerIndex) const;
+	int32 GetViewIndex(const FSceneView& InView) const;
+
+	void DrawDebugInfo(FSceneView& InView, AWaterZone* WaterZone);
 };
 
 struct FWaterMeshGPUWork
diff --git a/Engine/Plugins/Experimental/Water/Source/Runtime/Public/WaterZoneActor.h b/Engine/Plugins/Experimental/Water/Source/Runtime/Public/WaterZoneActor.h
--- a/Engine/Plugins/Experimental/Water/Source/Runtime/Public/WaterZoneActor.h
+++ b/Engine/Plugins/Experimental/Water/Source/Runtime/Public/WaterZoneActor.h
@@ -56,17 +56,19 @@ public:
 	FBox2D GetZoneBounds2D() const;
 	FBox GetZoneBounds() const;
 
+	void GetDynamicWaterInfoBounds(TArray<FBox>& Bounds) const;
+
 	void SetRenderTargetResolution(FIntPoint NewResolution);
 	FIntPoint GetRenderTargetResolution() const { return RenderTargetResolution; }
 
 	uint32 GetVelocityBlurRadius() const { return VelocityBlurRadius; }
 
-	FVector GetDynamicWaterInfoCenter() const;
 	FVector GetDynamicWaterInfoExtent() const;
-	FBox GetDynamicWaterInfoBounds() const;
+
+	// gets the water info center for a specific player view
+	FVector GetDynamicWaterInfoCenter(int32 PlayerIndex);
 
 	bool IsLocalOnlyTessellationEnabled() const { return bEnableLocalOnlyTessellation; }
-	void SetLocalTessellationCenter(const FVector& NewCenter) { LocalTessellationCenter = NewCenter;}
 
 	virtual void BeginPlay() override;
 	virtual void EndPlay(const EEndPlayReason::Type EndPlayReason) override;
@@ -89,6 +91,9 @@ public:
 	UPROPERTY(Transient, DuplicateTransient, VisibleAnywhere, BlueprintReadOnly, Category = Water)
 	TObjectPtr<UTextureRenderTarget2DArray> WaterInfoTextureArray;
 
+	UPROPERTY(Transient, DuplicateTransient, VisibleAnywhere, BlueprintReadOnly, Category = Water)
+	int32 WaterInfoTextureArrayNumSlices = 1;
+
 	FOnWaterInfoTextureArrayCreated& GetOnWaterInfoTextureArrayCreated() { return OnWaterInfoTextureArrayCreated; }
 	
 #if WITH_EDITOR
@@ -132,6 +137,8 @@ private:
 	void OnLevelRemovedFromWorld(ULevel* InLevel, UWorld* InWorld);
 	void OnLevelChanged(ULevel* InLevel, UWorld* InWorld);
 
+	bool ContainsActorsAffectingWaterZone(ULevel* InLevel, const FBox& WaterZoneBounds) const;
+
 	/** Returns true if the provided actor can affect waterzone resources.
 	 *
 	 * @param	InWaterZoneBounds	The bounds of this waterzone.
@@ -156,6 +163,9 @@ private:
 	void OnBoundsComponentModified();
 #endif // WITH_EDITOR
 
+	/** Mark aspects of the water zone for rebuild based on the Flags parameter within a given region. Optionally the caller can pass in a UObject to identify who requested the update. */
+	void MarkForRebuild(EWaterZoneRebuildFlags Flags, const FBox2D& UpdateRegion, const FBox& WaterInfoBounds, const UObject* DebugRequestingObject = nullptr);
+
 private:
 
 	UPROPERTY(Transient, Category = Water, VisibleAnywhere)
@@ -209,9 +219,6 @@ private:
 	FVector2f WaterHeightExtents;
 	float GroundZMin;
 
-	/** Current center of the local tessellation sliding window. Updated by the WaterViewExtension when the view crosses the update boundary */
-	FVector LocalTessellationCenter;
-
 	/** Unique Id for accessing zone data (Location, extent, ,...) in GPU buffers */
 	UPROPERTY(Transient, DuplicateTransient, NonTransactional, VisibleAnywhere, Category = Water)
 	int32 WaterZoneIndex = INDEX_NONE;
diff --git a/Engine/Plugins/Experimental/WaterAdvanced/Source/WaterAdvanced/Private/ShallowWaterSubsystem.cpp b/Engine/Plugins/Experimental/WaterAdvanced/Source/WaterAdvanced/Private/ShallowWaterSubsystem.cpp
--- a/Engine/Plugins/Experimental/WaterAdvanced/Source/WaterAdvanced/Private/ShallowWaterSubsystem.cpp
+++ b/Engine/Plugins/Experimental/WaterAdvanced/Source/WaterAdvanced/Private/ShallowWaterSubsystem.cpp
@@ -1226,7 +1226,9 @@ void UShallowWaterSubsystem::TryGetOrWaitForWaterInfoTextureFromWaterBodies(TSet
 
 				if (ShallowWaterNiagaraSimulation)
 				{
-					const FVector2D ZoneLocation = FVector2D(WaterZone->GetDynamicWaterInfoCenter());
+					// The following index assume that there is no split screen support and will request the position of the first player's water view.
+					const int32 PlayerIndex = 0;
+					const FVector2D ZoneLocation = FVector2D(WaterZone->GetDynamicWaterInfoCenter(PlayerIndex));
 
 					const FVector2D ZoneExtent = FVector2D(WaterZone->GetDynamicWaterInfoExtent());
 					const FVector2D WaterHeightExtents = FVector2D(WaterZone->GetWaterHeightExtents());
