[Water]

Fixed missing deprecation following water split screen support changes.

diff --git a/Engine/Plugins/Experimental/Water/Source/Runtime/Private/WaterViewExtension.cpp b/Engine/Plugins/Experimental/Water/Source/Runtime/Private/WaterViewExtension.cpp
--- a/Engine/Plugins/Experimental/Water/Source/Runtime/Private/WaterViewExtension.cpp
+++ b/Engine/Plugins/Experimental/Water/Source/Runtime/Private/WaterViewExtension.cpp
@@ -794,7 +794,7 @@ void FWaterViewExtension::RemoveWaterZone(AWaterZone* InWaterZone)
 	UE_LOG(LogWater, Verbose, TEXT("Water Zone (%s): RemoveWaterZone was called."), *GetNameSafe(InWaterZone));
 }
 
-FVector FWaterViewExtension::GetZoneLocation(AWaterZone* InWaterZone, int32 PlayerIndex) const
+FVector FWaterViewExtension::GetZoneLocation(const AWaterZone* InWaterZone, int32 PlayerIndex) const
 {
 	if (ensure(InWaterZone->HasActorRegisteredAllComponents()))
 	{
diff --git a/Engine/Plugins/Experimental/Water/Source/Runtime/Private/WaterZoneActor.cpp b/Engine/Plugins/Experimental/Water/Source/Runtime/Private/WaterZoneActor.cpp
--- a/Engine/Plugins/Experimental/Water/Source/Runtime/Private/WaterZoneActor.cpp
+++ b/Engine/Plugins/Experimental/Water/Source/Runtime/Private/WaterZoneActor.cpp
@@ -128,13 +128,13 @@ FBox2D AWaterZone::GetZoneBounds2D() const
 	return WaterZoneBounds;
 }
 
-void AWaterZone::GetDynamicWaterInfoBounds(TArray<FBox>& Bounds) const
+void AWaterZone::GetAllDynamicWaterInfoBounds(TArray<FBox>& OutBounds) const
 {
-	if (FWaterViewExtension* WaterViewExtension = UWaterSubsystem::GetWaterViewExtension(GetWorld()))
+	if (const FWaterViewExtension* WaterViewExtension = UWaterSubsystem::GetWaterViewExtension(GetWorld()))
 	{
-		FWaterViewExtension::FWaterZoneInfo* WaterZoneInfo = WaterViewExtension->WaterZoneInfos.Find(this);
+		const FWaterViewExtension::FWaterZoneInfo* WaterZoneInfo = WaterViewExtension->WaterZoneInfos.Find(this);
 
-		if (WaterZoneInfo != nullptr)
+		if (ensure(WaterZoneInfo != nullptr))
 		{
 			for (const FWaterViewExtension::FWaterZoneInfo::FWaterZoneViewInfo& ViewInfo : WaterZoneInfo->ViewInfos)
 			{
@@ -147,7 +147,30 @@ void AWaterZone::GetDynamicWaterInfoBounds(TArray<FBox>& Bounds) const
 
 				const FVector HalfExtent(GetDynamicWaterInfoExtent() / 2);
 
-				Bounds.Add(FBox(Center - HalfExtent, Center + HalfExtent));
+				OutBounds.Emplace(FBox(Center - HalfExtent, Center + HalfExtent));
+			}
+		}
+	}
+}
+
+void AWaterZone::GetAllDynamicWaterInfoCenters(TArray<FVector>& OutCenters) const
+{
+	if (const FWaterViewExtension* WaterViewExtension = UWaterSubsystem::GetWaterViewExtension(GetWorld()))
+	{
+		const FWaterViewExtension::FWaterZoneInfo* WaterZoneInfo = WaterViewExtension->WaterZoneInfos.Find(this);
+
+		if (ensure(WaterZoneInfo != nullptr))
+		{
+			for (const FWaterViewExtension::FWaterZoneInfo::FWaterZoneViewInfo& ViewInfo : WaterZoneInfo->ViewInfos)
+			{
+				FVector Center(ViewInfo.Center);
+
+				if (!IsLocalOnlyTessellationEnabled())
+				{
+					Center = GetActorLocation();
+				}
+
+				OutCenters.Emplace(Center);
 			}
 		}
 	}
@@ -243,27 +266,6 @@ void AWaterZone::DeclareConstructClasses(TArray<FTopLevelAssetPath>& OutConstruc
 }
 #endif
 
-void AWaterZone::MarkForRebuild(EWaterZoneRebuildFlags Flags, const FBox2D& UpdateRegion, const FBox& WaterInfoBounds, const UObject* DebugRequestingObject)
-{
-	const FBox2D WaterInfoBounds2D(FVector2D(WaterInfoBounds.Min), FVector2D(WaterInfoBounds.Max));
-
-	// Suppress updates which occur outside the bounds of the water zone.
-	if ((!UpdateRegion.bIsValid || UpdateRegion.Intersect(WaterInfoBounds2D)))
-	{
-		if (EnumHasAnyFlags(Flags, EWaterZoneRebuildFlags::UpdateWaterMesh))
-		{
-			UE_LOG(LogWater, Verbose, TEXT("AWaterZone (%s) UpdateWaterMesh in region {%s} (triggered by %s)"), *GetNameSafe(this), *UpdateRegion.ToString(), *GetNameSafe(DebugRequestingObject));
-			WaterMesh->MarkWaterMeshGridDirty();
-			WaterMesh->MarkRenderStateDirty();
-		}
-		if (EnumHasAnyFlags(Flags, EWaterZoneRebuildFlags::UpdateWaterInfoTexture))
-		{
-			UE_LOG(LogWater, Verbose, TEXT("AWaterZone (%s) UpdateWaterInfoTexture in region {%s} (triggered by %s)"), *GetNameSafe(this), *UpdateRegion.ToString(), *GetNameSafe(DebugRequestingObject));
-			bNeedsWaterInfoRebuild = true;
-		}
-	}
-}
-
 void AWaterZone::MarkForRebuild(EWaterZoneRebuildFlags Flags, const FBox2D& UpdateRegion, const UObject* DebugRequestingObject)
 {
 	// TODO: investigate why optimizing by updating only checking the bounds of the active views dynamic extents is not working.
@@ -741,15 +743,15 @@ FVector AWaterZone::GetDynamicWaterInfoExtent() const
 	}
 }
 
-FVector AWaterZone::GetDynamicWaterInfoCenter(int32 PlayerIndex)
+FVector AWaterZone::GetDynamicWaterInfoCenter(int32 PlayerIndex) const
 {
-	UWorld* World = GetWorld();
+	const UWorld* World = GetWorld();
 
 	check(World != nullptr);
 
 	FVector Center = GetActorLocation();
 
-	if (FWaterViewExtension* WaterViewExtension = UWaterSubsystem::GetWaterViewExtension(World))
+	if (const FWaterViewExtension* WaterViewExtension = UWaterSubsystem::GetWaterViewExtension(World))
 	{
 		Center = WaterViewExtension->GetZoneLocation(this, PlayerIndex);
 	}
@@ -757,6 +759,31 @@ FVector AWaterZone::GetDynamicWaterInfoCenter(int32 PlayerIndex)
 	return Center;
 }
 
+FBox AWaterZone::GetDynamicWaterInfoBounds(int32 PlayerIndex) const
+{
+	const UWorld* World = GetWorld();
+
+	check(World != nullptr);
+
+	FVector Center = GetActorLocation();
+
+	if (const FWaterViewExtension* WaterViewExtension = UWaterSubsystem::GetWaterViewExtension(World))
+	{
+		Center = WaterViewExtension->GetZoneLocation(this, PlayerIndex);
+	}
+	
+	const FVector WaterInfoHalfExtents = GetDynamicWaterInfoExtent() / 2.;
+	return FBox(Center - WaterInfoHalfExtents, Center + WaterInfoHalfExtents);
+}
+
+FVector AWaterZone::GetDynamicWaterInfoCenter() const
+{
+	// The following index assume that there is no split screen and will request the position of the first player's water view.
+	// This call is a temporary fallback and is deprecated in favor of the overload which takes a specific player view.
+	constexpr int32 PlayerIndex = 0;
+	return GetDynamicWaterInfoCenter(PlayerIndex);
+}
+
 void AWaterZone::OnLevelAddedToWorld(ULevel* InLevel, UWorld* InWorld)
 {
 	OnLevelChanged(InLevel, InWorld);
@@ -785,7 +812,7 @@ void AWaterZone::OnLevelChanged(ULevel* InLevel, UWorld* InWorld)
 	}
 
 	TArray<FBox> WaterZoneViewsBounds;
-	GetDynamicWaterInfoBounds(WaterZoneViewsBounds);
+	GetAllDynamicWaterInfoBounds(WaterZoneViewsBounds);
 
 	bool bContainsActorsAffectingWaterZone = false;
 
diff --git a/Engine/Plugins/Experimental/Water/Source/Runtime/Public/WaterMeshComponent.h b/Engine/Plugins/Experimental/Water/Source/Runtime/Public/WaterMeshComponent.h
--- a/Engine/Plugins/Experimental/Water/Source/Runtime/Public/WaterMeshComponent.h
+++ b/Engine/Plugins/Experimental/Water/Source/Runtime/Public/WaterMeshComponent.h
@@ -62,6 +62,11 @@ public:
 
 	FIntPoint GetExtentInTiles() const;
 
+	UE_DEPRECATED(5.5, "It is no longer possible to manually set the dynamic mesh center. This is controlled per view by the water view extension.")
+	void SetDynamicWaterMeshCenter(const FVector2D& NewCenter) { }
+	UE_DEPRECATED(5.5, "Dynamic water mesh center is now per-view and must be retrieved through the water view extension (water zone actor provides utilities to do this as well)")
+	FVector2D GetDynamicWaterMeshCenter() const { return FVector2D::ZeroVector; }
+
 	FVector2D GetGlobalWaterMeshCenter() const;
 
 	bool IsLocalOnlyTessellationEnabled() const;
diff --git a/Engine/Plugins/Experimental/Water/Source/Runtime/Public/WaterViewExtension.h b/Engine/Plugins/Experimental/Water/Source/Runtime/Public/WaterViewExtension.h
--- a/Engine/Plugins/Experimental/Water/Source/Runtime/Public/WaterViewExtension.h
+++ b/Engine/Plugins/Experimental/Water/Source/Runtime/Public/WaterViewExtension.h
@@ -36,7 +36,7 @@ public:
 	void AddWaterZone(AWaterZone* InWaterZone);
 	void RemoveWaterZone(AWaterZone* InWaterZone);
 
-	FVector GetZoneLocation(AWaterZone* InWaterZone, int32 PlayerIndex) const;
+	FVector GetZoneLocation(const AWaterZone* InWaterZone, int32 PlayerIndex) const;
 
 	void CreateSceneProxyQuadtrees(FWaterMeshSceneProxy* SceneProxy);
 
diff --git a/Engine/Plugins/Experimental/Water/Source/Runtime/Public/WaterZoneActor.h b/Engine/Plugins/Experimental/Water/Source/Runtime/Public/WaterZoneActor.h
--- a/Engine/Plugins/Experimental/Water/Source/Runtime/Public/WaterZoneActor.h
+++ b/Engine/Plugins/Experimental/Water/Source/Runtime/Public/WaterZoneActor.h
@@ -56,20 +56,34 @@ public:
 	FBox2D GetZoneBounds2D() const;
 	FBox GetZoneBounds() const;
 
-	void GetDynamicWaterInfoBounds(TArray<FBox>& Bounds) const;
+	/** Retrieves all the per-view bounds for this water zone */
+	void GetAllDynamicWaterInfoBounds(TArray<FBox>& OutBounds) const;
+	/** Retrieves all the per-view centers for this water zone */
+	void GetAllDynamicWaterInfoCenters(TArray<FVector>& OutCenters) const;
 
 	void SetRenderTargetResolution(FIntPoint NewResolution);
 	FIntPoint GetRenderTargetResolution() const { return RenderTargetResolution; }
 
 	uint32 GetVelocityBlurRadius() const { return VelocityBlurRadius; }
 
-	FVector GetDynamicWaterInfoExtent() const;
+	/** Retrieves the dynamic water info center for a specific player index. */
+	FVector GetDynamicWaterInfoCenter(int32 PlayerIndex) const;
+	/** Retrieves the dynamic water info bounds for a specific player index. */
+	FBox GetDynamicWaterInfoBounds(int32 PlayerIndex) const;
 
-	// gets the water info center for a specific player view
-	FVector GetDynamicWaterInfoCenter(int32 PlayerIndex);
+	FVector GetDynamicWaterInfoExtent() const;
 
 	bool IsLocalOnlyTessellationEnabled() const { return bEnableLocalOnlyTessellation; }
 
+	UE_DEPRECATED(5.5, "In 5.5 the dynamic water info is now per player view. Either call GetAllDynamicWaterInfoCenters or use the per-player index version")
+	FVector GetDynamicWaterInfoCenter() const;
+
+	UE_DEPRECATED(5.5, "In 5.5 the dynamic water info is now per player view. Either call GetAllDynamicWaterInfoCenters or use the per-player index version")
+	FBox GetDynamicWaterInfoBounds() const;
+
+	UE_DEPRECATED(5.5, "It is no longer possible to manually set the local tessellation center. This is controlled per view by the water view extension.")
+	void SetLocalTessellationCenter(const FVector& NewCenter) {}
+
 	virtual void BeginPlay() override;
 	virtual void EndPlay(const EEndPlayReason::Type EndPlayReason) override;
 	virtual void PostLoadSubobjects(FObjectInstancingGraph* OuterInstanceGraph) override;
@@ -163,9 +177,6 @@ private:
 	void OnBoundsComponentModified();
 #endif // WITH_EDITOR
 
-	/** Mark aspects of the water zone for rebuild based on the Flags parameter within a given region. Optionally the caller can pass in a UObject to identify who requested the update. */
-	void MarkForRebuild(EWaterZoneRebuildFlags Flags, const FBox2D& UpdateRegion, const FBox& WaterInfoBounds, const UObject* DebugRequestingObject = nullptr);
-
 private:
 
 	UPROPERTY(Transient, Category = Water, VisibleAnywhere)
