[Water]

Fixed the first water info render with render method 1 of the frame deleting all subsequent updates.

diff --git a/Engine/Plugins/Experimental/Water/Source/Runtime/Private/WaterInfoRendering.cpp b/Engine/Plugins/Experimental/Water/Source/Runtime/Private/WaterInfoRendering.cpp
--- a/Engine/Plugins/Experimental/Water/Source/Runtime/Private/WaterInfoRendering.cpp
+++ b/Engine/Plugins/Experimental/Water/Source/Runtime/Private/WaterInfoRendering.cpp
@@ -827,8 +827,6 @@ void UpdateWaterInfoRendering2(FSceneView& InView, const FRenderingContext& Cont
 {
 	TRACE_CPUPROFILER_EVENT_SCOPE(WaterInfo::UpdateWaterInfoRendering2);
 
-	InView.WaterInfoTextureRenderingParams.Reset();
-
 	if (!IsValid(Context.TextureRenderTarget))
 	{
 		return;
