[Water]

Removed a valid assert that the Water Mesh Scene Proxy is not null. This can occur during fullscreen resizes.

diff --git a/Engine/Plugins/Experimental/Water/Source/Runtime/Private/WaterViewExtension.cpp b/Engine/Plugins/Experimental/Water/Source/Runtime/Private/WaterViewExtension.cpp
--- a/Engine/Plugins/Experimental/Water/Source/Runtime/Private/WaterViewExtension.cpp
+++ b/Engine/Plugins/Experimental/Water/Source/Runtime/Private/WaterViewExtension.cpp
@@ -571,33 +571,35 @@ void FWaterViewExtension::SetupView(FSceneViewFamily& InViewFamily, FSceneView&
 				UWaterMeshComponent* WaterMeshComponent = WaterZone->GetWaterMeshComponent();
 				check(WaterMeshComponent != nullptr);
 
-				FWaterMeshSceneProxy*  SceneProxy = (FWaterMeshSceneProxy*) WaterMeshComponent->GetSceneProxy();
-				check(SceneProxy != nullptr);
+				FWaterMeshSceneProxy* SceneProxy = static_cast<FWaterMeshSceneProxy*>(WaterMeshComponent->GetSceneProxy());
 
-				// push a quadtree update
-				if (WaterZone->IsLocalOnlyTessellationEnabled() && (bBoundsUpdateNeeded || SceneProxy != WaterZoneViewInfo.OldSceneProxy))
+				if (SceneProxy != nullptr)
 				{
-					FScopeLock AddQuadtreeUpdate(&QuadtreeUpdateLock);
+					// push a quadtree update
+					if (WaterZone->IsLocalOnlyTessellationEnabled() && (bBoundsUpdateNeeded || SceneProxy != WaterZoneViewInfo.OldSceneProxy))
+					{
+						FScopeLock AddQuadtreeUpdate(&QuadtreeUpdateLock);
 
-					FQuadtreeUpdateInfo QuadtreeUpdate;
-					QuadtreeUpdate.SceneProxy = (FWaterMeshSceneProxy*)SceneProxy;
-					QuadtreeUpdate.Location = WaterZoneViewInfo.UpdateBounds->GetCenter();
-					// we use the actual PlayerIndex instead of ViewPlayerIndex, since the latter can change if more views are added while a quadtree already exists
-					QuadtreeUpdate.Key = InView.PlayerIndex;
+						FQuadtreeUpdateInfo QuadtreeUpdate;
+						QuadtreeUpdate.SceneProxy = SceneProxy;
+						QuadtreeUpdate.Location = WaterZoneViewInfo.UpdateBounds->GetCenter();
+						// we use the actual PlayerIndex instead of ViewPlayerIndex, since the latter can change if more views are added while a quadtree already exists
+						QuadtreeUpdate.Key = InView.PlayerIndex;
 
-					QuadtreeUpdates.Add(QuadtreeUpdate);
+						QuadtreeUpdates.Add(QuadtreeUpdate);
 
 
-					if (!QuadTreeKeyLocationMap.Contains(QuadtreeUpdate.Key))
-					{
-						QuadTreeKeyLocationMap.Add(QuadtreeUpdate.Key, QuadtreeUpdate.Location);
-					}
-					else
-					{
-						QuadTreeKeyLocationMap[QuadtreeUpdate.Key] = QuadtreeUpdate.Location;
-					}
+						if (!QuadTreeKeyLocationMap.Contains(QuadtreeUpdate.Key))
+						{
+							QuadTreeKeyLocationMap.Add(QuadtreeUpdate.Key, QuadtreeUpdate.Location);
+						}
+						else
+						{
+							QuadTreeKeyLocationMap[QuadtreeUpdate.Key] = QuadtreeUpdate.Location;
+						}
 
-					UE_LOG(LogWater, Verbose, TEXT("Water Zone (%s) queued a quadtree update for view %d updated."), *GetNameSafe(WaterZone), ViewPlayerIndex);
+						UE_LOG(LogWater, Verbose, TEXT("Water Zone (%s) queued a quadtree update for view %d updated."), *GetNameSafe(WaterZone), ViewPlayerIndex);
+					}
 				}
 
 				WaterZoneViewInfo.OldSceneProxy = SceneProxy;
@@ -625,12 +627,13 @@ void FWaterViewExtension::SetupView(FSceneViewFamily& InViewFamily, FSceneView&
 					UWaterMeshComponent* WaterMeshComponent = WaterZone->GetWaterMeshComponent();
 					check(WaterMeshComponent != nullptr);
 
-					FWaterMeshSceneProxy* SceneProxy = (FWaterMeshSceneProxy*)WaterMeshComponent->GetSceneProxy();
-					check(SceneProxy != nullptr);
-
-					NonDataViewsQuadtreeKeys.Add(InView.State, SceneProxy->FindBestQuadTreeForViewLocation(&InView));
+					FWaterMeshSceneProxy* SceneProxy = static_cast<FWaterMeshSceneProxy*>(WaterMeshComponent->GetSceneProxy());
+					if (ensure(SceneProxy != nullptr))
+					{
+						NonDataViewsQuadtreeKeys.Add(InView.State, SceneProxy->FindBestQuadTreeForViewLocation(&InView));
 
-					UE_LOG(LogWater, Verbose, TEXT("Water Zone (%s) queued a search for the closest quadtree for a View (0x%p) which has no WaterInfo."), *GetNameSafe(WaterZone), &InView);
+						UE_LOG(LogWater, Verbose, TEXT("Water Zone (%s) queued a search for the closest quadtree for a View (0x%p) which has no WaterInfo."), *GetNameSafe(WaterZone), &InView);
+					}
 				}
 			}
 		}	
@@ -794,7 +797,7 @@ void FWaterViewExtension::RemoveWaterZone(AWaterZone* InWaterZone)
 	UE_LOG(LogWater, Verbose, TEXT("Water Zone (%s): RemoveWaterZone was called."), *GetNameSafe(InWaterZone));
 }
 
-FVector FWaterViewExtension::GetZoneLocation(const AWaterZone* InWaterZone, int32 PlayerIndex) const
+FVector FWaterViewExtension::GetZoneLocation(AWaterZone* InWaterZone, int32 PlayerIndex) const
 {
 	if (ensure(InWaterZone->HasActorRegisteredAllComponents()))
 	{
