[Water]

Fixed bad merge

diff --git a/Engine/Plugins/Experimental/Water/Source/Runtime/Private/WaterViewExtension.cpp b/Engine/Plugins/Experimental/Water/Source/Runtime/Private/WaterViewExtension.cpp
--- a/Engine/Plugins/Experimental/Water/Source/Runtime/Private/WaterViewExtension.cpp
+++ b/Engine/Plugins/Experimental/Water/Source/Runtime/Private/WaterViewExtension.cpp
@@ -797,7 +797,7 @@ void FWaterViewExtension::RemoveWaterZone(AWaterZone* InWaterZone)
 	UE_LOG(LogWater, Verbose, TEXT("Water Zone (%s): RemoveWaterZone was called."), *GetNameSafe(InWaterZone));
 }
 
-FVector FWaterViewExtension::GetZoneLocation(AWaterZone* InWaterZone, int32 PlayerIndex) const
+FVector FWaterViewExtension::GetZoneLocation(const AWaterZone* InWaterZone, int32 PlayerIndex) const
 {
 	if (ensure(InWaterZone->HasActorRegisteredAllComponents()))
 	{
