[Water]

Fix water gpu buffers not being uploaded when there are no water zones in the map. This is required in some cases to ensure that we can use the static mesh fallback meshes without any loaded water zones (the static water materials will try to sample the gpu buffers, so they must be available).

diff --git a/Engine/Plugins/Experimental/Water/Source/Runtime/Private/WaterViewExtension.cpp b/Engine/Plugins/Experimental/Water/Source/Runtime/Private/WaterViewExtension.cpp
--- a/Engine/Plugins/Experimental/Water/Source/Runtime/Private/WaterViewExtension.cpp
+++ b/Engine/Plugins/Experimental/Water/Source/Runtime/Private/WaterViewExtension.cpp
@@ -492,170 +492,172 @@ void FWaterViewExtension::SetupView(FSceneViewFamily& InViewFamily, FSceneView&
 	{
 		ViewPlayerIndices.Empty();
 		NonDataViewsQuadtreeKeys.Empty();
-		return;
 	}
-
-	// Don't dirty the water info texture when we're rendering from a scene capture. Due to the frame delay after marking the texture as dirty, scene captures wouldn't have the right texture anyways.
-	// #todo_water [roey]: Once we have no frame-delay for updating the texture and lesser performance impact, we can re-enable updates within scene captures.
-	bool bShouldHaveWaterZoneViewData = !InView.bIsSceneCapture && !InView.bIsSceneCaptureCube && !InView.bIsReflectionCapture && !InView.bIsPlanarReflection && !InView.bIsVirtualTexture;
-
-	if (bShouldHaveWaterZoneViewData)
+	else
 	{
-		if (CurrentNumViews != NumViews)
+		// Don't dirty the water info texture when we're rendering from a scene capture. Due to the frame delay after marking the texture as dirty, scene captures wouldn't have the right texture anyways.
+		// #todo_water [roey]: Once we have no frame-delay for updating the texture and lesser performance impact, we can re-enable updates within scene captures.
+		bool bShouldHaveWaterZoneViewData = !InView.bIsSceneCapture && !InView.bIsSceneCaptureCube && !InView.bIsReflectionCapture && !InView.bIsPlanarReflection && !InView.bIsVirtualTexture;
+
+		if (bShouldHaveWaterZoneViewData)
 		{
-			for (AWaterZone* WaterZone : TActorRange<AWaterZone>(WorldPtr.Get()))
+			if (CurrentNumViews != NumViews)
 			{
-				if (WaterZone->HasActorRegisteredAllComponents())
+				for (AWaterZone* WaterZone : TActorRange<AWaterZone>(WorldPtr.Get()))
 				{
-					FWaterZoneInfo& WaterZoneInfo = WaterZoneInfos.FindChecked(WaterZone);
+					if (WaterZone->HasActorRegisteredAllComponents())
+					{
+						FWaterZoneInfo& WaterZoneInfo = WaterZoneInfos.FindChecked(WaterZone);
 
-					// make sure that if !IsLocalOnlyTessellationEnabled we only create a single WaterInfoTexture slice.
-					WaterZone->WaterInfoTextureArrayNumSlices = WaterZone->IsLocalOnlyTessellationEnabled() ? NumViews : 1;
-					// Mark for rebuild the WaterZone if the size changes
-					WaterZone->MarkForRebuild(EWaterZoneRebuildFlags::UpdateWaterInfoTexture);
+						// make sure that if !IsLocalOnlyTessellationEnabled we only create a single WaterInfoTexture slice.
+						WaterZone->WaterInfoTextureArrayNumSlices = WaterZone->IsLocalOnlyTessellationEnabled() ? NumViews : 1;
+						// Mark for rebuild the WaterZone if the size changes
+						WaterZone->MarkForRebuild(EWaterZoneRebuildFlags::UpdateWaterInfoTexture);
 
-					// init per view info
+						// init per view info
 
-					WaterZoneInfo.ViewInfos.Empty();
+						WaterZoneInfo.ViewInfos.Empty();
 
-					for (int i = 0; i < NumViews; ++i)
-					{
-						WaterZoneInfo.ViewInfos.Emplace(FWaterZoneInfo::FWaterZoneViewInfo());
+						for (int i = 0; i < NumViews; ++i)
+						{
+							WaterZoneInfo.ViewInfos.Emplace(FWaterZoneInfo::FWaterZoneViewInfo());
+						}
 					}
-				}
 
-				UE_LOG(LogWater, Verbose, TEXT("Number of views changed. Water Zone (%s) ViewInfos was reset."), *GetNameSafe(WaterZone));
-			}
+					UE_LOG(LogWater, Verbose, TEXT("Number of views changed. Water Zone (%s) ViewInfos was reset."), *GetNameSafe(WaterZone));
+				}
 
-			CurrentNumViews = NumViews;
-			bWaterInfoTextureRebuildPending = true;
+				CurrentNumViews = NumViews;
+				bWaterInfoTextureRebuildPending = true;
 
-			return;
-		}
+				return;
+			}
 
-		// Check if the view location is no longer within the current update bounds of a water zone and if so, queue an update for it.
-		for (AWaterZone* WaterZone : TActorRange<AWaterZone>(WorldPtr.Get()))
-		{
-			if (WaterZone->HasActorRegisteredAllComponents())
+			// Check if the view location is no longer within the current update bounds of a water zone and if so, queue an update for it.
+			for (AWaterZone* WaterZone : TActorRange<AWaterZone>(WorldPtr.Get()))
 			{
-				FWaterZoneInfo& WaterZoneInfo = WaterZoneInfos.FindChecked(WaterZone);
-
-				if (WaterZone->WaterInfoTextureArray.Get() == nullptr)
-				{
-					continue;
-				}
-
-				if (CVarDrawPerViewDebugInfo.GetValueOnGameThread())
+				if (WaterZone->HasActorRegisteredAllComponents())
 				{
-					DrawDebugInfo(InView, WaterZone);
-				}
+					FWaterZoneInfo& WaterZoneInfo = WaterZoneInfos.FindChecked(WaterZone);
 
-				FWaterZoneInfo::FWaterZoneViewInfo& WaterZoneViewInfo = WaterZoneInfo.ViewInfos[ViewPlayerIndex];
+					if (WaterZone->WaterInfoTextureArray.Get() == nullptr)
+					{
+						continue;
+					}
 
-				bool bBoundsUpdateNeeded = WaterZone->IsLocalOnlyTessellationEnabled() && (!WaterZoneViewInfo.UpdateBounds.IsSet() || !WaterZoneViewInfo.UpdateBounds->IsInside(FVector2D(ViewLocation)));
-				bBoundsUpdateNeeded |= WaterZoneViewInfo.bIsDirty;
+					if (CVarDrawPerViewDebugInfo.GetValueOnGameThread())
+					{
+						DrawDebugInfo(InView, WaterZone);
+					}
 
-				if (bBoundsUpdateNeeded)
-				{
-					UpdateViewInfo(WaterZone, InView);
+					FWaterZoneInfo::FWaterZoneViewInfo& WaterZoneViewInfo = WaterZoneInfo.ViewInfos[ViewPlayerIndex];
 
-					UE_LOG(LogWater, Verbose, TEXT("Water Zone (%s) ViewInfo for view %d updated."), *GetNameSafe(WaterZone), ViewPlayerIndex);
+					bool bBoundsUpdateNeeded = WaterZone->IsLocalOnlyTessellationEnabled() && (!WaterZoneViewInfo.UpdateBounds.IsSet() || !WaterZoneViewInfo.UpdateBounds->IsInside(FVector2D(ViewLocation)));
+					bBoundsUpdateNeeded |= WaterZoneViewInfo.bIsDirty;
 
-					// make sure that if !IsLocalOnlyTessellationEnabled, we only update the WaterInfoTexture for a single view
-					if (WaterZone->IsLocalOnlyTessellationEnabled() || ViewPlayerIndex == 0)
+					if (bBoundsUpdateNeeded)
 					{
-						RenderWaterInfoTexture(InViewFamily, InView, &WaterZoneInfo, Scene, WaterZoneViewInfo.Center);
+						UpdateViewInfo(WaterZone, InView);
+
+						UE_LOG(LogWater, Verbose, TEXT("Water Zone (%s) ViewInfo for view %d updated."), *GetNameSafe(WaterZone), ViewPlayerIndex);
+
+						// make sure that if !IsLocalOnlyTessellationEnabled, we only update the WaterInfoTexture for a single view
+						if (WaterZone->IsLocalOnlyTessellationEnabled() || ViewPlayerIndex == 0)
+						{
+							RenderWaterInfoTexture(InViewFamily, InView, &WaterZoneInfo, Scene, WaterZoneViewInfo.Center);
+						}
 					}
-				}
 
-				UWaterMeshComponent* WaterMeshComponent = WaterZone->GetWaterMeshComponent();
-				check(WaterMeshComponent != nullptr);
+					UWaterMeshComponent* WaterMeshComponent = WaterZone->GetWaterMeshComponent();
+					check(WaterMeshComponent != nullptr);
 
-				FWaterMeshSceneProxy* SceneProxy = static_cast<FWaterMeshSceneProxy*>(WaterMeshComponent->GetSceneProxy());
+					FWaterMeshSceneProxy* SceneProxy = static_cast<FWaterMeshSceneProxy*>(WaterMeshComponent->GetSceneProxy());
 
-				if (SceneProxy != nullptr)
-				{
-					// push a quadtree update
-					if (WaterZone->IsLocalOnlyTessellationEnabled() && (bBoundsUpdateNeeded || SceneProxy != WaterZoneViewInfo.OldSceneProxy))
+					if (SceneProxy != nullptr)
 					{
-						FScopeLock AddQuadtreeUpdate(&QuadtreeUpdateLock);
+						// push a quadtree update
+						if (WaterZone->IsLocalOnlyTessellationEnabled() && (bBoundsUpdateNeeded || SceneProxy != WaterZoneViewInfo.OldSceneProxy))
+						{
+							FScopeLock AddQuadtreeUpdate(&QuadtreeUpdateLock);
 
-						FQuadtreeUpdateInfo QuadtreeUpdate;
-						QuadtreeUpdate.SceneProxy = SceneProxy;
-						QuadtreeUpdate.Location = WaterZoneViewInfo.UpdateBounds->GetCenter();
-						// we use the actual PlayerIndex instead of ViewPlayerIndex, since the latter can change if more views are added while a quadtree already exists
-						QuadtreeUpdate.Key = InView.PlayerIndex;
+							FQuadtreeUpdateInfo QuadtreeUpdate;
+							QuadtreeUpdate.SceneProxy = SceneProxy;
+							QuadtreeUpdate.Location = WaterZoneViewInfo.UpdateBounds->GetCenter();
+							// we use the actual PlayerIndex instead of ViewPlayerIndex, since the latter can change if more views are added while a quadtree already exists
+							QuadtreeUpdate.Key = InView.PlayerIndex;
 
-						QuadtreeUpdates.Add(QuadtreeUpdate);
+							QuadtreeUpdates.Add(QuadtreeUpdate);
 
 
-						if (!QuadTreeKeyLocationMap.Contains(QuadtreeUpdate.Key))
-						{
-							QuadTreeKeyLocationMap.Add(QuadtreeUpdate.Key, QuadtreeUpdate.Location);
-						}
-						else
-						{
-							QuadTreeKeyLocationMap[QuadtreeUpdate.Key] = QuadtreeUpdate.Location;
-						}
+							if (!QuadTreeKeyLocationMap.Contains(QuadtreeUpdate.Key))
+							{
+								QuadTreeKeyLocationMap.Add(QuadtreeUpdate.Key, QuadtreeUpdate.Location);
+							}
+							else
+							{
+								QuadTreeKeyLocationMap[QuadtreeUpdate.Key] = QuadtreeUpdate.Location;
+							}
 
-						UE_LOG(LogWater, Verbose, TEXT("Water Zone (%s) queued a quadtree update for view %d updated."), *GetNameSafe(WaterZone), ViewPlayerIndex);
+							UE_LOG(LogWater, Verbose, TEXT("Water Zone (%s) queued a quadtree update for view %d updated."), *GetNameSafe(WaterZone), ViewPlayerIndex);
+						}
 					}
-				}
 
-				WaterZoneViewInfo.OldSceneProxy = SceneProxy;
-				WaterZoneViewInfo.bIsDirty = false;
+					WaterZoneViewInfo.OldSceneProxy = SceneProxy;
+					WaterZoneViewInfo.bIsDirty = false;
+				}
 			}
 		}
-	}
-	else
-	{
-		FScopeLock AddNonDataViewsQuadtreeKeys(&QuadtreeUpdateLock);
-
-		// if !bShouldHaveWaterZoneViewData, we want to store the closest quadtree ID for this view (since it is not
-		// going to have its own quadtree associated with it), so we can assign the WaterInfoTexture index of the view assigned to that quadtree
-		for (AWaterZone* WaterZone : TActorRange<AWaterZone>(WorldPtr.Get()))
+		else
 		{
-			if (WaterZone->WaterInfoTextureArray.Get() == nullptr)
-			{
-				continue;
-			}
+			FScopeLock AddNonDataViewsQuadtreeKeys(&QuadtreeUpdateLock);
 
-			if (WaterZone->HasActorRegisteredAllComponents())
+			// if !bShouldHaveWaterZoneViewData, we want to store the closest quadtree ID for this view (since it is not
+			// going to have its own quadtree associated with it), so we can assign the WaterInfoTexture index of the view assigned to that quadtree
+			for (AWaterZone* WaterZone : TActorRange<AWaterZone>(WorldPtr.Get()))
 			{
-				if (WaterZone->IsLocalOnlyTessellationEnabled())
+				if (WaterZone->WaterInfoTextureArray.Get() == nullptr)
 				{
-					UWaterMeshComponent* WaterMeshComponent = WaterZone->GetWaterMeshComponent();
-					check(WaterMeshComponent != nullptr);
+					continue;
+				}
 
-					FWaterMeshSceneProxy* SceneProxy = static_cast<FWaterMeshSceneProxy*>(WaterMeshComponent->GetSceneProxy());
-					if (ensure(SceneProxy != nullptr))
+				if (WaterZone->HasActorRegisteredAllComponents())
+				{
+					if (WaterZone->IsLocalOnlyTessellationEnabled())
 					{
-						NonDataViewsQuadtreeKeys.Add(InView.State, SceneProxy->FindBestQuadTreeForViewLocation(&InView));
+						UWaterMeshComponent* WaterMeshComponent = WaterZone->GetWaterMeshComponent();
+						check(WaterMeshComponent != nullptr);
+
+						FWaterMeshSceneProxy* SceneProxy = static_cast<FWaterMeshSceneProxy*>(WaterMeshComponent->GetSceneProxy());
+						if (ensure(SceneProxy != nullptr))
+						{
+							NonDataViewsQuadtreeKeys.Add(InView.State, SceneProxy->FindBestQuadTreeForViewLocation(&InView));
 
-						UE_LOG(LogWater, Verbose, TEXT("Water Zone (%s) queued a search for the closest quadtree for a View (0x%p) which has no WaterInfo."), *GetNameSafe(WaterZone), &InView);
+							UE_LOG(LogWater, Verbose, TEXT("Water Zone (%s) queued a search for the closest quadtree for a View (0x%p) which has no WaterInfo."), *GetNameSafe(WaterZone), &InView);
+						}
 					}
 				}
 			}
-		}	
+		}
+
+		// this is needed because when switching back from PIE to Editor, there isn't a proper initialization of
+		// the editor WaterViewExtension, since during PIE it is kept alive and not updated. Because of that when
+		// PIE ends and we start updating the editor view extension, it can contain out of date values which can cause
+		// water rendering artifacts until stepping out of bounds forces the first update. To fix that we keep
+		// track if the last ViewExtension that was updated changed, so we can reset CurrentNumViews which forces an update on the next frame.
+		// TODO: look into any existing callbacks which would allow to achieve this in a cleaner way.
+		static FWaterViewExtension* LastWaterViewExtensionUpdated = nullptr;
+		if (LastWaterViewExtensionUpdated != this)
+		{
+			CurrentNumViews = 0;
+			LastWaterViewExtensionUpdated = this;
+		}
 	}
+	
 
 	// The logic in UpdateGPUBuffers() used to be done in SetupViewFamily(). However, SetupView() (which is responsible for water info rendering) potentially modifies the WaterZone but is called after SetupViewFamily().
 	// This can lead to visual artifacts due to outdated data in the GPU buffers.
 	
 	UpdateGPUBuffers();
-
-	// this is needed because when switching back from PIE to Editor, there isn't a proper initialization of
-	// the editor WaterViewExtension, since during PIE it is kept alive and not updated. Because of that when
-	// PIE ends and we start updating the editor view extension, it can contain out of date values which can cause
-	// water rendering artifacts until stepping out of bounds forces the first update. To fix that we keep
-	// track if the last ViewExtension that was updated changed, so we can reset CurrentNumViews which forces an update on the next frame.
-	// TODO: look into any existing callbacks which would allow to achieve this in a cleaner way.
-	static FWaterViewExtension* LastWaterViewExtensionUpdated = nullptr;
-	if (LastWaterViewExtensionUpdated != this)
-	{
-		CurrentNumViews = 0;
-		LastWaterViewExtensionUpdated = this;
-	}
 }
 
 void FWaterViewExtension::DrawDebugInfo(FSceneView& InView, AWaterZone* WaterZone)
