Removing a WaterZoneActor ensure which was added recently and turned out to be too aggressive, since there are cases in which it is possible that the WaterZone's WaterZoneInfo doesn't exist yet.

diff --git a/Engine/Plugins/Experimental/Water/Source/Runtime/Private/WaterZoneActor.cpp b/Engine/Plugins/Experimental/Water/Source/Runtime/Private/WaterZoneActor.cpp
--- a/Engine/Plugins/Experimental/Water/Source/Runtime/Private/WaterZoneActor.cpp
+++ b/Engine/Plugins/Experimental/Water/Source/Runtime/Private/WaterZoneActor.cpp
@@ -134,7 +134,7 @@ void AWaterZone::GetAllDynamicWaterInfoBounds(TArray<FBox>& OutBounds) const
 	{
 		const FWaterViewExtension::FWaterZoneInfo* WaterZoneInfo = WaterViewExtension->WaterZoneInfos.Find(this);
 
-		if (ensure(WaterZoneInfo != nullptr))
+		if (WaterZoneInfo != nullptr)
 		{
 			for (const FWaterViewExtension::FWaterZoneInfo::FWaterZoneViewInfo& ViewInfo : WaterZoneInfo->ViewInfos)
 			{
@@ -150,6 +150,10 @@ void AWaterZone::GetAllDynamicWaterInfoBounds(TArray<FBox>& OutBounds) const
 				OutBounds.Emplace(FBox(Center - HalfExtent, Center + HalfExtent));
 			}
 		}
+		else
+		{
+			UE_LOG(LogWater, Verbose, TEXT("AWaterZone (%s) GetAllDynamicWaterInfoBounds did not find any WaterZoneInfo associated to this WaterZone in FWaterViewExtension"), *GetNameSafe(this));
+		}
 	}
 }
 
@@ -159,7 +163,7 @@ void AWaterZone::GetAllDynamicWaterInfoCenters(TArray<FVector>& OutCenters) cons
 	{
 		const FWaterViewExtension::FWaterZoneInfo* WaterZoneInfo = WaterViewExtension->WaterZoneInfos.Find(this);
 
-		if (ensure(WaterZoneInfo != nullptr))
+		if (WaterZoneInfo != nullptr)
 		{
 			for (const FWaterViewExtension::FWaterZoneInfo::FWaterZoneViewInfo& ViewInfo : WaterZoneInfo->ViewInfos)
 			{
@@ -173,6 +177,10 @@ void AWaterZone::GetAllDynamicWaterInfoCenters(TArray<FVector>& OutCenters) cons
 				OutCenters.Emplace(Center);
 			}
 		}
+		else
+		{
+			UE_LOG(LogWater, Verbose, TEXT("AWaterZone (%s) GetAllDynamicWaterInfoCenters did not find any WaterZoneInfo associated to this WaterZone in FWaterViewExtension"), *GetNameSafe(this));
+		}
 	}
 }
 
