[Water]

Fix for oceans that have bCenterOnWaterZone = true not automatically finding their parent water zone when the water zone location is far from the origin.

diff --git a/Engine/Plugins/Experimental/Water/Source/Runtime/Private/WaterBodyOceanComponent.cpp b/Engine/Plugins/Experimental/Water/Source/Runtime/Private/WaterBodyOceanComponent.cpp
--- a/Engine/Plugins/Experimental/Water/Source/Runtime/Private/WaterBodyOceanComponent.cpp
+++ b/Engine/Plugins/Experimental/Water/Source/Runtime/Private/WaterBodyOceanComponent.cpp
@@ -489,7 +489,7 @@ void UWaterBodyOceanComponent::OnPostRegisterAllComponents()
 FBoxSphereBounds UWaterBodyOceanComponent::CalcBounds(const FTransform& LocalToWorld) const
 {
 	const FTransform& ComponentTransform = GetComponentTransform();
-	FVector RelativeLocationToZone = bCenterOnWaterZone ? ComponentTransform.InverseTransformPosition(FVector(SavedZoneLocation, 0.0)) : FVector::ZeroVector;
+	FVector RelativeLocationToZone = (bCenterOnWaterZone && !OwningWaterZone.IsNull()) ? ComponentTransform.InverseTransformPosition(FVector(SavedZoneLocation, 0.0)) : FVector::ZeroVector;
 	RelativeLocationToZone.Z = 0;
 
 	const FVector2D OceanExtentScaled = (OceanExtents / FVector2D(GetComponentScale())) / 2.;
