[Water]

Water info texture renders which occur around component registration/scene proxy recreation will fail to capture the updated scene proxy due to the deferred nature of scene proxy rebuilds. Instead, we now also capture all such cases by triggering a water info rebuild during the scene proxy creation.

diff --git a/Engine/Plugins/Experimental/Water/Source/Runtime/Private/WaterBodyInfoMeshComponent.cpp b/Engine/Plugins/Experimental/Water/Source/Runtime/Private/WaterBodyInfoMeshComponent.cpp
--- a/Engine/Plugins/Experimental/Water/Source/Runtime/Private/WaterBodyInfoMeshComponent.cpp
+++ b/Engine/Plugins/Experimental/Water/Source/Runtime/Private/WaterBodyInfoMeshComponent.cpp
@@ -5,6 +5,8 @@
 #include "StaticMeshSceneProxy.h"
 #include "UObject/UObjectIterator.h"
 #include "Rendering/CustomRenderPass.h"
+#include "WaterBodyActor.h"
+#include "WaterBodyComponent.h"
 
 #include UE_INLINE_GENERATED_CPP_BY_NAME(WaterBodyInfoMeshComponent)
 
@@ -28,6 +30,8 @@ static TAutoConsoleVariable<int32> CVarShowWaterInfoSceneProxies(
 	ECVF_RenderThreadSafe
 );
 
+extern TAutoConsoleVariable<int32> CVarWaterMeshGPUQuadTree;
+
 int32 GetWaterInfoRenderingMethod()
 {
 	static IConsoleVariable* CVarWaterInfoRenderMethod = IConsoleManager::Get().FindConsoleVariable(TEXT("r.Water.WaterInfo.RenderMethod"));
@@ -66,6 +70,23 @@ FWaterBodyInfoMeshSceneProxy::FWaterBodyInfoMeshSceneProxy(UWaterBodyInfoMeshCom
 	{
 		SetEnabled(false);
 	}
+
+	check(Component != nullptr);
+
+	if (AWaterBody* WaterBodyActor = Component->GetOwner<AWaterBody>())
+	{
+		if (UWaterBodyComponent* WaterBodyComponent = WaterBodyActor->GetWaterBodyComponent())
+		{
+			EWaterZoneRebuildFlags Flags = EWaterZoneRebuildFlags::UpdateWaterInfoTexture;
+			// The GPU quad tree is built by rasterizing the water info meshes into a texture.
+			// Just like the water info texture, the quadtree must also be updated if the proxy is rebuilt.
+			if (CVarWaterMeshGPUQuadTree.GetValueOnGameThread() == 1)
+			{
+				Flags |= EWaterZoneRebuildFlags::UpdateWaterMesh;
+			}
+			WaterBodyComponent->MarkOwningWaterZoneForRebuild(Flags);
+		}
+	}
 }
 
 bool FWaterBodyInfoMeshSceneProxy::GetMeshElement(int32 LODIndex, int32 BatchIndex, int32 ElementIndex, uint8 InDepthPriorityGroup, bool bUseSelectionOutline, bool bAllowPreCulledIndices, FMeshBatch& OutMeshBatch) const
diff --git a/Engine/Plugins/Experimental/Water/Source/Runtime/Private/WaterMeshComponent.cpp b/Engine/Plugins/Experimental/Water/Source/Runtime/Private/WaterMeshComponent.cpp
--- a/Engine/Plugins/Experimental/Water/Source/Runtime/Private/WaterMeshComponent.cpp
+++ b/Engine/Plugins/Experimental/Water/Source/Runtime/Private/WaterMeshComponent.cpp
@@ -34,7 +34,7 @@ static TAutoConsoleVariable<float> CVarWaterMeshLODScaleBias(
 	TEXT("This value is added to the LOD Scale of each Mesh Component. Negative values will lower the overall density/resolution or the vertex grid and make the LODs smaller, higher values will increase the density/resolution and make the LODs larger. Smallest value is -0.5. That will make the inner LOD as tight and optimized as possible"),
 	ECVF_Scalability);
 
-static TAutoConsoleVariable<int32> CVarWaterMeshGPUQuadTree(
+TAutoConsoleVariable<int32> CVarWaterMeshGPUQuadTree(
 	TEXT("r.Water.WaterMesh.GPUQuadTree"),
 	0,
 	TEXT("Builds the water quadtree on the GPU and does indirect draws of water tiles, driven by the GPU."),
@@ -368,7 +368,7 @@ void UWaterMeshComponent::RebuildWaterMesh(float InTileSize, const FIntPoint& In
 						Boxes.Add(Box);
 					}
 				}
-
+				
 				for (const FBox& Box : Boxes)
 				{
 #if !(UE_BUILD_SHIPPING || UE_BUILD_TEST)
diff --git a/Engine/Plugins/Experimental/Water/Source/Runtime/Private/WaterZoneActor.cpp b/Engine/Plugins/Experimental/Water/Source/Runtime/Private/WaterZoneActor.cpp
--- a/Engine/Plugins/Experimental/Water/Source/Runtime/Private/WaterZoneActor.cpp
+++ b/Engine/Plugins/Experimental/Water/Source/Runtime/Private/WaterZoneActor.cpp
@@ -287,7 +287,6 @@ void AWaterZone::MarkForRebuild(EWaterZoneRebuildFlags Flags, const FBox2D& Upda
 		{
 			UE_LOG(LogWater, Verbose, TEXT("AWaterZone (%s) UpdateWaterMesh in region {%s} (triggered by %s)"), *GetNameSafe(this), *UpdateRegion.ToString(), *GetNameSafe(DebugRequestingObject));
 			WaterMesh->MarkWaterMeshGridDirty();
-			WaterMesh->MarkRenderStateDirty();
 		}
 		if (EnumHasAnyFlags(Flags, EWaterZoneRebuildFlags::UpdateWaterInfoTexture))
 		{
