changed WaterZoneActor::GetDynamicWaterInfoCenter API to not trigger an ensure, and return a ValidZoneLocation flag in case we are querying uninitialized WaterZoneInfo data

diff --git a/Engine/Plugins/Experimental/Water/Source/Runtime/Private/WaterViewExtension.cpp b/Engine/Plugins/Experimental/Water/Source/Runtime/Private/WaterViewExtension.cpp
--- a/Engine/Plugins/Experimental/Water/Source/Runtime/Private/WaterViewExtension.cpp
+++ b/Engine/Plugins/Experimental/Water/Source/Runtime/Private/WaterViewExtension.cpp
@@ -799,7 +799,7 @@ void FWaterViewExtension::RemoveWaterZone(AWaterZone* InWaterZone)
 	UE_LOG(LogWater, Verbose, TEXT("Water Zone (%s): RemoveWaterZone was called."), *GetNameSafe(InWaterZone));
 }
 
-FVector FWaterViewExtension::GetZoneLocation(const AWaterZone* InWaterZone, int32 PlayerIndex) const
+bool FWaterViewExtension::GetZoneLocation(const AWaterZone* InWaterZone, int32 PlayerIndex, FVector& OutLocation) const
 {
 	if (ensure(InWaterZone->HasActorRegisteredAllComponents()))
 	{
@@ -808,14 +808,17 @@ FVector FWaterViewExtension::GetZoneLocation(const AWaterZone* InWaterZone, int3
 		{
 			int32 ViewPlayerIndex = GetViewIndex(PlayerIndex);
 
-			if (ensure(ViewPlayerIndex != INDEX_NONE && ViewPlayerIndex < WaterZoneInfo->ViewInfos.Num()))
+			if (ViewPlayerIndex != INDEX_NONE && ViewPlayerIndex < WaterZoneInfo->ViewInfos.Num())
 			{
-				return WaterZoneInfo->ViewInfos[ViewPlayerIndex].Center;
-			}		
+				OutLocation = WaterZoneInfo->ViewInfos[ViewPlayerIndex].Center;
+				return true;
+			}
 		}
 	}
 
-	return FVector();
+	UE_LOG(LogWater, Verbose, TEXT("Water Zone (%s) called FWaterViewExtension::GetZoneLocation() but didn't get a valid location because of missing/uninitialized WaterZoneInfo->ViewInfos."), *GetNameSafe(InWaterZone));
+
+	return false;
 }
 
 void FWaterViewExtension::CreateSceneProxyQuadtrees(FWaterMeshSceneProxy* SceneProxy)
diff --git a/Engine/Plugins/Experimental/Water/Source/Runtime/Private/WaterZoneActor.cpp b/Engine/Plugins/Experimental/Water/Source/Runtime/Private/WaterZoneActor.cpp
--- a/Engine/Plugins/Experimental/Water/Source/Runtime/Private/WaterZoneActor.cpp
+++ b/Engine/Plugins/Experimental/Water/Source/Runtime/Private/WaterZoneActor.cpp
@@ -750,37 +750,41 @@ FVector AWaterZone::GetDynamicWaterInfoExtent() const
 	}
 }
 
-FVector AWaterZone::GetDynamicWaterInfoCenter(int32 PlayerIndex) const
+bool AWaterZone::GetDynamicWaterInfoCenter(int32 PlayerIndex, FVector& OutCenter) const
 {
 	const UWorld* World = GetWorld();
 
 	check(World != nullptr);
 
+	bool bHasValidZoneLocation = false;
 	FVector Center = GetActorLocation();
 
 	if (const FWaterViewExtension* WaterViewExtension = UWaterSubsystem::GetWaterViewExtension(World))
 	{
-		Center = WaterViewExtension->GetZoneLocation(this, PlayerIndex);
+		bHasValidZoneLocation = WaterViewExtension->GetZoneLocation(this, PlayerIndex, Center);
 	}
 	
-	return Center;
+	return bHasValidZoneLocation;
 }
 
-FBox AWaterZone::GetDynamicWaterInfoBounds(int32 PlayerIndex) const
+bool AWaterZone::GetDynamicWaterInfoBounds(int32 PlayerIndex, FBox& OutBounds) const
 {
 	const UWorld* World = GetWorld();
 
 	check(World != nullptr);
 
+	bool bHasValidZoneLocation = false;
 	FVector Center = GetActorLocation();
 
 	if (const FWaterViewExtension* WaterViewExtension = UWaterSubsystem::GetWaterViewExtension(World))
 	{
-		Center = WaterViewExtension->GetZoneLocation(this, PlayerIndex);
+		bHasValidZoneLocation = WaterViewExtension->GetZoneLocation(this, PlayerIndex, Center);
 	}
 	
 	const FVector WaterInfoHalfExtents = GetDynamicWaterInfoExtent() / 2.;
-	return FBox(Center - WaterInfoHalfExtents, Center + WaterInfoHalfExtents);
+	OutBounds = FBox(Center - WaterInfoHalfExtents, Center + WaterInfoHalfExtents);
+
+	return bHasValidZoneLocation;
 }
 
 FVector AWaterZone::GetDynamicWaterInfoCenter() const
@@ -788,7 +792,9 @@ FVector AWaterZone::GetDynamicWaterInfoCenter() const
 	// The following index assume that there is no split screen and will request the position of the first player's water view.
 	// This call is a temporary fallback and is deprecated in favor of the overload which takes a specific player view.
 	constexpr int32 PlayerIndex = 0;
-	return GetDynamicWaterInfoCenter(PlayerIndex);
+	FVector Center;
+	GetDynamicWaterInfoCenter(PlayerIndex, Center);
+	return Center;
 }
 
 void AWaterZone::OnLevelAddedToWorld(ULevel* InLevel, UWorld* InWorld)
diff --git a/Engine/Plugins/Experimental/Water/Source/Runtime/Public/WaterViewExtension.h b/Engine/Plugins/Experimental/Water/Source/Runtime/Public/WaterViewExtension.h
--- a/Engine/Plugins/Experimental/Water/Source/Runtime/Public/WaterViewExtension.h
+++ b/Engine/Plugins/Experimental/Water/Source/Runtime/Public/WaterViewExtension.h
@@ -36,7 +36,7 @@ public:
 	void AddWaterZone(AWaterZone* InWaterZone);
 	void RemoveWaterZone(AWaterZone* InWaterZone);
 
-	FVector GetZoneLocation(const AWaterZone* InWaterZone, int32 PlayerIndex) const;
+	bool GetZoneLocation(const AWaterZone* InWaterZone, int32 PlayerIndex, FVector& OutLocation) const;
 
 	void CreateSceneProxyQuadtrees(FWaterMeshSceneProxy* SceneProxy);
 
diff --git a/Engine/Plugins/Experimental/Water/Source/Runtime/Public/WaterZoneActor.h b/Engine/Plugins/Experimental/Water/Source/Runtime/Public/WaterZoneActor.h
--- a/Engine/Plugins/Experimental/Water/Source/Runtime/Public/WaterZoneActor.h
+++ b/Engine/Plugins/Experimental/Water/Source/Runtime/Public/WaterZoneActor.h
@@ -66,10 +66,10 @@ public:
 
 	uint32 GetVelocityBlurRadius() const { return VelocityBlurRadius; }
 
-	/** Retrieves the dynamic water info center for a specific player index. */
-	FVector GetDynamicWaterInfoCenter(int32 PlayerIndex) const;
-	/** Retrieves the dynamic water info bounds for a specific player index. */
-	FBox GetDynamicWaterInfoBounds(int32 PlayerIndex) const;
+	/** Retrieves the dynamic water info center for a specific player index in OutCenter. Returns true if a valid zone location has been found. */
+	bool GetDynamicWaterInfoCenter(int32 PlayerIndex, FVector& OutCenter) const;
+	/** Retrieves the dynamic water info bounds for a specific player index in OutBounds. Returns true if a valid zone location has been found.*/
+	bool GetDynamicWaterInfoBounds(int32 PlayerIndex, FBox& OutBounds) const;
 
 	FVector GetDynamicWaterInfoExtent() const;
 
diff --git a/Engine/Plugins/Experimental/WaterAdvanced/Source/WaterAdvanced/Private/ShallowWaterRiverActor.cpp b/Engine/Plugins/Experimental/WaterAdvanced/Source/WaterAdvanced/Private/ShallowWaterRiverActor.cpp
--- a/Engine/Plugins/Experimental/WaterAdvanced/Source/WaterAdvanced/Private/ShallowWaterRiverActor.cpp
+++ b/Engine/Plugins/Experimental/WaterAdvanced/Source/WaterAdvanced/Private/ShallowWaterRiverActor.cpp
@@ -348,12 +348,13 @@ void UShallowWaterRiverComponent::Rebuild()
 			}			
 
 			const int32 PlayerIndex = 0;
-			const FVector2D ZoneLocation = FVector2D(WaterZone->GetDynamicWaterInfoCenter(PlayerIndex));
+			FVector ZoneLocation;
+			WaterZone->GetDynamicWaterInfoCenter(PlayerIndex, ZoneLocation);
 			const FVector2D ZoneExtent = FVector2D(WaterZone->GetDynamicWaterInfoExtent());
 			const FVector2D WaterHeightExtents = FVector2D(WaterZone->GetWaterHeightExtents());
 			const float GroundZMin = WaterZone->GetGroundZMin();
 
-			RiverSimSystem->SetVariableVec2(FName("WaterZoneLocation"), ZoneLocation);
+			RiverSimSystem->SetVariableVec2(FName("WaterZoneLocation"), FVector2D(ZoneLocation));
 			RiverSimSystem->SetVariableVec2(FName("WaterZoneExtent"), ZoneExtent);
 			RiverSimSystem->SetVariableInt(FName("WaterZoneIdx"), WaterZone->GetWaterZoneIndex());
 			
diff --git a/Engine/Plugins/Experimental/WaterAdvanced/Source/WaterAdvanced/Private/ShallowWaterSubsystem.cpp b/Engine/Plugins/Experimental/WaterAdvanced/Source/WaterAdvanced/Private/ShallowWaterSubsystem.cpp
--- a/Engine/Plugins/Experimental/WaterAdvanced/Source/WaterAdvanced/Private/ShallowWaterSubsystem.cpp
+++ b/Engine/Plugins/Experimental/WaterAdvanced/Source/WaterAdvanced/Private/ShallowWaterSubsystem.cpp
@@ -1222,13 +1222,14 @@ void UShallowWaterSubsystem::TryGetOrWaitForWaterInfoTextureFromWaterBodies(TSet
 				{
 					// The following index assume that there is no split screen support and will request the position of the first player's water view.
 					const int32 PlayerIndex = 0;
-					const FVector2D ZoneLocation = FVector2D(WaterZone->GetDynamicWaterInfoCenter(PlayerIndex));
+					FVector ZoneLocation;
+					WaterZone->GetDynamicWaterInfoCenter(PlayerIndex, ZoneLocation);
 
 					const FVector2D ZoneExtent = FVector2D(WaterZone->GetDynamicWaterInfoExtent());
 					const FVector2D WaterHeightExtents = FVector2D(WaterZone->GetWaterHeightExtents());
 					const float GroundZMin = WaterZone->GetGroundZMin();
 
-					ShallowWaterNiagaraSimulation->SetVariableVec2(FName("WaterZoneLocation"), ZoneLocation);
+					ShallowWaterNiagaraSimulation->SetVariableVec2(FName("WaterZoneLocation"), FVector2D(ZoneLocation));
 					ShallowWaterNiagaraSimulation->SetVariableVec2(FName("WaterZoneExtent"), ZoneExtent);
 					ShallowWaterNiagaraSimulation->SetVariableInt(FName("WaterZoneIdx"), WaterZone->GetWaterZoneIndex());
 				}
