[Water]

Deprecated water info RM 0. RM 0 has not been exercised since the introduction of RM1/RM2 and isn't fully compatible with the new water info feature set.

diff --git a/Engine/Plugins/Experimental/Water/Source/Runtime/Private/WaterInfoRendering.cpp b/Engine/Plugins/Experimental/Water/Source/Runtime/Private/WaterInfoRendering.cpp
--- a/Engine/Plugins/Experimental/Water/Source/Runtime/Private/WaterInfoRendering.cpp
+++ b/Engine/Plugins/Experimental/Water/Source/Runtime/Private/WaterInfoRendering.cpp
@@ -422,198 +422,6 @@ private:
 	TMap<FLandscapeRenderSystem*, int8> LandscapeLODOverridesToRestore;
 };
 
-static void UpdateWaterInfoRendering_RenderThread(
-	FRHICommandListImmediate& RHICmdList,
-	const FUpdateWaterInfoParams& Params)
-{
-	FScopedLandscapeLODOverride ScopedLandscapeLODOverride(Params);
-
-	static bool bIsUpdatingWaterInfo = false;
-	check(bIsUpdatingWaterInfo == false);
-	bIsUpdatingWaterInfo = true;
-	ON_SCOPE_EXIT {bIsUpdatingWaterInfo = false;};
-
-	FMaterialRenderProxy::UpdateDeferredCachedUniformExpressions();
-
-	FRenderTarget* RenderTarget = Params.RenderTarget;
-	FTexture* OutputTexture = Params.OutputTexture;
-
-	TRefCountPtr<IPooledRenderTarget> ExtractedDepthTexture;
-	TRefCountPtr<IPooledRenderTarget> ExtractedColorTexture;
-
-
-	// Depth-only pass for actors which are considered the ground for water rendering
-	{
-		TRACE_CPUPROFILER_EVENT_SCOPE(WaterInfo::DepthPass);
-
-		FSceneRenderer* DepthRenderer = Params.DepthRenderer;
-		FSceneInterface* Scene = DepthRenderer->ViewFamily.Scene;
-
-		DepthRenderer->RenderThreadBegin(RHICmdList);
-
-		SCOPED_DRAW_EVENT(RHICmdList, DepthRendering_RT);
-
-		FRDGBuilder GraphBuilder(RHICmdList, RDG_EVENT_NAME("WaterInfoDepthRendering"), ERDGBuilderFlags::Parallel);
-
-		FRDGTextureRef TargetTexture = RegisterExternalTexture(GraphBuilder, RenderTarget->GetRenderTargetTexture(), TEXT("WaterDepthTarget"));
-
-		FRDGTextureDesc DepthTextureDesc(TargetTexture->Desc);
-		FRDGTextureRef DepthTexture = GraphBuilder.CreateTexture(DepthTextureDesc, TEXT("WaterDepthTexture"));
-
-		FViewInfo& View = DepthRenderer->Views[0];
-
-		AddClearRenderTargetPass(GraphBuilder, TargetTexture, FLinearColor::Black, View.UnscaledViewRect);
-
-		View.bDisableQuerySubmissions = true;
-		View.bIgnoreExistingQueries = true;
-
-		{
-			RDG_RHI_EVENT_SCOPE(GraphBuilder, RenderWaterInfoDepth);
-			DepthRenderer->Render(GraphBuilder);
-		}
-		
-		if (Scene->GetShadingPath() == EShadingPath::Mobile)
-		{
-			RDG_EVENT_SCOPE(GraphBuilder, "CaptureSceneColor");
-			CopySceneCaptureComponentToTarget(
-				GraphBuilder,
-				DepthTexture,
-				DepthRenderer->ViewFamily,
-				MakeStridedViewOfBase<const FSceneView>(MakeArrayView(DepthRenderer->Views)));
-		}
-		else
-		{
-			FRHICopyTextureInfo CopyInfo;
-			AddCopyTexturePass(GraphBuilder, TargetTexture, DepthTexture, CopyInfo);
-		}
-		
-		// We currently can't have multiple scene renderers run within the same RDGBuilder. Therefore, we must
-		// extract the texture to allow it to survive until the water info pass which runs in a later RDG graph (however, still within the same frame).
-		GraphBuilder.QueueTextureExtraction(DepthTexture, &ExtractedDepthTexture, ERDGResourceExtractionFlags::AllowTransient);
-
-		GraphBuilder.Execute();
-
-		DepthRenderer->RenderThreadEnd(RHICmdList);
-	}
-
-	// Render the water bodies' data including flow, zoffset, depth
-	{
-		TRACE_CPUPROFILER_EVENT_SCOPE(WaterInfo::ColorPass);
-
-		FSceneRenderer* ColorRenderer = Params.ColorRenderer;
-		FSceneInterface* Scene = ColorRenderer->ViewFamily.Scene;
-
-		FEnableWaterInfoMeshesScoped ScopedWaterInfoMeshEnable(ColorRenderer->Scene, ColorRenderer->Views[0]);
-
-		ColorRenderer->RenderThreadBegin(RHICmdList);
-
-		FViewInfo& View = ColorRenderer->Views[0];
-
-		SCOPED_DRAW_EVENT(RHICmdList, ColorRendering_RT);
-
-		FRDGBuilder GraphBuilder(RHICmdList, RDG_EVENT_NAME("WaterInfoColorRendering"), ERDGBuilderFlags::Parallel);
-
-		FRDGTextureRef TargetTexture = RegisterExternalTexture(GraphBuilder, RenderTarget->GetRenderTargetTexture(), TEXT("WaterColorTarget"));
-
-		FRDGTextureDesc ColorTextureDesc(TargetTexture->Desc);
-		ColorTextureDesc.Format = PF_A32B32G32R32F;
-		FRDGTextureRef ColorTexture = GraphBuilder.CreateTexture(ColorTextureDesc, TEXT("WaterColorTexture"));
-
-		AddClearRenderTargetPass(GraphBuilder, TargetTexture, FLinearColor::Black, View.UnscaledViewRect);
-
-		View.bDisableQuerySubmissions = true;
-		View.bIgnoreExistingQueries = true;
-
-		{
-			RDG_RHI_EVENT_SCOPE(GraphBuilder, RenderWaterInfoColor);
-			ColorRenderer->Render(GraphBuilder);
-		}
-
-		// This CopySceneCaptureComponentToTarget is required on all platforms as it extracts a higher pixel depth texture than the scene render target so we can't just copy out of it
-		{
-			RDG_EVENT_SCOPE(GraphBuilder, "CaptureSceneColor");
-			CopySceneCaptureComponentToTarget(
-				GraphBuilder,
-				ColorTexture,
-				ColorRenderer->ViewFamily,
-				MakeStridedViewOfBase<const FSceneView>(MakeArrayView(ColorRenderer->Views)));
-		}
-
-		// We currently can't have multiple scene renderers run within the same RDGBuilder. Therefore, we must
-		// extract the texture to allow it to survive until the water info pass which runs in a later RDG graph (however, still within the same frame).
-		GraphBuilder.QueueTextureExtraction(ColorTexture, &ExtractedColorTexture, ERDGResourceExtractionFlags::AllowTransient);
-
-		GraphBuilder.Execute();
-		
-		ColorRenderer->RenderThreadEnd(RHICmdList);
-	}
-
-	// Depth-only pass for water body dilated sections which get composited back into the water info texture at a lower priority than regular water body data
-	{
-		TRACE_CPUPROFILER_EVENT_SCOPE(WaterInfo::DilationPass);
-
-		FSceneRenderer* DilationRenderer = Params.DilationRenderer;
-		FSceneInterface* Scene = DilationRenderer->ViewFamily.Scene;
-
-		FEnableWaterInfoMeshesScoped ScopedWaterInfoMeshEnable(DilationRenderer->Scene, DilationRenderer->Views[0]);
-
-		DilationRenderer->RenderThreadBegin(RHICmdList);
-
-		FViewInfo& View = DilationRenderer->Views[0];
-
-		SCOPED_DRAW_EVENT(RHICmdList, DilationRendering_RT);
-
-		FRDGBuilder GraphBuilder(RHICmdList, RDG_EVENT_NAME("WaterInfoDilationRendering"), ERDGBuilderFlags::Parallel);
-
-		FRDGTextureRef TargetTexture = RegisterExternalTexture(GraphBuilder, RenderTarget->GetRenderTargetTexture(), TEXT("WaterDilationTarget"));
-
-		FRDGTextureRef DepthTexture = GraphBuilder.RegisterExternalTexture(ExtractedDepthTexture);
-		FRDGTextureRef ColorTexture = GraphBuilder.RegisterExternalTexture(ExtractedColorTexture);
-
-		AddClearRenderTargetPass(GraphBuilder, TargetTexture, FLinearColor::Black, View.UnscaledViewRect);
-
-		View.bDisableQuerySubmissions = true;
-		View.bIgnoreExistingQueries = true;
-
-		{
-			RDG_RHI_EVENT_SCOPE(GraphBuilder, RenderWaterInfoDilation);
-			DilationRenderer->Render(GraphBuilder);
-		}
-
-		FRDGTextureDesc TextureDesc(TargetTexture->Desc);
-		FRDGTextureRef DilationTexture = GraphBuilder.CreateTexture(TextureDesc, TEXT("WaterDilationTexture"));
-
-		if (Scene->GetShadingPath() == EShadingPath::Mobile)
-		{
-			RDG_EVENT_SCOPE(GraphBuilder, "CaptureSceneColor");
-			CopySceneCaptureComponentToTarget(
-				GraphBuilder,
-				DilationTexture,
-				DilationRenderer->ViewFamily,
-				MakeStridedViewOfBase<const FSceneView>(MakeArrayView(DilationRenderer->Views)));
-		}
-		else
-		{
-			FRHICopyTextureInfo CopyInfo;
-			AddCopyTexturePass(GraphBuilder, TargetTexture, DilationTexture, CopyInfo);
-		}
-
-		FRDGTextureRef MergeTargetTexture = GraphBuilder.CreateTexture(TextureDesc, TEXT("WaterInfoMerged"));
-		MergeWaterInfoAndDepth(GraphBuilder, DilationRenderer->ViewFamily, DilationRenderer->Views[0], MergeTargetTexture, DepthTexture, ColorTexture, DilationTexture, Params);
-
-		FRDGTextureRef FinalizedTexture = GraphBuilder.CreateTexture(TargetTexture->Desc, TEXT("WaterInfoFinalized"));
-		FinalizeWaterInfo(GraphBuilder, DilationRenderer->ViewFamily, DilationRenderer->Views[0], MergeTargetTexture, FinalizedTexture, Params);
-		
-		FRDGTextureRef ShaderResourceTexture = RegisterExternalTexture(GraphBuilder, OutputTexture->TextureRHI, TEXT("WaterInfoResolve"));
-		AddCopyTexturePass(GraphBuilder, FinalizedTexture, ShaderResourceTexture);
-		GraphBuilder.Execute();
-
-		DilationRenderer->RenderThreadEnd(RHICmdList);
-	}
-
-	RHICmdList.Transition(FRHITransitionInfo(Params.OutputTexture->TextureRHI, ERHIAccess::RTV, ERHIAccess::SRVMask));
-}
-
 struct FCreateWaterInfoSceneRendererParams
 {
 public:
@@ -632,196 +440,6 @@ public:
 	TSet<FPrimitiveComponentId> ShowOnlyPrimitives;
 };
 
-static FSceneRenderer* CreateWaterInfoSceneRenderer(const FCreateWaterInfoSceneRendererParams& Params)
-{
-	TRACE_CPUPROFILER_EVENT_SCOPE(WaterInfo::CreateWaterInfoRenderer);
-
-	check(Params.Scene != nullptr)
-	check(Params.RenderTarget != nullptr)
-	check(Params.CaptureSource != ESceneCaptureSource::SCS_MAX)
-
-	FEngineShowFlags ShowFlags(ESFIM_Game);
-	ShowFlags.NaniteMeshes = 0;
-	ShowFlags.Atmosphere = 0;
-	ShowFlags.Lighting = 0;
-	ShowFlags.Bloom = 0;
-	ShowFlags.ScreenPercentage = 0;
-	ShowFlags.Translucency = 0;
-	ShowFlags.SeparateTranslucency = 0;
-	ShowFlags.AntiAliasing = 0;
-	ShowFlags.Fog = 0;
-	ShowFlags.VolumetricFog = 0;
-	ShowFlags.DynamicShadows = 0;
-
-	ShowFlags.SetDisableOcclusionQueries(true);
-	ShowFlags.SetVirtualShadowMapPersistentData(false);
-	
-	FSceneViewFamilyContext ViewFamily(FSceneViewFamily::ConstructionValues(
-		Params.RenderTarget,
-		Params.Scene,
-		ShowFlags)
-		.SetRealtimeUpdate(false)
-		.SetResolveScene(false));
-	ViewFamily.SceneCaptureSource = Params.CaptureSource;
-
-	FSceneViewInitOptions ViewInitOptions;
-	ViewInitOptions.SetViewRectangle(FIntRect(0, 0, Params.RenderTargetSize.X, Params.RenderTargetSize.Y));
-	ViewInitOptions.ViewFamily = &ViewFamily;
-	ViewInitOptions.ViewActor = Params.Context.ZoneToRender;
-	ViewInitOptions.ViewRotationMatrix = Params.ViewRotationMatrix;
-	ViewInitOptions.ViewOrigin = Params.ViewLocation;
-	ViewInitOptions.BackgroundColor = FLinearColor::Black;
-	ViewInitOptions.OverrideFarClippingPlaneDistance = -1.f;
-	ViewInitOptions.SceneViewStateInterface = nullptr;
-	ViewInitOptions.ProjectionMatrix = Params.ProjectionMatrix;
-	ViewInitOptions.LODDistanceFactor = 0.001f;
-	ViewInitOptions.OverlayColor = FLinearColor::Black;
-	// Must be set to false to prevent the renders from using different VSM page pool sizes leading to unnecessary reallocations.
-	ViewInitOptions.bIsSceneCapture = false;
-
-	if (ViewFamily.Scene->GetWorld() != nullptr && ViewFamily.Scene->GetWorld()->GetWorldSettings() != nullptr)
-	{
-		ViewInitOptions.WorldToMetersScale = ViewFamily.Scene->GetWorld()->GetWorldSettings()->WorldToMeters;
-	}
-
-	FSceneView* View = new FSceneView(ViewInitOptions);
-	View->GPUMask = FRHIGPUMask::All();
-	View->bOverrideGPUMask = true;
-	View->AntiAliasingMethod = EAntiAliasingMethod::AAM_None;
-	View->SetupAntiAliasingMethod();
-
-	View->ShowOnlyPrimitives = Params.ShowOnlyPrimitives;
-
-	ViewFamily.Views.Add(View);
-
-	View->StartFinalPostprocessSettings(Params.ViewLocation);
-	View->EndFinalPostprocessSettings(ViewInitOptions);
-
-	ViewFamily.SetScreenPercentageInterface(new FLegacyScreenPercentageDriver(ViewFamily, 1.f));
-
-	ViewFamily.ViewExtensions = GEngine->ViewExtensions->GatherActiveExtensions(FSceneViewExtensionContext(ViewFamily.Scene));
-	for (const FSceneViewExtensionRef& Extension : ViewFamily.ViewExtensions)
-	{
-		Extension->SetupViewFamily(ViewFamily);
-		Extension->SetupView(ViewFamily, *View);
-	}
-
-	return FSceneRenderer::CreateSceneRenderer(&ViewFamily, nullptr);
-}
-
-void UpdateWaterInfoRendering(
-	FSceneInterface* Scene,
-	const WaterInfo::FRenderingContext& Context,
-	const FVector& WaterInfoCenter)
-{
-	TRACE_CPUPROFILER_EVENT_SCOPE(WaterInfo::UpdateWaterInfoRendering);
-
-	static auto* CVarRenderCaptureNextWaterInfoDraws = IConsoleManager::Get().FindConsoleVariable(TEXT("r.Water.WaterInfo.RenderCaptureNextWaterInfoDraws"));
-	int32 RenderCaptureNextWaterInfoDraws = CVarRenderCaptureNextWaterInfoDraws ? CVarRenderCaptureNextWaterInfoDraws->AsVariableInt()->GetValueOnGameThread() : 0;
-	RenderCaptureInterface::FScopedCapture RenderCapture((RenderCaptureNextWaterInfoDraws != 0), TEXT("RenderWaterInfo"));
-	if (RenderCaptureNextWaterInfoDraws != 0)
-	{
-		RenderCaptureNextWaterInfoDraws = FMath::Max(0, RenderCaptureNextWaterInfoDraws - 1);
-		CVarRenderCaptureNextWaterInfoDraws->SetWithCurrentPriority(RenderCaptureNextWaterInfoDraws);
-	}
-
-	if (!IsValid(Context.TextureRenderTarget) || Scene == nullptr)
-	{
-		return;
-	}
-	const FVector ZoneExtent = Context.ZoneToRender->GetDynamicWaterInfoExtent();
-
-	FVector ViewLocation = WaterInfoCenter;
-	ViewLocation.Z = Context.CaptureZ;
-
-	const FBox2D CaptureBounds(FVector2D(ViewLocation - ZoneExtent), FVector2D(ViewLocation + ZoneExtent));
-
-	// Zone rendering always happens facing towards negative z.
-	const FVector LookAt = ViewLocation - FVector(0.f, 0.f, 1.f);
-	
-	const FIntPoint CaptureExtent(Context.TextureRenderTarget->GetSurfaceWidth(), Context.TextureRenderTarget->GetSurfaceHeight());
-
-	// Initialize the generic parameters which are passed to each of the scene renderers
-	FCreateWaterInfoSceneRendererParams CreateSceneRendererParams(Context);
-	CreateSceneRendererParams.Scene = Scene;
-	CreateSceneRendererParams.RenderTarget = Context.TextureRenderTarget->GameThread_GetRenderTargetResource();
-	CreateSceneRendererParams.RenderTargetSize = CaptureExtent;
-	CreateSceneRendererParams.ViewLocation = ViewLocation;
-	CreateSceneRendererParams.ProjectionMatrix = BuildOrthoMatrix(ZoneExtent.X, ZoneExtent.Y);
-	CreateSceneRendererParams.ViewRotationMatrix = FLookAtMatrix(ViewLocation, LookAt, FVector(0.f, -1.f, 0.f));
-	CreateSceneRendererParams.ViewRotationMatrix = CreateSceneRendererParams.ViewRotationMatrix.RemoveTranslation();
-	CreateSceneRendererParams.ViewRotationMatrix.RemoveScaling();
-
-	TSet<FPrimitiveComponentId> ComponentsToRenderInDepthPass;
-	if (Context.GroundPrimitiveComponents.Num() > 0)
-	{
-		ComponentsToRenderInDepthPass.Reserve(Context.GroundPrimitiveComponents.Num());
-		for (TWeakObjectPtr<UPrimitiveComponent> GroundPrimComp : Context.GroundPrimitiveComponents)
-		{
-			if (GroundPrimComp.IsValid())
-			{
-				ComponentsToRenderInDepthPass.Add(GroundPrimComp.Get()->GetPrimitiveSceneId());
-			}
-		}
-	}
-	CreateSceneRendererParams.CaptureSource = SCS_DeviceDepth;
-	CreateSceneRendererParams.ShowOnlyPrimitives = MoveTemp(ComponentsToRenderInDepthPass);
-	FSceneRenderer* DepthRenderer = CreateWaterInfoSceneRenderer(CreateSceneRendererParams);
-	
-	TSet<FPrimitiveComponentId> ComponentsToRenderInColorPass;
-	TSet<FPrimitiveComponentId> ComponentsToRenderInDilationPass;
-	if (Context.WaterBodies.Num() > 0)
-	{
-		ComponentsToRenderInColorPass.Reserve(Context.WaterBodies.Num());
-		ComponentsToRenderInDilationPass.Reserve(Context.WaterBodies.Num());
-		for (const TWeakObjectPtr<UWaterBodyComponent> WaterBodyToRenderPtr : Context.WaterBodies)
-		{
-			if (UWaterBodyComponent* WaterBodyToRender = WaterBodyToRenderPtr.Get())
-			{
-				// Perform our own simple culling based on the known Capture bounds:
-				const FBox WaterBodyBounds = WaterBodyToRender->Bounds.GetBox();
-				if (CaptureBounds.Intersect(FBox2D(FVector2D(WaterBodyBounds.Min), FVector2D(WaterBodyBounds.Max))))
-				{
-					ComponentsToRenderInColorPass.Add(WaterBodyToRender->GetWaterInfoMeshComponent()->GetPrimitiveSceneId());
-					ComponentsToRenderInDilationPass.Add(WaterBodyToRender->GetDilatedWaterInfoMeshComponent()->GetPrimitiveSceneId());
-				}
-			}
-		}
-	}
-	CreateSceneRendererParams.CaptureSource = SCS_SceneColorSceneDepth;
-	CreateSceneRendererParams.ShowOnlyPrimitives = MoveTemp(ComponentsToRenderInColorPass);
-	FSceneRenderer* ColorRenderer = CreateWaterInfoSceneRenderer(CreateSceneRendererParams);
-
-	CreateSceneRendererParams.CaptureSource = SCS_DeviceDepth;
-	CreateSceneRendererParams.ShowOnlyPrimitives = MoveTemp(ComponentsToRenderInDilationPass);
-	FSceneRenderer* DilationRenderer = CreateWaterInfoSceneRenderer(CreateSceneRendererParams);
-
-	FTextureRenderTargetResource* TextureRenderTargetResource = Context.TextureRenderTarget->GameThread_GetRenderTargetResource();
-
-	FUpdateWaterInfoParams Params;
-	Params.Scene = Scene;
-	Params.DepthRenderer = DepthRenderer;
-	Params.ColorRenderer = ColorRenderer;
-	Params.DilationRenderer = DilationRenderer;
-	Params.RenderTarget = TextureRenderTargetResource;
-	Params.OutputTexture = TextureRenderTargetResource;
-	Params.CaptureZ = ViewLocation.Z;
-	Params.WaterHeightExtents = Context.ZoneToRender->GetWaterHeightExtents();
-	Params.GroundZMin = Context.ZoneToRender->GetGroundZMin();
-	Params.VelocityBlurRadius = Context.ZoneToRender->GetVelocityBlurRadius();
-	Params.WaterZoneExtents = ZoneExtent;
-
-	UE::RenderCommandPipe::FSyncScope SyncScope;
-
-	ENQUEUE_RENDER_COMMAND(WaterInfoCommand)(
-	[Params, ZoneName = Context.ZoneToRender->GetActorNameOrLabel()](FRHICommandListImmediate& RHICmdList)
-		{
-			SCOPED_DRAW_EVENTF(RHICmdList, WaterZoneInfoRendering_RT, TEXT("RenderWaterInfo_%s"), ZoneName);
-
-			UpdateWaterInfoRendering_RenderThread(RHICmdList, Params);
-		});
-}
-
 void UpdateWaterInfoRendering2(FSceneView& InView, const FRenderingContext& Context, int32 RenderTargetArrayLayer, const FVector& WaterInfoCenter)
 {
 	TRACE_CPUPROFILER_EVENT_SCOPE(WaterInfo::UpdateWaterInfoRendering2);
diff --git a/Engine/Plugins/Experimental/Water/Source/Runtime/Private/WaterInfoRendering.h b/Engine/Plugins/Experimental/Water/Source/Runtime/Private/WaterInfoRendering.h
--- a/Engine/Plugins/Experimental/Water/Source/Runtime/Private/WaterInfoRendering.h
+++ b/Engine/Plugins/Experimental/Water/Source/Runtime/Private/WaterInfoRendering.h
@@ -27,11 +27,6 @@ struct FRenderingContext
 	TArray<TWeakObjectPtr<UPrimitiveComponent>> GroundPrimitiveComponents;
 	float CaptureZ;
 };
-	
-void UpdateWaterInfoRendering(
-	FSceneInterface* Scene,
-	const FRenderingContext& Context,
-	const FVector& WaterInfoCenter);
 
 void UpdateWaterInfoRendering2(
 	FSceneView& InView, 
diff --git a/Engine/Plugins/Experimental/Water/Source/Runtime/Private/WaterViewExtension.cpp b/Engine/Plugins/Experimental/Water/Source/Runtime/Private/WaterViewExtension.cpp
--- a/Engine/Plugins/Experimental/Water/Source/Runtime/Private/WaterViewExtension.cpp
+++ b/Engine/Plugins/Experimental/Water/Source/Runtime/Private/WaterViewExtension.cpp
@@ -32,7 +32,7 @@ extern void OnCVarWaterInfoSceneProxiesValueChanged(IConsoleVariable*);
 static TAutoConsoleVariable<int32> CVarWaterInfoRenderMethod(
 	TEXT("r.Water.WaterInfo.RenderMethod"),
 	2,
-	TEXT("0: SceneCaptures, 1: Custom, 2: CustomRenderPasses"),
+	TEXT("1: Custom, 2: CustomRenderPasses"),
 	FConsoleVariableDelegate::CreateStatic(OnCVarWaterInfoSceneProxiesValueChanged),
 	ECVF_Default | ECVF_RenderThreadSafe);
 
@@ -432,13 +432,9 @@ void FWaterViewExtension::RenderWaterInfoTexture(FSceneViewFamily& InViewFamily,
 	check(ViewPlayerIndex != INDEX_NONE);
 
 	const UE::WaterInfo::FRenderingContext& Context(WaterZoneInfo->RenderContext);
-	// Old method of rendering the water info texture; uses scene captures
-	if (WaterInfoRenderMethod == 0)
-	{
-		UE::WaterInfo::UpdateWaterInfoRendering(Scene, Context, ZoneCenter);
-	}
+	
 	// Render the water info texture using custom render pass method
-	else if (WaterInfoRenderMethod == 2)
+	if (WaterInfoRenderMethod == 2)
 	{
 		UE::WaterInfo::UpdateWaterInfoRendering_CustomRenderPass(Scene, InViewFamily, Context, ViewPlayerIndex, ZoneCenter);
 	}
@@ -447,6 +443,10 @@ void FWaterViewExtension::RenderWaterInfoTexture(FSceneViewFamily& InViewFamily,
 	{
 		UE::WaterInfo::UpdateWaterInfoRendering2(InView, Context, ViewPlayerIndex, ZoneCenter);
 	}
+	else if (WaterInfoRenderMethod == 0)
+	{
+		UE_LOG(LogWater, Error, TEXT("Water Info Render Method 0 is deprecated and no longer functions! Please set r.Water.WaterInfo.RenderMethod to either 1 or 2"));
+	}
 }
 
 void FWaterViewExtension::SetupView(FSceneViewFamily& InViewFamily, FSceneView& InView)
