[Water]

Fixed a bug where water bodies could be added twice to the list of ActorsAffectingLandscape in the water brush. Added a new condition to AddActorInternal which ensures that duplicates can never be added.

diff --git a/Engine/Plugins/Experimental/Water/Source/Editor/Private/WaterLandscapeBrush.cpp b/Engine/Plugins/Experimental/Water/Source/Editor/Private/WaterLandscapeBrush.cpp
--- a/Engine/Plugins/Experimental/Water/Source/Editor/Private/WaterLandscapeBrush.cpp
+++ b/Engine/Plugins/Experimental/Water/Source/Editor/Private/WaterLandscapeBrush.cpp
@@ -41,7 +41,8 @@ void AWaterLandscapeBrush::AddActorInternal(AActor* Actor, const UWorld* ThisWor
 		!Actor->IsUnreachable() &&
 		Actor->GetLevel() != nullptr &&
 		!Actor->GetLevel()->bIsBeingRemoved &&
-		ThisWorld == Actor->GetWorld())
+		ThisWorld == Actor->GetWorld() &&
+		!ActorsAffectingLandscape.Contains(Cast<IWaterBrushActorInterface>(Actor)))
 	{
 		if (bModify)
 		{
@@ -163,15 +164,14 @@ void AWaterLandscapeBrush::UpdateActors(bool bInTriggerEvents)
 void AWaterLandscapeBrush::OnWaterBrushActorChanged(const IWaterBrushActorInterface::FWaterBrushActorChangedEventParams& InParams)
 {
 	AActor* Actor = CastChecked<AActor>(InParams.WaterBrushActor);
-	bool bAffectsLandscape = InParams.WaterBrushActor->AffectsLandscape();
-	bool bAffectsWaterMesh = InParams.WaterBrushActor->AffectsWaterMesh();
+	const bool bAffectsLandscape = InParams.WaterBrushActor->AffectsLandscape();
+	const bool bAffectsWaterMesh = InParams.WaterBrushActor->AffectsWaterMesh();
 
-	int32 ActorIndex = ActorsAffectingLandscape.IndexOfByKey(InParams.WaterBrushActor);
+	const bool bActorAlreadyAffectingLandscape = ActorsAffectingLandscape.Contains(InParams.WaterBrushActor);
 	// if the actor went from affecting landscape to non-affecting landscape (and vice versa), update the brush
 	bool bForceUpdateBrush = false;
-	bool bForceUpdateWaterMesh = false;
 	
-	if (bAffectsLandscape != (ActorIndex != INDEX_NONE))
+	if (bAffectsLandscape != bActorAlreadyAffectingLandscape)
 	{
 		if (bAffectsLandscape)
 		{
@@ -182,8 +182,6 @@ void AWaterLandscapeBrush::OnWaterBrushActorChanged(const IWaterBrushActorInterf
 			RemoveActorInternal(Actor);
 		}
 
-		// Force rebuild the mesh if a water body actor has been added or removed (islands don't affect the water mesh so it's not necessary for them): 
-		bForceUpdateWaterMesh = InParams.WaterBrushActor->CanEverAffectWaterMesh();
 		bForceUpdateBrush = true;
 	}
 
