Fixed rare timing issues caused by SceneProxy dangling pointers because of the WaterMesh scene proxy changing in between game thread and render thread.
This fixes occasional crashes when scrubbing replays.

diff --git a/Engine/Plugins/Experimental/Water/Source/Runtime/Private/WaterViewExtension.cpp b/Engine/Plugins/Experimental/Water/Source/Runtime/Private/WaterViewExtension.cpp
--- a/Engine/Plugins/Experimental/Water/Source/Runtime/Private/WaterViewExtension.cpp
+++ b/Engine/Plugins/Experimental/Water/Source/Runtime/Private/WaterViewExtension.cpp
@@ -581,7 +581,7 @@ void FWaterViewExtension::SetupView(FSceneViewFamily& InViewFamily, FSceneView&
 							FScopeLock AddQuadtreeUpdate(&QuadtreeUpdateLock);
 
 							FQuadtreeUpdateInfo QuadtreeUpdate;
-							QuadtreeUpdate.SceneProxy = SceneProxy;
+							QuadtreeUpdate.WaterZone = WaterZone;
 							QuadtreeUpdate.Location = WaterZoneViewInfo.UpdateBounds->GetCenter();
 							// we use the actual PlayerIndex instead of ViewPlayerIndex, since the latter can change if more views are added while a quadtree already exists
 							QuadtreeUpdate.Key = InView.PlayerIndex;
@@ -701,13 +701,22 @@ void FWaterViewExtension::PreRenderViewFamily_RenderThread(FRDGBuilder& GraphBui
 
 		for (const FQuadtreeUpdateInfo& QuadtreeUpdate : QuadtreeUpdates)
 		{
-			// TODO: We could add a CreateOrUpdateViewWaterQuadTree to WaterSceneProxy
-			// if the creation fails it means a quadtree for the current view already exists, so we just update it
-			if (!QuadtreeUpdate.SceneProxy->CreateViewWaterQuadTree(QuadtreeUpdate.Key, QuadtreeUpdate.Location))
+			UWaterMeshComponent* WaterMeshComponent = QuadtreeUpdate.WaterZone->GetWaterMeshComponent();
+			if (ensure(WaterMeshComponent != nullptr))
 			{
-				bool bResult = QuadtreeUpdate.SceneProxy->UpdateViewWaterQuadTree(QuadtreeUpdate.Key, QuadtreeUpdate.Location);
+				FWaterMeshSceneProxy* SceneProxy = static_cast<FWaterMeshSceneProxy*>(WaterMeshComponent->GetSceneProxy());
 
-				check(bResult);
+				if (SceneProxy != nullptr)
+				{
+					// TODO: We could add a CreateOrUpdateViewWaterQuadTree to WaterSceneProxy
+					// if the creation fails it means a quadtree for the current view already exists, so we just update it
+					if (!SceneProxy->CreateViewWaterQuadTree(QuadtreeUpdate.Key, QuadtreeUpdate.Location))
+					{
+						bool bResult = SceneProxy->UpdateViewWaterQuadTree(QuadtreeUpdate.Key, QuadtreeUpdate.Location);
+
+						check(bResult);
+					}
+				}
 			}
 		}
 
diff --git a/Engine/Plugins/Experimental/Water/Source/Runtime/Public/WaterViewExtension.h b/Engine/Plugins/Experimental/Water/Source/Runtime/Public/WaterViewExtension.h
--- a/Engine/Plugins/Experimental/Water/Source/Runtime/Public/WaterViewExtension.h
+++ b/Engine/Plugins/Experimental/Water/Source/Runtime/Public/WaterViewExtension.h
@@ -77,7 +77,7 @@ private:
 
 	struct FQuadtreeUpdateInfo
 	{
-		FWaterMeshSceneProxy* SceneProxy;
+		AWaterZone* WaterZone;
 		FVector2D Location;
 		int32 Key;
 	};
