Added perf tags to a couple of functions that get called on tick

#rnx

diff --git a/Engine/Plugins/Experimental/Water/Source/Runtime/Private/WaterMeshComponent.cpp b/Engine/Plugins/Experimental/Water/Source/Runtime/Private/WaterMeshComponent.cpp
--- a/Engine/Plugins/Experimental/Water/Source/Runtime/Private/WaterMeshComponent.cpp
+++ b/Engine/Plugins/Experimental/Water/Source/Runtime/Private/WaterMeshComponent.cpp
@@ -256,7 +256,7 @@ static bool IsMaterialUsedWithWater(const UMaterialInterface* InMaterial)
 
 void UWaterMeshComponent::RebuildWaterMesh(float InTileSize, const FIntPoint& InExtentInTiles)
 {
-	TRACE_CPUPROFILER_EVENT_SCOPE(RebuildWaterMesh);
+	TRACE_CPUPROFILER_EVENT_SCOPE(UWaterMeshComponent::RebuildWaterMesh);
 
 	AWaterZone* WaterZone = CastChecked<AWaterZone>(GetOwner());
 
@@ -541,6 +541,7 @@ void UWaterMeshComponent::RebuildWaterMesh(float InTileSize, const FIntPoint& In
 
 void UWaterMeshComponent::Update()
 {
+	TRACE_CPUPROFILER_EVENT_SCOPE(UWaterMeshComponent::Update);
 	bIsEnabled = FWaterUtils::IsWaterMeshEnabled(/*bIsRenderThread = */false) && FApp::CanEverRender();
 
 	// Early out
diff --git a/Engine/Plugins/Experimental/Water/Source/Runtime/Private/WaterSubsystem.cpp b/Engine/Plugins/Experimental/Water/Source/Runtime/Private/WaterSubsystem.cpp
--- a/Engine/Plugins/Experimental/Water/Source/Runtime/Private/WaterSubsystem.cpp
+++ b/Engine/Plugins/Experimental/Water/Source/Runtime/Private/WaterSubsystem.cpp
@@ -202,6 +202,8 @@ FWaterViewExtension* UWaterSubsystem::GetWaterViewExtension(const UWorld* InWorl
 
 void UWaterSubsystem::Tick(float DeltaTime)
 {
+	TRACE_CPUPROFILER_EVENT_SCOPE(UWaterSubsystem::Tick);
+
 	Super::Tick(DeltaTime);
 
 	check(GetWorld() != nullptr);
diff --git a/Engine/Plugins/Experimental/Water/Source/Runtime/Private/WaterZoneActor.cpp b/Engine/Plugins/Experimental/Water/Source/Runtime/Private/WaterZoneActor.cpp
--- a/Engine/Plugins/Experimental/Water/Source/Runtime/Private/WaterZoneActor.cpp
+++ b/Engine/Plugins/Experimental/Water/Source/Runtime/Private/WaterZoneActor.cpp
@@ -362,6 +362,7 @@ FVector2D AWaterZone::GetZoneExtent() const
 
 void AWaterZone::Update()
 {
+	TRACE_CPUPROFILER_EVENT_SCOPE(AWaterZone::Update);
 	if (bNeedsWaterInfoRebuild || (ForceUpdateWaterInfoNextFrames != 0))
 	{
 		ForceUpdateWaterInfoNextFrames = (ForceUpdateWaterInfoNextFrames < 0) ? ForceUpdateWaterInfoNextFrames : FMath::Max(0, ForceUpdateWaterInfoNextFrames - 1);
