FLinearColor: Disable core.uninitialized.ArraySubscript clang checker

#rnx

diff --git a/Engine/Plugins/Experimental/Water/Source/Runtime/Water.Build.cs b/Engine/Plugins/Experimental/Water/Source/Runtime/Water.Build.cs
--- a/Engine/Plugins/Experimental/Water/Source/Runtime/Water.Build.cs
+++ b/Engine/Plugins/Experimental/Water/Source/Runtime/Water.Build.cs
@@ -8,6 +8,8 @@ public class Water : ModuleRules
 	{
 		PCHUsage = ModuleRules.PCHUsageMode.UseExplicitOrSharedPCHs;
 
+		StaticAnalyzerDisabledCheckers.Add("core.uninitialized.ArraySubscript");
+
 		PrivateIncludePaths.AddRange(
 			new string[] {
 				System.IO.Path.Combine(GetModuleDirectory("Renderer"), "Private"),
diff --git a/Engine/Source/Runtime/Engine/Engine.Build.cs b/Engine/Source/Runtime/Engine/Engine.Build.cs
--- a/Engine/Source/Runtime/Engine/Engine.Build.cs
+++ b/Engine/Source/Runtime/Engine/Engine.Build.cs
@@ -13,6 +13,8 @@ public class Engine : ModuleRules
 
 		SharedPCHHeaderFile = "Public/EngineSharedPCH.h";
 
+		StaticAnalyzerDisabledCheckers.Add("core.uninitialized.ArraySubscript");
+
 		PublicIncludePathModuleNames.AddRange(
 			new string[] {
 				"AnimationCore",
diff --git a/Engine/Source/Runtime/SlateCore/SlateCore.Build.cs b/Engine/Source/Runtime/SlateCore/SlateCore.Build.cs
--- a/Engine/Source/Runtime/SlateCore/SlateCore.Build.cs
+++ b/Engine/Source/Runtime/SlateCore/SlateCore.Build.cs
@@ -6,6 +6,8 @@ public class SlateCore : ModuleRules
 {
 	public SlateCore(ReadOnlyTargetRules Target) : base(Target)
 	{
+		StaticAnalyzerDisabledCheckers.Add("core.uninitialized.ArraySubscript");
+
 		PrivateDefinitions.Add("UE_REPORT_SLATE_VECTOR_DEPRECATION=1");
 
 		PublicDependencyModuleNames.AddRange(
