[Water]

Adds a silent fixup to water bodies which have non-default details mode set on their critical components (spline/wbc). Non-default detail mode could lead to components being stripped during cook.

diff --git a/Engine/Plugins/Experimental/Water/Source/Runtime/Private/WaterBodyComponent.cpp b/Engine/Plugins/Experimental/Water/Source/Runtime/Private/WaterBodyComponent.cpp
--- a/Engine/Plugins/Experimental/Water/Source/Runtime/Private/WaterBodyComponent.cpp
+++ b/Engine/Plugins/Experimental/Water/Source/Runtime/Private/WaterBodyComponent.cpp
@@ -1599,7 +1599,7 @@ void UWaterBodyComponent::OnPostRegisterAllComponents()
 		}
 		UpdateWaterBodyRenderData();
 	}
-
+	
 	// Ensure that the sprite component is updated once the water body is fully setup after PostRegister.
 	UpdateWaterSpriteComponent();
 #endif // WITH_EDITOR
@@ -1748,6 +1748,32 @@ void UWaterBodyComponent::PostLoad()
 		
 		OwningWaterZone.Reset();
 	}
+
+	// If the detail mode of the water body component or the water spline component are ever not DM_Low,
+	// depending on per platform project settings they may be culled out in a cooked build for certain platforms
+	// breaking our assumptions that the waterbodycomponent/watersplinecomponent should always be present in game.
+	// This is a tricky issue to debug if it comes up, since there will be very little indication as to why the component
+	// went missing. There is not a valid use case for having DetailMode != DM_Low on either of these components so we
+	// are conservative here and simply prevent error if it's ever not the case while encouraging users to fixup their data.
+
+	if (DetailMode != DM_Low || (GetWaterSpline() && GetWaterSpline()->DetailMode != DM_Low))
+	{
+		DetailMode = DM_Low;
+		if (UWaterSplineComponent* SplineComp = GetWaterSpline())
+		{
+			SplineComp->DetailMode = DM_Low;
+		}
+
+		// Push the request to the user that they should mark this modified water package as dirty and resave it to persist the DetailMode change.
+		const IWaterModuleInterface& WaterModule = FModuleManager::GetModuleChecked<IWaterModuleInterface>("Water");
+		if (IWaterEditorServices* WaterEditorServices = WaterModule.GetWaterEditorServices())
+		{
+			if (GetWorld() && GetWorld()->WorldType == EWorldType::Editor)
+			{
+				WaterEditorServices->TryMarkPackageAsModified(GetPackage());
+			}
+		}
+	}
 #endif // WITH_EDITOR
 }
 
@@ -2109,7 +2135,6 @@ void UWaterBodyComponent::FixupEditorTransform()
 
 }
 
-
 void UWaterBodyComponent::CreateWaterSpriteComponent()
 {
 	UTexture2D* Texture = nullptr;
