- Update exposed parameter to allow water component or actors
- This allows you to select an actor in the outliner and have it resolve to the water component internall

#rnx

diff --git a/Engine/Plugins/Experimental/Water/Source/Runtime/Private/NiagaraDataInterfaceWater.cpp b/Engine/Plugins/Experimental/Water/Source/Runtime/Private/NiagaraDataInterfaceWater.cpp
--- a/Engine/Plugins/Experimental/Water/Source/Runtime/Private/NiagaraDataInterfaceWater.cpp
+++ b/Engine/Plugins/Experimental/Water/Source/Runtime/Private/NiagaraDataInterfaceWater.cpp
@@ -358,7 +358,14 @@ bool UNiagaraDataInterfaceWater::PerInstanceTick(void* PerInstanceData, FNiagara
 		}
 		else
 		{
-			InstData->WaterBodyComponent = SourceBodyComponent;
+			if (AActor* SourceActor = Cast<AActor>(SourceBodyComponent))
+			{
+				InstData->WaterBodyComponent = SourceActor->GetComponentByClass<UWaterBodyComponent>();
+			}
+			else
+			{
+				InstData->WaterBodyComponent = Cast<UWaterBodyComponent>(SourceBodyComponent);
+			}
 		}
 		InstData->bFindClosestBody = false;
 		InstData->WaterBodyChangeId = SourceBodyChangeId;
@@ -502,5 +509,10 @@ void UNiagaraDataInterfaceWater::GetWaveParamLookupTableOffset(FVectorVMExternal
 	}
 }
 
-#undef LOCTEXT_NAMESPACE
+void UNiagaraDataInterfaceWater::SetWaterBodyComponent(UWaterBodyComponent* InWaterBodyComponent)
+{
+	SourceBodyComponent = InWaterBodyComponent;
+	++SourceBodyChangeId;
+}
 
+#undef LOCTEXT_NAMESPACE
diff --git a/Engine/Plugins/Experimental/Water/Source/Runtime/Public/NiagaraDataInterfaceWater.h b/Engine/Plugins/Experimental/Water/Source/Runtime/Public/NiagaraDataInterfaceWater.h
--- a/Engine/Plugins/Experimental/Water/Source/Runtime/Public/NiagaraDataInterfaceWater.h
+++ b/Engine/Plugins/Experimental/Water/Source/Runtime/Public/NiagaraDataInterfaceWater.h
@@ -40,11 +40,7 @@ public:
 	void GetWaveParamLookupTableOffset(FVectorVMExternalFunctionContext& Context);
 
 	/** Sets the current water body to be used by this data interface */
-	void SetWaterBodyComponent(UWaterBodyComponent* InWaterBodyComponent)
-	{
-		SourceBodyComponent = InWaterBodyComponent;
-		++SourceBodyChangeId;
-	}
+	void SetWaterBodyComponent(UWaterBodyComponent* InWaterBodyComponent);
 
 protected:
 #if WITH_EDITORONLY_DATA
@@ -63,8 +59,8 @@ private:
 	UPROPERTY(EditAnywhere, Category = "Water", meta = (EditCondition = "bEvaluateSystemDepth"))
 	bool bEvaluateSystemDepthPerFrame = true;
 
-	UPROPERTY(EditAnywhere, Category = "Water")
-	TObjectPtr<UWaterBodyComponent> SourceBodyComponent;
+	UPROPERTY(EditAnywhere, Category = "Water", meta = (DisplayName = "Source Actor Or Component", AllowedClasses = "/Script/Engine.WaterBodyComponent,/Script/Engine.Actor"))
+	TObjectPtr<UObject> SourceBodyComponent;
 
 	uint32 SourceBodyChangeId = 0;
 };
