[Water]

Removed the transient flag from the water brush capture/jumpflood components. The editor does not handle loading of non-transient-but-now-transient actor components very well for blueprint classes. In some _very_ rare cases, a previously serialized component will load and try to attach itself to the CDO of the blueprint rather than having the AttachParent property remapped to point to the instance itself.

diff --git a/Engine/Plugins/Experimental/Water/Source/Editor/Private/WaterBrushManager.h b/Engine/Plugins/Experimental/Water/Source/Editor/Private/WaterBrushManager.h
--- a/Engine/Plugins/Experimental/Water/Source/Editor/Private/WaterBrushManager.h
+++ b/Engine/Plugins/Experimental/Water/Source/Editor/Private/WaterBrushManager.h
@@ -30,10 +30,10 @@ class AWaterBrushManager : public AWaterLandscapeBrush
 public:
 	GENERATED_BODY()
 	
- 	UPROPERTY(Transient)
+ 	UPROPERTY()
 	TObjectPtr<USceneCaptureComponent2D> SceneCaptureComponent2D = nullptr;
 
-	UPROPERTY(Transient)
+	UPROPERTY()
 	TObjectPtr<UJumpFloodComponent2D> JumpFloodComponent2D = nullptr;
 
 	// RTs
