[Water]

Fix water body static meshes not able to be computed during cook.

diff --git a/Engine/Plugins/Experimental/Water/Source/Runtime/Private/WaterBodyComponent.cpp b/Engine/Plugins/Experimental/Water/Source/Runtime/Private/WaterBodyComponent.cpp
--- a/Engine/Plugins/Experimental/Water/Source/Runtime/Private/WaterBodyComponent.cpp
+++ b/Engine/Plugins/Experimental/Water/Source/Runtime/Private/WaterBodyComponent.cpp
@@ -2018,8 +2018,8 @@ void UWaterBodyComponent::SetWaterBodyStaticMeshEnabled(bool bEnabled)
 
 void UWaterBodyComponent::UpdateWaterBodyRenderData()
 {
-	// Don't need to update render data if we can't ever render
-	if (FApp::CanEverRender())
+	// Don't need to update render data if we can't ever render, unless we're cooking in which case we do need to regenerate the water mesh data for the game
+	if (FApp::CanEverRender() || IsRunningCookCommandlet())
 	{
 		if (const UWorld* World = GetWorld())
 		{
