[Water]

Fixed reference to deprecated edit condition on the underwater post process material property.

diff --git a/Engine/Plugins/Experimental/Water/Source/Runtime/Public/WaterBodyComponent.h b/Engine/Plugins/Experimental/Water/Source/Runtime/Public/WaterBodyComponent.h
--- a/Engine/Plugins/Experimental/Water/Source/Runtime/Public/WaterBodyComponent.h
+++ b/Engine/Plugins/Experimental/Water/Source/Runtime/Public/WaterBodyComponent.h
@@ -601,7 +601,7 @@ public:
 	TObjectPtr<UMaterialInterface> WaterStaticMeshMaterial;
 
 	/** Post process material to apply when the camera goes underwater (only available when bGenerateCollisions is true because collisions are needed to detect if it's under water). */
-	UPROPERTY(Category = Rendering, EditAnywhere, BlueprintReadOnly, meta = (EditCondition = "bGenerateCollisions", DisplayAfter = "WaterMaterial"))
+	UPROPERTY(Category = Rendering, EditAnywhere, BlueprintReadOnly, meta = (EditCondition = "bGenerateOverlapEvents", DisplayAfter = "WaterMaterial"))
 	TObjectPtr<UMaterialInterface> UnderwaterPostProcessMaterial;
 
 	UPROPERTY(Category = Rendering, EditAnywhere, BlueprintReadOnly, meta = (DisplayAfter = "WaterMaterial"))
