[Water]

Prevent users setting the velocity blur radius too high and crashing their gpu.

diff --git a/Engine/Plugins/Experimental/Water/Source/Runtime/Private/WaterInfoRendering.cpp b/Engine/Plugins/Experimental/Water/Source/Runtime/Private/WaterInfoRendering.cpp
--- a/Engine/Plugins/Experimental/Water/Source/Runtime/Private/WaterInfoRendering.cpp
+++ b/Engine/Plugins/Experimental/Water/Source/Runtime/Private/WaterInfoRendering.cpp
@@ -469,7 +469,7 @@ void UpdateWaterInfoRendering2(FSceneView& InView, const FRenderingContext& Cont
 	RenderingParams.CaptureZ = ViewLocation.Z;
 	RenderingParams.WaterHeightExtents = Context.ZoneToRender->GetWaterHeightExtents();
 	RenderingParams.GroundZMin = Context.ZoneToRender->GetGroundZMin();
-	RenderingParams.VelocityBlurRadius = Context.ZoneToRender->GetVelocityBlurRadius();
+	RenderingParams.VelocityBlurRadius = FMath::Clamp(Context.ZoneToRender->GetVelocityBlurRadius(), 0, 16);
 	RenderingParams.WaterZoneExtents = ZoneExtent;
 	RenderingParams.RenderTargetArrayLayer = RenderTargetArrayLayer;
 
@@ -760,7 +760,7 @@ void UpdateWaterInfoRendering_CustomRenderPass(
 	Params.CaptureZ = ViewLocation.Z;
 	Params.WaterHeightExtents = Context.ZoneToRender->GetWaterHeightExtents();
 	Params.GroundZMin = Context.ZoneToRender->GetGroundZMin();
-	Params.VelocityBlurRadius = Context.ZoneToRender->GetVelocityBlurRadius();
+	Params.VelocityBlurRadius = FMath::Clamp(Context.ZoneToRender->GetVelocityBlurRadius(), 0, 16);
 	Params.WaterZoneExtents = ZoneExtent;
 
 	PassInput.ShowOnlyPrimitives = MoveTemp(ComponentsToRenderInDilationPass);
diff --git a/Engine/Plugins/Experimental/Water/Source/Runtime/Public/WaterZoneActor.h b/Engine/Plugins/Experimental/Water/Source/Runtime/Public/WaterZoneActor.h
--- a/Engine/Plugins/Experimental/Water/Source/Runtime/Public/WaterZoneActor.h
+++ b/Engine/Plugins/Experimental/Water/Source/Runtime/Public/WaterZoneActor.h
@@ -202,7 +202,7 @@ private:
 	bool bHalfPrecisionTexture = true;
 
 	/** Radius of the velocity blur in the finalize water info pass */
-	UPROPERTY(Category = Water, EditAnywhere, AdvancedDisplay)
+	UPROPERTY(Category = Water, EditAnywhere, AdvancedDisplay, meta = (ClampMin = 0, ClampMax = 16))
 	int32 VelocityBlurRadius = 1;
 
 	/** Higher number is higher priority. If Water Zones overlap and a water body does not have a manual water zone override, this priority will be used when automatically assigning the zone. */
