[Water]

Moved WaterInfoMaterial assignment out of constructor to prevent synchronous load block on the async load thread.

diff --git a/Engine/Plugins/Experimental/Water/Source/Editor/Private/WaterBodyActorFactory.cpp b/Engine/Plugins/Experimental/Water/Source/Editor/Private/WaterBodyActorFactory.cpp
--- a/Engine/Plugins/Experimental/Water/Source/Editor/Private/WaterBodyActorFactory.cpp
+++ b/Engine/Plugins/Experimental/Water/Source/Editor/Private/WaterBodyActorFactory.cpp
@@ -10,6 +10,7 @@
 #include "WaterBodyRiverComponent.h"
 #include "WaterSplineComponent.h"
 #include "WaterZoneActor.h"
+#include "WaterRuntimeSettings.h"
 
 #include UE_INLINE_GENERATED_CPP_BY_NAME(WaterBodyActorFactory)
 
@@ -45,6 +46,8 @@ void UWaterBodyActorFactory::PostSpawnActor(UObject* Asset, AActor* NewActor)
 		WaterBodyComponent->SetHLODMaterial(WaterBodyDefaults->GetWaterHLODMaterial());
 		WaterBodyComponent->SetUnderwaterPostProcessMaterial(WaterBodyDefaults->GetUnderwaterPostProcessMaterial());
 
+		WaterBodyComponent->SetWaterInfoMaterial(GetDefault<UWaterRuntimeSettings>()->GetDefaultWaterInfoMaterial());
+
 		UWaterSplineComponent* WaterSpline = WaterBodyComponent->GetWaterSpline();
 		if (ShouldOverrideWaterSplineDefaults(WaterSpline))
 		{
diff --git a/Engine/Plugins/Experimental/Water/Source/Runtime/Private/WaterBodyComponent.cpp b/Engine/Plugins/Experimental/Water/Source/Runtime/Private/WaterBodyComponent.cpp
--- a/Engine/Plugins/Experimental/Water/Source/Runtime/Private/WaterBodyComponent.cpp
+++ b/Engine/Plugins/Experimental/Water/Source/Runtime/Private/WaterBodyComponent.cpp
@@ -47,6 +47,7 @@
 #include "WaterBodyMeshBuilder.h"
 #include "MeshDescription.h"
 #include "Algo/RemoveIf.h"
+#include "UObject/FortniteReleaseBranchCustomObjectVersion.h"
 #endif
 
 #define LOCTEXT_NAMESPACE "Water"
@@ -114,8 +115,6 @@ UWaterBodyComponent::UWaterBodyComponent(const FObjectInitializer& ObjectInitial
 	bFillCollisionUnderneathForNavmesh = false;
 	bCanEverAffectNavigation = false;
 
-	WaterInfoMaterial = GetDefault<UWaterRuntimeSettings>()->GetDefaultWaterInfoMaterial();
-
 #if WITH_EDITORONLY_DATA
 	// Maintain the old default values for deprecated members so delta serialization is still correct when we deprecate them in PostLoad
 	bCanAffectNavigation_DEPRECATED = false;
@@ -284,6 +283,15 @@ void UWaterBodyComponent::SetWaterMaterial(UMaterialInterface* InMaterial)
 	}
 }
 
+void UWaterBodyComponent::SetWaterInfoMaterial(UMaterialInterface* InMaterial)
+{
+	if (WaterInfoMaterial != InMaterial)
+	{
+		WaterInfoMaterial = InMaterial;
+		UpdateMaterialInstances();
+	}
+}
+
 void UWaterBodyComponent::SetWaterStaticMeshMaterial(UMaterialInterface* InMaterial)
 {
 	if (WaterStaticMeshMaterial != InMaterial)
@@ -1774,6 +1782,17 @@ void UWaterBodyComponent::PostLoad()
 			}
 		}
 	}
+
+	if (GetLinkerCustomVersion(FWaterCustomVersion::GUID) < FWaterCustomVersion::DefaultWaterInfoMaterialAssignmentMoved)
+	{
+		// the defaults before this version no longer carry the default provided WaterInfoMaterial from their constructor
+		// so we can detect old data that hasn't been overriden by users by checking if the WaterInfoMaterial is null. 
+		// If a user _had_ overriden the material, it would have been set in ::Serialize.
+		if (WaterInfoMaterial == nullptr)
+		{
+			WaterInfoMaterial = GetDefault<UWaterRuntimeSettings>()->GetDefaultWaterInfoMaterial();
+		}
+	}
 #endif // WITH_EDITOR
 }
 
diff --git a/Engine/Plugins/Experimental/Water/Source/Runtime/Public/WaterBodyComponent.h b/Engine/Plugins/Experimental/Water/Source/Runtime/Public/WaterBodyComponent.h
--- a/Engine/Plugins/Experimental/Water/Source/Runtime/Public/WaterBodyComponent.h
+++ b/Engine/Plugins/Experimental/Water/Source/Runtime/Public/WaterBodyComponent.h
@@ -235,6 +235,10 @@ public:
 	UFUNCTION(BlueprintCallable, Category = Rendering)
 	void SetWaterMaterial(UMaterialInterface* InMaterial);
 
+	/** Sets the material used to draw the Water Info Texture for this water body */
+	UFUNCTION(BlueprintCallable, Category = Rendering)
+	void SetWaterInfoMaterial(UMaterialInterface* InMaterial);
+
 	/** Sets water static mesh material */
 	UFUNCTION(BlueprintCallable, Category = Rendering)
 	void SetWaterStaticMeshMaterial(UMaterialInterface* InMaterial);
diff --git a/Engine/Plugins/Experimental/Water/Source/Runtime/Public/WaterVersion.h b/Engine/Plugins/Experimental/Water/Source/Runtime/Public/WaterVersion.h
--- a/Engine/Plugins/Experimental/Water/Source/Runtime/Public/WaterVersion.h
+++ b/Engine/Plugins/Experimental/Water/Source/Runtime/Public/WaterVersion.h
@@ -26,6 +26,8 @@ struct WATER_API FWaterCustomVersion
 		FixupUnserializedGerstnerWaves,
 		// Move responsability of updating the Water MPC params to WaterMesh
 		MoveWaterMPCParamsToWaterMesh,
+		// Moved where the default water info material is assigned to prevent blocking async load thread
+		DefaultWaterInfoMaterialAssignmentMoved,
 
 		// -----<new versions can be added above this line>-------------------------------------------------
 		VersionPlusOne,
