[Water]

Fix a crash with waterinfo render method 2 when "use half precision texture" is not enabled.

diff --git a/Engine/Plugins/Experimental/Water/Source/Runtime/Private/WaterInfoRendering.cpp b/Engine/Plugins/Experimental/Water/Source/Runtime/Private/WaterInfoRendering.cpp
--- a/Engine/Plugins/Experimental/Water/Source/Runtime/Private/WaterInfoRendering.cpp
+++ b/Engine/Plugins/Experimental/Water/Source/Runtime/Private/WaterInfoRendering.cpp
@@ -585,7 +585,7 @@ public:
 
 	virtual void OnPreRender(FRDGBuilder& GraphBuilder) override
 	{
-		const FRDGTextureDesc TextureDesc = FRDGTextureDesc::Create2D(RenderTargetSize, PF_FloatRGBA, FClearValueBinding::Black, TexCreate_RenderTargetable | TexCreate_ShaderResource);
+		const FRDGTextureDesc TextureDesc = FRDGTextureDesc::Create2D(RenderTargetSize, WaterInfoRenderTarget->GetRenderTargetTexture()->GetDesc().Format, FClearValueBinding::Black, TexCreate_RenderTargetable | TexCreate_ShaderResource);
 		RenderTargetTexture = GraphBuilder.CreateTexture(TextureDesc, TEXT("WaterDilationTexture"));
 		AddClearRenderTargetPass(GraphBuilder, RenderTargetTexture, FLinearColor::Black, Views[0]->UnscaledViewRect);
 	}
