Refactored how we populate the world duplication seed to allow routing it to entities.

diff --git a/Engine/Plugins/Experimental/Water/Source/Runtime/Private/WaterBodyActor.cpp b/Engine/Plugins/Experimental/Water/Source/Runtime/Private/WaterBodyActor.cpp
--- a/Engine/Plugins/Experimental/Water/Source/Runtime/Private/WaterBodyActor.cpp
+++ b/Engine/Plugins/Experimental/Water/Source/Runtime/Private/WaterBodyActor.cpp
@@ -727,20 +727,23 @@ bool AWaterBody::IsComponentRelevantForNavigation(UActorComponent* Component) co
 }
 
 #if WITH_EDITOR
-void AWaterBody::PopulatePIEDuplicationSeed(AActor::FDuplicationSeedInterface& DuplicationSeed)
+void AWaterBody::PreDuplicateFromRoot(FObjectDuplicationParameters& DupParams)
 {
-	Super::PopulatePIEDuplicationSeed(DuplicationSeed);
+	Super::PreDuplicateFromRoot(DupParams);
 
-	// Avoid copying expensive UStaticMesh when entering pie, instead use the duplication seed to have the PIE instance share the same mesh pointer.
-	TArray<TObjectPtr<UWaterBodyMeshComponent>> MeshComponents = { WaterInfoMeshComponent, DilatedWaterInfoMeshComponent };
-	MeshComponents.Append(WaterBodyStaticMeshComponents);
-
-	for (const TObjectPtr<UWaterBodyMeshComponent>& MeshComponent : MeshComponents)
+	if (DupParams.DuplicateMode == EDuplicateMode::PIE)
 	{
-		if (IsValid(MeshComponent) && IsValid(MeshComponent->GetStaticMesh()))
+		// Avoid copying expensive UStaticMesh when entering pie, instead use the duplication seed to have the PIE instance share the same mesh pointer.
+		TArray<TObjectPtr<UWaterBodyMeshComponent>> MeshComponents = { WaterInfoMeshComponent, DilatedWaterInfoMeshComponent };
+		MeshComponents.Append(WaterBodyStaticMeshComponents);
+
+		for (const TObjectPtr<UWaterBodyMeshComponent>& MeshComponent : MeshComponents)
 		{
-			UStaticMesh* StaticMesh = MeshComponent->GetStaticMesh();
-			DuplicationSeed.AddEntry(StaticMesh, StaticMesh);
+			if (IsValid(MeshComponent) && IsValid(MeshComponent->GetStaticMesh()))
+			{
+				UStaticMesh* StaticMesh = MeshComponent->GetStaticMesh();
+				DupParams.DuplicationSeed.Emplace(StaticMesh, StaticMesh);
+			}
 		}
 	}
 }
diff --git a/Engine/Plugins/Experimental/Water/Source/Runtime/Public/WaterBodyActor.h b/Engine/Plugins/Experimental/Water/Source/Runtime/Public/WaterBodyActor.h
--- a/Engine/Plugins/Experimental/Water/Source/Runtime/Public/WaterBodyActor.h
+++ b/Engine/Plugins/Experimental/Water/Source/Runtime/Public/WaterBodyActor.h
@@ -83,7 +83,7 @@ public:
 
 	virtual void PostActorCreated() override;
 
-	virtual void PopulatePIEDuplicationSeed(AActor::FDuplicationSeedInterface& DuplicationSeed) override;
+	virtual void PreDuplicateFromRoot(FObjectDuplicationParameters& DupParams) override;
 
 	virtual void PostEditMove(bool bFinished) override;
 	virtual void PreEditChange(FProperty* PropertyThatWillChange) override;
diff --git a/Engine/Source/Runtime/Engine/Classes/GameFramework/Actor.h b/Engine/Source/Runtime/Engine/Classes/GameFramework/Actor.h
--- a/Engine/Source/Runtime/Engine/Classes/GameFramework/Actor.h
+++ b/Engine/Source/Runtime/Engine/Classes/GameFramework/Actor.h
@@ -2362,19 +2362,18 @@ public:
 	ENGINE_API virtual bool IsSelectedInEditor() const override;
 
 	/** DuplicationSeedInterface prevents exposing the entire DupSeed map to every actor and provides a method to do so instead. */
+	UE_DEPRECATED(5.6, "Use PreDuplicateFromRoot instead.")
 	struct FDuplicationSeedInterface
 	{
 	public:
-		FDuplicationSeedInterface(TMap<UObject*, UObject*>& InDuplicationSeed);
+		FDuplicationSeedInterface(TMap<UObject*, UObject*>& InDuplicationSeed) {}
 
 		/** Add a new entry to the duplication seed. */
-		ENGINE_API void AddEntry(UObject* Source, UObject* Dest);
-
-	private:
-		TMap<UObject*, UObject*>& DuplicationSeed;
+		ENGINE_API void AddEntry(UObject* Source, UObject* Dest) {}
 	};
 
 	/** Populate the duplication seed when duplicating the actor for PIE. This allows objects to be remapped to existing objects rather than duplicating. */
+	UE_DEPRECATED(5.6, "Use PreDuplicateFromRoot instead.")
 	virtual void PopulatePIEDuplicationSeed(FDuplicationSeedInterface& DupSeed) {}
 
 	/** Defines if preview should be shown when you select an actor, but none of its children */
@@ -2411,6 +2410,7 @@ public:
 #endif // WITH_EDITOR
 	//~ End UObject Interface
 
+	ENGINE_API virtual void PreDuplicateFromRoot(FObjectDuplicationParameters& DupParams) {}
 	ENGINE_API virtual void PreSaveFromRoot(FObjectPreSaveRootContext ObjectSaveContext);
 	ENGINE_API virtual void PostSaveFromRoot(FObjectPostSaveRootContext ObjectSaveContext);
 
diff --git a/Engine/Source/Runtime/Engine/Private/Actor.cpp b/Engine/Source/Runtime/Engine/Private/Actor.cpp
--- a/Engine/Source/Runtime/Engine/Private/Actor.cpp
+++ b/Engine/Source/Runtime/Engine/Private/Actor.cpp
@@ -115,16 +115,6 @@ FCriticalSection CSVActorClassNameToCountMapLock;
 #endif // (CSV_PROFILER_STATS && !UE_BUILD_SHIPPING)
 
 #if WITH_EDITOR
-AActor::FDuplicationSeedInterface::FDuplicationSeedInterface(TMap<UObject*, UObject*>& InDuplicationSeed)
-	: DuplicationSeed(InDuplicationSeed)
-{
-}
-
-void AActor::FDuplicationSeedInterface::AddEntry(UObject* Source, UObject* Destination)
-{	
-	DuplicationSeed.Emplace(Source, Destination);
-}
-
 namespace ActorUtils
 {
 	bool CanDeleteOrReplaceCommon(const AActor* InActor, FText& OutReason)
diff --git a/Engine/Source/Runtime/Engine/Private/Level.cpp b/Engine/Source/Runtime/Engine/Private/Level.cpp
--- a/Engine/Source/Runtime/Engine/Private/Level.cpp
+++ b/Engine/Source/Runtime/Engine/Private/Level.cpp
@@ -1497,19 +1497,15 @@ void ULevel::PreDuplicate(FObjectDuplicationParameters& DupParams)
 {
 	Super::PreDuplicate(DupParams);
 
-#if WITH_EDITOR
-	if (DupParams.DuplicateMode == EDuplicateMode::PIE)
+	for (AActor* Actor : Actors)
 	{
-		AActor::FDuplicationSeedInterface DuplicationSeedInterface(DupParams.DuplicationSeed);
-		for (AActor* Actor : Actors)
+		if (Actor)
 		{
-			if (Actor != nullptr)
-			{
-				Actor->PopulatePIEDuplicationSeed(DuplicationSeedInterface);
-			}
+			Actor->PreDuplicateFromRoot(DupParams);
 		}
 	}
 
+#if WITH_EDITOR
 	FExternalPackageHelper::DuplicateExternalPackages(this, DupParams, GetActorPackagingScheme());
 #endif
 }
diff --git a/Engine/Source/Runtime/Engine/Private/WorldPartition/StaticLightingData/MapBuildDataActor.cpp b/Engine/Source/Runtime/Engine/Private/WorldPartition/StaticLightingData/MapBuildDataActor.cpp
--- a/Engine/Source/Runtime/Engine/Private/WorldPartition/StaticLightingData/MapBuildDataActor.cpp
+++ b/Engine/Source/Runtime/Engine/Private/WorldPartition/StaticLightingData/MapBuildDataActor.cpp
@@ -51,14 +51,17 @@ void AMapBuildDataActor::BeginDestroy()
 }
 
 #if WITH_EDITOR
-void AMapBuildDataActor::PopulatePIEDuplicationSeed(FDuplicationSeedInterface& DupSeed)
+void AMapBuildDataActor::PreDuplicateFromRoot(FObjectDuplicationParameters& DupParams)
 {
-	Super::PopulatePIEDuplicationSeed(DupSeed);
+	Super::PreDuplicateFromRoot(DupParams);
 
-	if (BuildData)
+	if (DupParams.DuplicateMode == EDuplicateMode::PIE)
 	{
-		// Prevent MapBuildData duplication when PIEing, saves a little bit of memory. 
-		DupSeed.AddEntry(BuildData, BuildData);
+		if (BuildData)
+		{
+			// Prevent MapBuildData duplication when PIEing, saves a little bit of memory. 
+			DupParams.DuplicationSeed.Emplace(BuildData, BuildData);
+		}
 	}
 }
 #endif
diff --git a/Engine/Source/Runtime/Engine/Public/WorldPartition/StaticLightingData/MapBuildDataActor.h b/Engine/Source/Runtime/Engine/Public/WorldPartition/StaticLightingData/MapBuildDataActor.h
--- a/Engine/Source/Runtime/Engine/Public/WorldPartition/StaticLightingData/MapBuildDataActor.h
+++ b/Engine/Source/Runtime/Engine/Public/WorldPartition/StaticLightingData/MapBuildDataActor.h
@@ -40,7 +40,7 @@ protected:
 	virtual void PostUnregisterAllComponents() override;
 #if WITH_EDITOR
 	virtual TUniquePtr<class FWorldPartitionActorDesc> CreateClassActorDesc() const override;
-	virtual void PopulatePIEDuplicationSeed(FDuplicationSeedInterface& DupSeed) override;
+	virtual void PreDuplicateFromRoot(FObjectDuplicationParameters& DupParams) override;
 	virtual void GetStreamingBounds(FBox& OutRuntimeBounds, FBox& OutEditorBounds) const override;
 #endif
 	//~ End AActor Interface.
