[Water]

Prevent automatically spawning a water zone in a map when adding a water body which doesn't need it.

diff --git a/Engine/Plugins/Experimental/Water/Source/Editor/Private/WaterEditorModule.cpp b/Engine/Plugins/Experimental/Water/Source/Editor/Private/WaterEditorModule.cpp
--- a/Engine/Plugins/Experimental/Water/Source/Editor/Private/WaterEditorModule.cpp
+++ b/Engine/Plugins/Experimental/Water/Source/Editor/Private/WaterEditorModule.cpp
@@ -265,7 +265,7 @@ void FWaterEditorModule::OnLevelActorAddedToWorld(AActor* Actor)
 			}
 		}
 
-		if (WaterBodyActor != nullptr)
+		if (WaterBodyActor != nullptr && (WaterBodyActor->CanEverAffectWaterInfo() || WaterBodyActor->CanEverAffectWaterMesh()))
 		{
 			const bool bHasZoneActor = !!TActorIterator<AWaterZone>(ActorWorld);
 			// Setup the water zone actor for this water body : 
diff --git a/Engine/Plugins/Experimental/Water/Source/Runtime/Public/WaterBodyActor.h b/Engine/Plugins/Experimental/Water/Source/Runtime/Public/WaterBodyActor.h
--- a/Engine/Plugins/Experimental/Water/Source/Runtime/Public/WaterBodyActor.h
+++ b/Engine/Plugins/Experimental/Water/Source/Runtime/Public/WaterBodyActor.h
@@ -51,6 +51,7 @@ public:
 	virtual bool AffectsLandscape() const override { return WaterBodyComponent->AffectsLandscape(); }
 	virtual bool AffectsWaterMesh() const override { return WaterBodyComponent->AffectsWaterMesh(); }
 	virtual bool CanEverAffectWaterMesh() const override { return WaterBodyComponent->CanEverAffectWaterMesh(); }
+	virtual bool CanEverAffectWaterInfo() const override { return WaterBodyComponent->CanEverAffectWaterInfo(); }
 
 #if WITH_EDITOR
 	virtual const FWaterCurveSettings& GetWaterCurveSettings() const override { return WaterBodyComponent->GetWaterCurveSettings(); }
diff --git a/Engine/Plugins/Experimental/Water/Source/Runtime/Public/WaterBodyComponent.h b/Engine/Plugins/Experimental/Water/Source/Runtime/Public/WaterBodyComponent.h
--- a/Engine/Plugins/Experimental/Water/Source/Runtime/Public/WaterBodyComponent.h
+++ b/Engine/Plugins/Experimental/Water/Source/Runtime/Public/WaterBodyComponent.h
@@ -121,6 +121,7 @@ public:
 	virtual bool AffectsWaterMesh() const;
 	virtual bool AffectsWaterInfo() const;
 	virtual bool CanEverAffectWaterMesh() const { return true; }
+	virtual bool CanEverAffectWaterInfo() const { return true; }
 
 	void UpdateAll(const FOnWaterBodyChangedParams& InParams);
 
diff --git a/Engine/Plugins/Experimental/Water/Source/Runtime/Public/WaterBodyCustomComponent.h b/Engine/Plugins/Experimental/Water/Source/Runtime/Public/WaterBodyCustomComponent.h
--- a/Engine/Plugins/Experimental/Water/Source/Runtime/Public/WaterBodyCustomComponent.h
+++ b/Engine/Plugins/Experimental/Water/Source/Runtime/Public/WaterBodyCustomComponent.h
@@ -20,6 +20,7 @@ public:
 	virtual TArray<UPrimitiveComponent*> GetCollisionComponents(bool bInOnlyEnabledComponents = true) const override;
 	virtual TArray<UPrimitiveComponent*> GetStandardRenderableComponents() const override;
 	virtual bool CanEverAffectWaterMesh() const { return false; }
+	virtual bool CanEverAffectWaterInfo() const { return false; }
 
 	virtual FBoxSphereBounds CalcBounds(const FTransform& LocalToWorld) const override;
 
diff --git a/Engine/Plugins/Experimental/Water/Source/Runtime/Public/WaterBodyIslandActor.h b/Engine/Plugins/Experimental/Water/Source/Runtime/Public/WaterBodyIslandActor.h
--- a/Engine/Plugins/Experimental/Water/Source/Runtime/Public/WaterBodyIslandActor.h
+++ b/Engine/Plugins/Experimental/Water/Source/Runtime/Public/WaterBodyIslandActor.h
@@ -49,6 +49,7 @@ public:
 	virtual bool AffectsLandscape() const override { return true; }
 	virtual bool AffectsWaterMesh() const override { return false; }
 	virtual bool CanEverAffectWaterMesh() const override { return false; }
+	virtual bool CanEverAffectWaterInfo() const override { return false; }
 
 #if WITH_EDITOR
 	virtual const FWaterCurveSettings& GetWaterCurveSettings() const { return WaterCurveSettings; }
diff --git a/Engine/Plugins/Experimental/Water/Source/Runtime/Public/WaterBrushActorInterface.h b/Engine/Plugins/Experimental/Water/Source/Runtime/Public/WaterBrushActorInterface.h
--- a/Engine/Plugins/Experimental/Water/Source/Runtime/Public/WaterBrushActorInterface.h
+++ b/Engine/Plugins/Experimental/Water/Source/Runtime/Public/WaterBrushActorInterface.h
@@ -31,15 +31,20 @@ class WATER_API IWaterBrushActorInterface
 	virtual bool AffectsLandscape() const = 0;
 
 	/**
-	 * Returns true if this water actor is currently affecting (i.e. being rendered by) a AWaterZoneActor :
+	 * Returns true if this water actor is currently affecting (i.e. being rendered by) a AWaterZone actor :
 	 */
 	virtual bool AffectsWaterMesh() const = 0;
 
 	/**
-	 * Returns true if this water actor could potentially be affecting (i.e. being rendered by) a AWaterZoneActor :
+	 * Returns true if this water actor could potentially be affecting (i.e. being rendered by) a AWaterZone actor :
 	 */
 	virtual bool CanEverAffectWaterMesh() const = 0;
 
+	/**
+	 * Returns true if this water actor could potentially be affecting the water info texture of an AWaterZone actor:
+	 */
+	virtual bool CanEverAffectWaterInfo() const = 0;
+
 #if WITH_EDITOR
 	/** 
 	 * Returns the curve settings for this water actor
