[Water]

Fixed deprecation pitfall relying on postload ordering for bp base classes.

diff --git a/Engine/Plugins/Experimental/Water/Source/Runtime/Private/WaterBodyComponent.cpp b/Engine/Plugins/Experimental/Water/Source/Runtime/Private/WaterBodyComponent.cpp
--- a/Engine/Plugins/Experimental/Water/Source/Runtime/Private/WaterBodyComponent.cpp
+++ b/Engine/Plugins/Experimental/Water/Source/Runtime/Private/WaterBodyComponent.cpp
@@ -1677,6 +1677,22 @@ void UWaterBodyComponent::Serialize(FArchive& Ar)
 		FSoftObjectPath WaterSpriteTexture(GetWaterSpriteTextureName());
 		Ar << WaterSpriteTexture;
 	}
+
+	if (Ar.IsLoading())
+	{
+		// The WaterInfoMaterial fixup must happen in Serialize to ensure that blueprint hierarchy implements the fixup in the correct order
+		// such that parent BPs set the value before children are serialized.
+		if (GetLinkerCustomVersion(FWaterCustomVersion::GUID) < FWaterCustomVersion::DefaultWaterInfoMaterialAssignmentMoved)
+		{
+			// the defaults before this version no longer carry the default provided WaterInfoMaterial from their constructor
+			// so we can detect old data that hasn't been overriden by users by checking if the WaterInfoMaterial is null. 
+			// If a user _had_ overriden the material, it would have been set in UObject::Serialize.
+			if (WaterInfoMaterial == nullptr)
+			{
+				WaterInfoMaterial = GetDefault<UWaterRuntimeSettings>()->GetDefaultWaterInfoMaterial();
+			}
+		}
+	}
 #endif
 }
 
@@ -1782,17 +1798,6 @@ void UWaterBodyComponent::PostLoad()
 			}
 		}
 	}
-
-	if (GetLinkerCustomVersion(FWaterCustomVersion::GUID) < FWaterCustomVersion::DefaultWaterInfoMaterialAssignmentMoved)
-	{
-		// the defaults before this version no longer carry the default provided WaterInfoMaterial from their constructor
-		// so we can detect old data that hasn't been overriden by users by checking if the WaterInfoMaterial is null. 
-		// If a user _had_ overriden the material, it would have been set in ::Serialize.
-		if (WaterInfoMaterial == nullptr)
-		{
-			WaterInfoMaterial = GetDefault<UWaterRuntimeSettings>()->GetDefaultWaterInfoMaterial();
-		}
-	}
 #endif // WITH_EDITOR
 }
 
