Cleaned up a water view hack used to track PIE world destruction in order to force a WaterInfoTexture update

diff --git a/Engine/Plugins/Experimental/Water/Source/Runtime/Private/WaterViewExtension.cpp b/Engine/Plugins/Experimental/Water/Source/Runtime/Private/WaterViewExtension.cpp
--- a/Engine/Plugins/Experimental/Water/Source/Runtime/Private/WaterViewExtension.cpp
+++ b/Engine/Plugins/Experimental/Water/Source/Runtime/Private/WaterViewExtension.cpp
@@ -56,6 +56,19 @@ FWaterViewExtension::~FWaterViewExtension()
 {
 }
 
+// this delegates helps the editor view's WaterViewExtension detect scenarios when PIE ends, to make sure we schedule a ForceWaterZoneInfoRebuild
+void FWaterViewExtension::OnWorldDestroyed(UWorld* InWorld)
+{
+	// this is needed because when switching back from PIE to Editor, there isn't a proper initialization of
+	// the editor WaterViewExtension, since during PIE it is kept alive and not updated. Because of that when
+	// PIE ends and we start updating the editor view extension, it can contain out of date values which can cause
+	// water rendering artifacts until stepping out of bounds forces the first update.
+	if (InWorld->IsPlayInEditor())
+	{
+		ForceWaterZoneInfoRebuild = true;
+	}
+}
+
 void FWaterViewExtension::Initialize()
 {
 	// Register the view extension to the Gerstner Wave subsystem so we can rebuild the water gpu data when waves change.
@@ -66,6 +79,8 @@ void FWaterViewExtension::Initialize()
 
 	CurrentNumViews = 0;
 
+    FWorldDelegates::OnPreWorldFinishDestroy.AddRaw(this, &FWaterViewExtension::OnWorldDestroyed);
+
 	QuadTreeKeyLocationMap.Empty();
 }
 
@@ -84,6 +99,8 @@ void FWaterViewExtension::Deinitialize()
 
 	CurrentNumViews = 0;
 
+    FWorldDelegates::OnPreWorldFinishDestroy.RemoveAll(this);
+
 	QuadTreeKeyLocationMap.Empty();
 }
 
@@ -639,17 +656,9 @@ void FWaterViewExtension::SetupView(FSceneViewFamily& InViewFamily, FSceneView&
 			}
 		}
 
-		// this is needed because when switching back from PIE to Editor, there isn't a proper initialization of
-		// the editor WaterViewExtension, since during PIE it is kept alive and not updated. Because of that when
-		// PIE ends and we start updating the editor view extension, it can contain out of date values which can cause
-		// water rendering artifacts until stepping out of bounds forces the first update. To fix that we keep
-		// track if the last ViewExtension that was updated changed, so we can reset CurrentNumViews which forces an update on the next frame.
-		// TODO: look into any existing callbacks which would allow to achieve this in a cleaner way.
-		static FWaterViewExtension* LastWaterViewExtensionUpdated = nullptr;
-		if (LastWaterViewExtensionUpdated != this)
+		if (ForceWaterZoneInfoRebuild)
 		{
 			CurrentNumViews = 0;
-			LastWaterViewExtensionUpdated = this;
 		}
 	}
 	
diff --git a/Engine/Plugins/Experimental/Water/Source/Runtime/Public/WaterViewExtension.h b/Engine/Plugins/Experimental/Water/Source/Runtime/Public/WaterViewExtension.h
--- a/Engine/Plugins/Experimental/Water/Source/Runtime/Public/WaterViewExtension.h
+++ b/Engine/Plugins/Experimental/Water/Source/Runtime/Public/WaterViewExtension.h
@@ -59,6 +59,8 @@ public:
 	};
 	TWeakObjectPtrKeyMap<AWaterZone, FWaterZoneInfo> WaterZoneInfos;
 
+	bool ForceWaterZoneInfoRebuild = false;
+
 private:
 	int32 CurrentNumViews = 0;
 
@@ -108,6 +110,8 @@ private:
 	int32 GetViewIndex(const FSceneView& InView) const;
 
 	void DrawDebugInfo(FSceneView& InView, AWaterZone* WaterZone);
+
+	void OnWorldDestroyed(UWorld* InWorld);
 };
 
 struct FWaterMeshGPUWork
