WaterViewExtension: minor coding convention fixes

diff --git a/Engine/Plugins/Experimental/Water/Source/Runtime/Private/WaterViewExtension.cpp b/Engine/Plugins/Experimental/Water/Source/Runtime/Private/WaterViewExtension.cpp
--- a/Engine/Plugins/Experimental/Water/Source/Runtime/Private/WaterViewExtension.cpp
+++ b/Engine/Plugins/Experimental/Water/Source/Runtime/Private/WaterViewExtension.cpp
@@ -56,7 +56,7 @@ FWaterViewExtension::~FWaterViewExtension()
 {
 }
 
-// this delegates helps the editor view's WaterViewExtension detect scenarios when PIE ends, to make sure we schedule a ForceWaterZoneInfoRebuild
+// this delegates helps the editor view's WaterViewExtension detect scenarios when PIE ends, to make sure we schedule a bForceWaterZoneInfoRebuild
 void FWaterViewExtension::OnWorldDestroyed(UWorld* InWorld)
 {
 	// this is needed because when switching back from PIE to Editor, there isn't a proper initialization of
@@ -65,7 +65,7 @@ void FWaterViewExtension::OnWorldDestroyed(UWorld* InWorld)
 	// water rendering artifacts until stepping out of bounds forces the first update.
 	if (InWorld->IsPlayInEditor())
 	{
-		ForceWaterZoneInfoRebuild = true;
+		bForceWaterZoneInfoRebuild = true;
 	}
 }
 
@@ -656,7 +656,7 @@ void FWaterViewExtension::SetupView(FSceneViewFamily& InViewFamily, FSceneView&
 			}
 		}
 
-		if (ForceWaterZoneInfoRebuild)
+		if (bForceWaterZoneInfoRebuild)
 		{
 			CurrentNumViews = 0;
 		}
diff --git a/Engine/Plugins/Experimental/Water/Source/Runtime/Private/WaterZoneActor.cpp b/Engine/Plugins/Experimental/Water/Source/Runtime/Private/WaterZoneActor.cpp
--- a/Engine/Plugins/Experimental/Water/Source/Runtime/Private/WaterZoneActor.cpp
+++ b/Engine/Plugins/Experimental/Water/Source/Runtime/Private/WaterZoneActor.cpp
@@ -132,7 +132,7 @@ void AWaterZone::GetAllDynamicWaterInfoBounds(TArray<FBox>& OutBounds) const
 {
 	if (const FWaterViewExtension* WaterViewExtension = UWaterSubsystem::GetWaterViewExtension(GetWorld()))
 	{
-		const FWaterViewExtension::FWaterZoneInfo* WaterZoneInfo = WaterViewExtension->WaterZoneInfos.Find(this);
+		const FWaterViewExtension::FWaterZoneInfo* WaterZoneInfo = WaterViewExtension->GetWaterZoneInfos().Find(this);
 
 		if (WaterZoneInfo != nullptr)
 		{
@@ -161,7 +161,7 @@ void AWaterZone::GetAllDynamicWaterInfoCenters(TArray<FVector>& OutCenters) cons
 {
 	if (const FWaterViewExtension* WaterViewExtension = UWaterSubsystem::GetWaterViewExtension(GetWorld()))
 	{
-		const FWaterViewExtension::FWaterZoneInfo* WaterZoneInfo = WaterViewExtension->WaterZoneInfos.Find(this);
+		const FWaterViewExtension::FWaterZoneInfo* WaterZoneInfo = WaterViewExtension->GetWaterZoneInfos().Find(this);
 
 		if (WaterZoneInfo != nullptr)
 		{
diff --git a/Engine/Plugins/Experimental/Water/Source/Runtime/Public/WaterViewExtension.h b/Engine/Plugins/Experimental/Water/Source/Runtime/Public/WaterViewExtension.h
--- a/Engine/Plugins/Experimental/Water/Source/Runtime/Public/WaterViewExtension.h
+++ b/Engine/Plugins/Experimental/Water/Source/Runtime/Public/WaterViewExtension.h
@@ -57,13 +57,14 @@ public:
 		};
 		TArray<FWaterZoneViewInfo, TInlineAllocator<4>> ViewInfos;
 	};
-	TWeakObjectPtrKeyMap<AWaterZone, FWaterZoneInfo> WaterZoneInfos;
 
-	bool ForceWaterZoneInfoRebuild = false;
+	const TWeakObjectPtrKeyMap<AWaterZone, FWaterZoneInfo>& GetWaterZoneInfos() const  { return WaterZoneInfos; };
 
 private:
 	int32 CurrentNumViews = 0;
 
+	TWeakObjectPtrKeyMap<AWaterZone, FWaterZoneInfo> WaterZoneInfos;
+
 	struct FWaterGPUResources
 	{
 		FBufferRHIRef WaterBodyDataBuffer;
@@ -94,6 +95,8 @@ private:
 
 	bool bRebuildGPUData = true;
 
+	bool bForceWaterZoneInfoRebuild = false;
+
 	// store the locations of every active water mesh scene proxy quad tree based on the key
 	TMap<int32, FVector2D> QuadTreeKeyLocationMap;
 
