[Water]

Expose functions to set water spline metadata properties.

diff --git a/Engine/Plugins/Experimental/Water/Source/Runtime/Private/WaterBodyComponent.cpp b/Engine/Plugins/Experimental/Water/Source/Runtime/Private/WaterBodyComponent.cpp
--- a/Engine/Plugins/Experimental/Water/Source/Runtime/Private/WaterBodyComponent.cpp
+++ b/Engine/Plugins/Experimental/Water/Source/Runtime/Private/WaterBodyComponent.cpp
@@ -879,6 +879,24 @@ float UWaterBodyComponent::GetAudioIntensityAtSplineInputKey(float InKey) const
 	return WaterSplineMetadata ? WaterSplineMetadata->AudioIntensity.Eval(InKey, 0.f) : 0.0f;
 }
 
+void UWaterBodyComponent::SetWaterVelocityAtSplineInputKey(float InKey, float InVelocity)
+{
+	if (WaterSplineMetadata)
+	{
+		const int32 PointIndexForKey = WaterSplineMetadata->WaterVelocityScalar.GetPointIndexForInputValue(InKey);
+		WaterSplineMetadata->WaterVelocityScalar.Points[PointIndexForKey].OutVal = InVelocity;
+	}
+}
+
+void UWaterBodyComponent::SetAudioIntensityAtSplineInputKey(float InKey, float InAudioIntensity)
+{
+	if (WaterSplineMetadata)
+	{
+		const int32 PointIndexForKey = WaterSplineMetadata->AudioIntensity.GetPointIndexForInputValue(InKey);
+		WaterSplineMetadata->AudioIntensity.Points[PointIndexForKey].OutVal = InAudioIntensity;
+	}
+}
+
 void UWaterBodyComponent::OnRegister()
 {
 #if WITH_EDITORONLY_DATA
diff --git a/Engine/Plugins/Experimental/Water/Source/Runtime/Private/WaterBodyRiverComponent.cpp b/Engine/Plugins/Experimental/Water/Source/Runtime/Private/WaterBodyRiverComponent.cpp
--- a/Engine/Plugins/Experimental/Water/Source/Runtime/Private/WaterBodyRiverComponent.cpp
+++ b/Engine/Plugins/Experimental/Water/Source/Runtime/Private/WaterBodyRiverComponent.cpp
@@ -376,6 +376,34 @@ void UWaterBodyRiverComponent::SetLakeAndOceanTransitionMaterials(UMaterialInter
 	}
 }
 
+float UWaterBodyRiverComponent::GetRiverWidthAtSplineInputKey(float InKey) const
+{
+	return WaterSplineMetadata ? WaterSplineMetadata->RiverWidth.Eval(InKey, 0.f) : 0.0f;
+}
+
+float UWaterBodyRiverComponent::GetRiverDepthAtSplineInputKey(float InKey) const
+{
+	return WaterSplineMetadata ? WaterSplineMetadata->Depth.Eval(InKey, 0.f) : 0.0f;
+}
+
+void UWaterBodyRiverComponent::SetRiverWidthAtSplineInputKey(float InKey, float InWidth)
+{
+	if (WaterSplineMetadata)
+	{
+		const int32 PointIndexForKey = WaterSplineMetadata->RiverWidth.GetPointIndexForInputValue(InKey);
+		WaterSplineMetadata->RiverWidth.Points[PointIndexForKey].OutVal = InWidth;
+	}
+}
+
+void UWaterBodyRiverComponent::SetRiverDepthAtSplineInputKey(float InKey, float InDepth)
+{
+	if (WaterSplineMetadata)
+	{
+		const int32 PointIndexForKey = WaterSplineMetadata->Depth.GetPointIndexForInputValue(InKey);
+		WaterSplineMetadata->Depth.Points[PointIndexForKey].OutVal = InDepth;
+	}
+}
+
 void UWaterBodyRiverComponent::Reset()
 {
 	for (USplineMeshComponent* Comp : SplineMeshComponents)
diff --git a/Engine/Plugins/Experimental/Water/Source/Runtime/Public/WaterBodyComponent.h b/Engine/Plugins/Experimental/Water/Source/Runtime/Public/WaterBodyComponent.h
--- a/Engine/Plugins/Experimental/Water/Source/Runtime/Public/WaterBodyComponent.h
+++ b/Engine/Plugins/Experimental/Water/Source/Runtime/Public/WaterBodyComponent.h
@@ -309,9 +309,20 @@ public:
 	 */
 	UFUNCTION(BlueprintCallable, Category = WaterBody)
 	virtual float GetWaterVelocityAtSplineInputKey(float InKey) const;
+
+	UFUNCTION(BlueprintCallable, Category = WaterBody)
 	virtual FVector GetWaterVelocityVectorAtSplineInputKey(float InKey) const;
+
+	UFUNCTION(BlueprintCallable, Category = WaterBody)
 	virtual float GetAudioIntensityAtSplineInputKey(float InKey) const;
 
+
+	UFUNCTION(BlueprintCallable, Category = WaterBody)
+	void SetWaterVelocityAtSplineInputKey(float InKey, float InVelocity);
+
+	UFUNCTION(BlueprintCallable, Category = WaterBody)
+	void SetAudioIntensityAtSplineInputKey(float InKey, float InAudioIntensity);
+
 	/**
 	 * Gets the islands that influence this water body
 	 */
diff --git a/Engine/Plugins/Experimental/Water/Source/Runtime/Public/WaterBodyRiverComponent.h b/Engine/Plugins/Experimental/Water/Source/Runtime/Public/WaterBodyRiverComponent.h
--- a/Engine/Plugins/Experimental/Water/Source/Runtime/Public/WaterBodyRiverComponent.h
+++ b/Engine/Plugins/Experimental/Water/Source/Runtime/Public/WaterBodyRiverComponent.h
@@ -38,6 +38,19 @@ public:
 	UFUNCTION(BlueprintCallable, Category = Rendering)
 	void SetLakeAndOceanTransitionMaterials(UMaterialInterface* InLakeTransition, UMaterialInterface* InOceanTransition);
 
+
+	UFUNCTION(BlueprintCallable, Category = WaterBody)
+	float GetRiverWidthAtSplineInputKey(float InKey) const;
+
+	UFUNCTION(BlueprintCallable, Category = WaterBody)
+	float GetRiverDepthAtSplineInputKey(float InKey) const;
+
+	UFUNCTION(BlueprintCallable, Category = WaterBody)
+	void SetRiverWidthAtSplineInputKey(float InKey, float InWidth);
+
+	UFUNCTION(BlueprintCallable, Category = WaterBody)
+	void SetRiverDepthAtSplineInputKey(float InKey, float InDepth);
+
 protected:
 	/** UWaterBodyComponent Interface */
 	virtual void Reset() override;
