[Water]

Removed nonexistent parameter from edit condition on WaterNavAreaClass property causing it to never be editable.

diff --git a/Engine/Plugins/Experimental/Water/Source/Runtime/Public/WaterBodyComponent.h b/Engine/Plugins/Experimental/Water/Source/Runtime/Public/WaterBodyComponent.h
--- a/Engine/Plugins/Experimental/Water/Source/Runtime/Public/WaterBodyComponent.h
+++ b/Engine/Plugins/Experimental/Water/Source/Runtime/Public/WaterBodyComponent.h
@@ -686,7 +686,7 @@ protected:
 	FPostProcessSettings CurrentPostProcessSettings;
 
 	// The navigation area class that will be generated on nav mesh
-	UPROPERTY(EditAnywhere, BlueprintReadWrite, Category = Navigation, meta = (EditCondition = "bCanAffectNavigation && bGenerateCollisions"))
+	UPROPERTY(EditAnywhere, BlueprintReadWrite, Category = Navigation, meta = (EditCondition = "bCanAffectNavigation"))
 	TSubclassOf<UNavAreaBase> WaterNavAreaClass;
 
 	/** If the Water Material assigned to this component has Fixed Depth enabled, this is the depth that is passed. */
