[Water]

Deprecated bInvertShape heightmap setting which does nothing.

diff --git a/Engine/Plugins/Experimental/Water/Source/Editor/Private/WaterEditorSettings.cpp b/Engine/Plugins/Experimental/Water/Source/Editor/Private/WaterEditorSettings.cpp
--- a/Engine/Plugins/Experimental/Water/Source/Editor/Private/WaterEditorSettings.cpp
+++ b/Engine/Plugins/Experimental/Water/Source/Editor/Private/WaterEditorSettings.cpp
@@ -131,7 +131,6 @@ FWaterBodyIslandDefaults::FWaterBodyIslandDefaults()
 	BrushDefaults.CurveSettings.CurveRampWidth = 1024.0f;
 
 	BrushDefaults.HeightmapSettings.BlendMode = EWaterBrushBlendType::AlphaBlend;
-	BrushDefaults.HeightmapSettings.bInvertShape = false;
 	BrushDefaults.HeightmapSettings.FalloffSettings.FalloffMode = EWaterBrushFalloffMode::Angle;
 	BrushDefaults.HeightmapSettings.FalloffSettings.FalloffAngle = 45.0f;
 	BrushDefaults.HeightmapSettings.FalloffSettings.FalloffWidth = 1024.0f;
diff --git a/Engine/Plugins/Experimental/Water/Source/Runtime/Private/WaterBodyActor.cpp b/Engine/Plugins/Experimental/Water/Source/Runtime/Private/WaterBodyActor.cpp
--- a/Engine/Plugins/Experimental/Water/Source/Runtime/Private/WaterBodyActor.cpp
+++ b/Engine/Plugins/Experimental/Water/Source/Runtime/Private/WaterBodyActor.cpp
@@ -549,7 +549,6 @@ void AWaterBody::DeprecateData()
 		FMemory::Memcpy((void*)&WaterHeightmapSettings_DEPRECATED.MemberName, (void*)&TerrainCarvingSettings_DEPRECATED.MemberName, sizeof(WaterHeightmapSettings_DEPRECATED.MemberName));
 
 		COPY_SIMILAR_SETTINGS(BlendMode);
-		WaterHeightmapSettings_DEPRECATED.bInvertShape = TerrainCarvingSettings_DEPRECATED.bInvertShape;
 		COPY_SIMILAR_SETTINGS(FalloffSettings);
 		COPY_SIMILAR_SETTINGS(Effects.Blurring);
 		COPY_SIMILAR_SETTINGS(Effects.CurlNoise);
diff --git a/Engine/Plugins/Experimental/Water/Source/Runtime/Private/WaterBodyHeightmapSettings.cpp b/Engine/Plugins/Experimental/Water/Source/Runtime/Private/WaterBodyHeightmapSettings.cpp
--- a/Engine/Plugins/Experimental/Water/Source/Runtime/Private/WaterBodyHeightmapSettings.cpp
+++ b/Engine/Plugins/Experimental/Water/Source/Runtime/Private/WaterBodyHeightmapSettings.cpp
@@ -6,7 +6,6 @@
 
 FWaterBodyHeightmapSettings::FWaterBodyHeightmapSettings()
 	: BlendMode(EWaterBrushBlendType::AlphaBlend)
-	, bInvertShape(false)
 {
 }
 
diff --git a/Engine/Plugins/Experimental/Water/Source/Runtime/Public/WaterBodyHeightmapSettings.h b/Engine/Plugins/Experimental/Water/Source/Runtime/Public/WaterBodyHeightmapSettings.h
--- a/Engine/Plugins/Experimental/Water/Source/Runtime/Public/WaterBodyHeightmapSettings.h
+++ b/Engine/Plugins/Experimental/Water/Source/Runtime/Public/WaterBodyHeightmapSettings.h
@@ -31,15 +31,15 @@ struct FWaterBodyHeightmapSettings
 	UPROPERTY(EditAnywhere, BlueprintReadWrite, Category = TerrainCarvingSettings)
 	EWaterBrushBlendType BlendMode;
 
-	UPROPERTY(EditAnywhere, BlueprintReadWrite, AdvancedDisplay, Category = TerrainCarvingSettings)
-	bool bInvertShape;
-
 	UPROPERTY(EditAnywhere, BlueprintReadWrite, Category = TerrainCarvingSettings)
 	FWaterFalloffSettings FalloffSettings;
 
 	UPROPERTY(EditAnywhere, BlueprintReadWrite, Category = TerrainCarvingSettings)
 	FWaterBrushEffects Effects;
 
+	UPROPERTY()
+	bool bInvertShape_DEPRECATED;
+
 	UPROPERTY()
 	int32 Priority_DEPRECATED = 0;
 };
