[Water]

Fixed default water materials for custom water bodies so they work without a water zone present.

#virtualized

diff --git a/Engine/Plugins/Experimental/Water/Source/Editor/Private/WaterEditorSettings.cpp b/Engine/Plugins/Experimental/Water/Source/Editor/Private/WaterEditorSettings.cpp
--- a/Engine/Plugins/Experimental/Water/Source/Editor/Private/WaterEditorSettings.cpp
+++ b/Engine/Plugins/Experimental/Water/Source/Editor/Private/WaterEditorSettings.cpp
@@ -115,6 +115,7 @@ FWaterBodyCustomDefaults::FWaterBodyCustomDefaults()
 	, WaterMesh(FSoftObjectPath(TEXT("/Water/Meshes/S_WaterPlane_256.S_WaterPlane_256")))
 {
 	WaterMaterial = FSoftObjectPath(TEXT("/Water/Materials/WaterSurface/Water_Material_CustomMesh.Water_Material_CustomMesh"));
+	UnderwaterPostProcessMaterial = FSoftObjectPath(TEXT("/Water/Materials/PostProcessing/MI_UnderWater_PostProcess_Volume_Custom.MI_UnderWater_PostProcess_Volume_Custom"));
 }
 
 UStaticMesh* FWaterBodyCustomDefaults::GetWaterMesh() const
