[Water]

Moved public AActor overrides out of a protected block in AWaterBodyIsland.

diff --git a/Engine/Plugins/Experimental/Water/Source/Runtime/Public/WaterBodyIslandActor.h b/Engine/Plugins/Experimental/Water/Source/Runtime/Public/WaterBodyIslandActor.h
--- a/Engine/Plugins/Experimental/Water/Source/Runtime/Public/WaterBodyIslandActor.h
+++ b/Engine/Plugins/Experimental/Water/Source/Runtime/Public/WaterBodyIslandActor.h
@@ -59,6 +59,15 @@ public:
 	virtual void GetBrushRenderDependencies(TSet<UObject*>& OutDependencies) const override;
 #endif //WITH_EDITOR
 	//~ End IWaterBrushActorInterface interface
+	
+	//~ Begin AActor Interface
+#if WITH_EDITOR
+	virtual void PostEditMove(bool bFinished) override;
+	virtual void PostEditUndo() override;
+	virtual void PostEditImport() override;
+	virtual void PostEditChangeProperty(FPropertyChangedEvent& PropertyChangedEvent) override;
+#endif // WITH_EDITOR
+	// ~End AActor Interface
 
 	UFUNCTION(BlueprintCallable, Category=Water)
 	UWaterSplineComponent* GetWaterSpline() const { return SplineComp; }
@@ -92,10 +101,6 @@ protected:
 	virtual void Destroyed() override;
 
 #if WITH_EDITOR
-	virtual void PostEditMove(bool bFinished) override;
-	virtual void PostEditUndo() override;
-	virtual void PostEditChangeProperty(FPropertyChangedEvent& PropertyChangedEvent) override;
-	virtual void PostEditImport() override;
 	void UpdateAll(const FOnWaterBodyIslandChangedParams& InParams);
 	
 	void OnWaterSplineDataChanged(const FOnWaterSplineDataChangedParams& InParams);
