Fix for water being offset after PIE ends

diff --git a/Engine/Plugins/Experimental/Water/Source/Runtime/Private/WaterViewExtension.cpp b/Engine/Plugins/Experimental/Water/Source/Runtime/Private/WaterViewExtension.cpp
--- a/Engine/Plugins/Experimental/Water/Source/Runtime/Private/WaterViewExtension.cpp
+++ b/Engine/Plugins/Experimental/Water/Source/Runtime/Private/WaterViewExtension.cpp
@@ -56,7 +56,7 @@ FWaterViewExtension::~FWaterViewExtension()
 {
 }
 
-// this delegates helps the editor view's WaterViewExtension detect scenarios when PIE ends, to make sure we schedule a bForceWaterZoneInfoRebuild
+// this delegates helps the editor view's WaterViewExtension detect scenarios when PIE ends, to make sure we schedule a forced bounds update
 void FWaterViewExtension::OnWorldDestroyed(UWorld* InWorld)
 {
 	// this is needed because when switching back from PIE to Editor, there isn't a proper initialization of
@@ -65,7 +65,7 @@ void FWaterViewExtension::OnWorldDestroyed(UWorld* InWorld)
 	// water rendering artifacts until stepping out of bounds forces the first update.
 	if (InWorld->IsPlayInEditor())
 	{
-		bForceWaterZoneInfoRebuild = true;
+		bRequestForcedBoundsUpdate = true;
 	}
 }
 
@@ -579,6 +579,7 @@ void FWaterViewExtension::SetupView(FSceneViewFamily& InViewFamily, FSceneView&
 
 					bool bBoundsUpdateNeeded = WaterZone->IsLocalOnlyTessellationEnabled() && (!WaterZoneViewInfo.UpdateBounds.IsSet() || !WaterZoneViewInfo.UpdateBounds->IsInside(FVector2D(ViewLocation)));
 					bBoundsUpdateNeeded |= WaterZoneViewInfo.bIsDirty;
+					bBoundsUpdateNeeded |= bForceBoundsUpdate;
 
 					if (bBoundsUpdateNeeded)
 					{
@@ -664,9 +665,16 @@ void FWaterViewExtension::SetupView(FSceneViewFamily& InViewFamily, FSceneView&
 			}
 		}
 
-		if (bForceWaterZoneInfoRebuild)
+		if (bRequestForcedBoundsUpdate)
 		{
-			CurrentNumViews = 0;
+			bForceBoundsUpdate = true;
+			bRequestForcedBoundsUpdate = false;
+
+			UE_LOG(LogWater, Verbose, TEXT("Forced Bounds Update requested for view %d."), ViewPlayerIndex);
+		}
+		else
+		{
+			bForceBoundsUpdate = false;
 		}
 	}
 	
diff --git a/Engine/Plugins/Experimental/Water/Source/Runtime/Public/WaterViewExtension.h b/Engine/Plugins/Experimental/Water/Source/Runtime/Public/WaterViewExtension.h
--- a/Engine/Plugins/Experimental/Water/Source/Runtime/Public/WaterViewExtension.h
+++ b/Engine/Plugins/Experimental/Water/Source/Runtime/Public/WaterViewExtension.h
@@ -93,7 +93,8 @@ private:
 
 	bool bRebuildGPUData = true;
 
-	bool bForceWaterZoneInfoRebuild = false;
+	bool bRequestForcedBoundsUpdate = false;
+	bool bForceBoundsUpdate = false;
 
 	// store the locations of every active water mesh scene proxy quad tree based on the key
 	TMap<int32, FVector2D> QuadTreeKeyLocationMap;
