[Water]

Swapped FWaterBodyManager to holding weak references to water bodies and water zones. There seem to be niche cases where water body components can have mismatched Register/Unregister calls and dangling pointers can remain in the water body components list.

diff --git a/Engine/Plugins/Experimental/Water/Source/Runtime/Private/WaterBodyManager.cpp b/Engine/Plugins/Experimental/Water/Source/Runtime/Private/WaterBodyManager.cpp
--- a/Engine/Plugins/Experimental/Water/Source/Runtime/Private/WaterBodyManager.cpp
+++ b/Engine/Plugins/Experimental/Water/Source/Runtime/Private/WaterBodyManager.cpp
@@ -86,9 +86,9 @@ void FWaterBodyManager::RequestWaveDataRebuild()
 
 void FWaterBodyManager::ForEachWaterBodyComponent(TFunctionRef<bool(UWaterBodyComponent*)> Pred) const
 {
-	for (UWaterBodyComponent* WaterBodyComponent : WaterBodyComponents)
+	for (const TWeakObjectPtr<UWaterBodyComponent>& WaterBodyComponentPtr : WaterBodyComponents)
 	{
-		if (WaterBodyComponent)
+		if (UWaterBodyComponent* WaterBodyComponent = WaterBodyComponentPtr.Get())
 		{
 			if (!Pred(WaterBodyComponent))
 			{
@@ -108,9 +108,9 @@ void FWaterBodyManager::ForEachWaterBodyComponent(const UWorld* World, TFunction
 
 void FWaterBodyManager::ForEachWaterZone(TFunctionRef<bool(AWaterZone*)> Pred) const
 {
-	for (AWaterZone* WaterZone : WaterZones)
+	for (const TWeakObjectPtr<AWaterZone>& WaterZonePtr : WaterZones)
 	{
-		if (WaterZone)
+		if (AWaterZone* WaterZone = WaterZonePtr.Get())
 		{
 			if (!Pred(WaterZone))
 			{
diff --git a/Engine/Plugins/Experimental/Water/Source/Runtime/Public/WaterBodyManager.h b/Engine/Plugins/Experimental/Water/Source/Runtime/Public/WaterBodyManager.h
--- a/Engine/Plugins/Experimental/Water/Source/Runtime/Public/WaterBodyManager.h
+++ b/Engine/Plugins/Experimental/Water/Source/Runtime/Public/WaterBodyManager.h
@@ -4,9 +4,9 @@
 
 #include "CoreMinimal.h"
 #include "Templates/SharedPointer.h"
+#include "WaterBodyComponent.h"
+#include "WaterZoneActor.h"
 
-class UWaterBodyComponent;
-class AWaterZone;
 class FWaterViewExtension;
 
 DECLARE_MULTICAST_DELEGATE_OneParam(FWaterBodyEvent, UWaterBodyComponent*);
@@ -67,10 +67,10 @@ public:
 
 private:
 	/** List of components registered to this manager. */
-	TSparseArray<UWaterBodyComponent*> WaterBodyComponents;
+	TSparseArray<TWeakObjectPtr<UWaterBodyComponent>> WaterBodyComponents;
 
 	/** List of Water zones registered to this manager. */
-	TSparseArray<AWaterZone*> WaterZones;
+	TSparseArray<TWeakObjectPtr<AWaterZone>> WaterZones;
 
 	float GlobalMaxWaveHeight = 0.0f;
 
