Water: Removed ScenePrivate.h include in WaterInfoRendering.cpp.

diff --git a/Engine/Plugins/Experimental/Water/Source/Runtime/Private/WaterInfoRendering.cpp b/Engine/Plugins/Experimental/Water/Source/Runtime/Private/WaterInfoRendering.cpp
--- a/Engine/Plugins/Experimental/Water/Source/Runtime/Private/WaterInfoRendering.cpp
+++ b/Engine/Plugins/Experimental/Water/Source/Runtime/Private/WaterInfoRendering.cpp
@@ -26,7 +26,6 @@
 #include "RenderGraphBuilder.h"
 #include "SceneRenderTargetParameters.h"
 
-#include "ScenePrivate.h"
 #include "SceneRendering.h"
 #include "Rendering/CustomRenderPass.h"
 
