WaterViewExtension thread safety improvements: making sure that the WaterSceneProxy is only accessed in the renderthread, and moving all the SceneProxy related game thread logic from SetupView to BeginRenderViewFamily since SetupView was too early in the frame and might have caused passing dangling proxies to the render thread

diff --git a/Engine/Plugins/Experimental/Water/Source/Runtime/Private/WaterViewExtension.cpp b/Engine/Plugins/Experimental/Water/Source/Runtime/Private/WaterViewExtension.cpp
--- a/Engine/Plugins/Experimental/Water/Source/Runtime/Private/WaterViewExtension.cpp
+++ b/Engine/Plugins/Experimental/Water/Source/Runtime/Private/WaterViewExtension.cpp
@@ -400,7 +400,7 @@ void FWaterViewExtension::SetupViewFamily(FSceneViewFamily& InViewFamily)
 
 }
 
-void FWaterViewExtension::UpdateViewInfo(AWaterZone* WaterZone, FSceneView& InView)
+void FWaterViewExtension::UpdateViewInfo(AWaterZone* WaterZone, const FSceneView& InView)
 {
 	const int32 ViewPlayerIndex = GetViewIndex(InView);
 	check(ViewPlayerIndex != INDEX_NONE);
@@ -516,15 +516,16 @@ void FWaterViewExtension::SetupView(FSceneViewFamily& InViewFamily, FSceneView&
 	if (bWaterInfoTextureRebuildPending)
 	{
 		ViewPlayerIndices.Empty();
-		NonDataViewsQuadtreeKeys.Empty();
+
+		ENQUEUE_RENDER_COMMAND(WaterViewExtensionNonDataViewsQuadtreeKeysReset)(
+			[this](FRHICommandList& THICmdList)
+			{
+				NonDataViewsQuadtreeKeys.Empty();
+			});
 	}
 	else
 	{
-		// Don't dirty the water info texture when we're rendering from a scene capture. Due to the frame delay after marking the texture as dirty, scene captures wouldn't have the right texture anyways.
-		// #todo_water [roey]: Once we have no frame-delay for updating the texture and lesser performance impact, we can re-enable updates within scene captures.
-		bool bShouldHaveWaterZoneViewData = !InView.bIsSceneCapture && !InView.bIsSceneCaptureCube && !InView.bIsReflectionCapture && !InView.bIsPlanarReflection && !InView.bIsVirtualTexture;
-
-		if (bShouldHaveWaterZoneViewData)
+		if (ShouldHaveWaterZoneViewData(InView))
 		{
 			if (CurrentNumViews != NumViews)
 			{
@@ -592,8 +593,69 @@ void FWaterViewExtension::SetupView(FSceneViewFamily& InViewFamily, FSceneView&
 						{
 							RenderWaterInfoTexture(InViewFamily, InView, &WaterZoneInfo, Scene, WaterZoneViewInfo.Center);
 						}
+
+						WaterZoneViewInfo.bShouldUpdateQuadtree = true;
+					}
+
+					WaterZoneViewInfo.bIsDirty = false;
+				}
+			}
+		}
+
+		if (bRequestForcedBoundsUpdate)
+		{
+			bForceBoundsUpdate = true;
+			bRequestForcedBoundsUpdate = false;
+
+			UE_LOG(LogWater, Verbose, TEXT("Forced Bounds Update requested for view %d."), ViewPlayerIndex);
+		}
+		else
+		{
+			bForceBoundsUpdate = false;
+		}
+	}
+	
+
+	// The logic in UpdateGPUBuffers() used to be done in SetupViewFamily(). However, SetupView() (which is responsible for water info rendering) potentially modifies the WaterZone but is called after SetupViewFamily().
+	// This can lead to visual artifacts due to outdated data in the GPU buffers.
+	
+	UpdateGPUBuffers();
+}
+
+void FWaterViewExtension::BeginRenderViewFamily(FSceneViewFamily& InViewFamily)
+{
+	if (bWaterInfoTextureRebuildPending)
+	{
+		return;
+	}
+
+	const TWeakObjectPtr<UWorld> WorldPtr = GetWorld();
+
+	if (!ensureMsgf(WorldPtr.IsValid(), TEXT("FWaterViewExtension::BeginRenderViewFamily was called while it's owning world is not valid! Lifetime of the WaterViewExtension is tied to the world, this should be impossible!")))
+	{
+		return;
+	}
+
+	for (const FSceneView* View : InViewFamily.Views)
+	{
+		if (ShouldHaveWaterZoneViewData(*View))
+		{
+			for (AWaterZone* WaterZone : TActorRange<AWaterZone>(WorldPtr.Get()))
+			{
+				int32 ViewPlayerIndex = GetViewIndex(*View);
+				check(ViewPlayerIndex != INDEX_NONE);
+
+				if (WaterZone->HasActorRegisteredAllComponents())
+				{
+					FWaterZoneInfo& WaterZoneInfo = WaterZoneInfos.FindChecked(WaterZone);
+
+					if (WaterZone->WaterInfoTextureArray.Get() == nullptr)
+					{
+						continue;
 					}
 
+					FWaterZoneInfo::FWaterZoneViewInfo& WaterZoneViewInfo = WaterZoneInfo.ViewInfos[ViewPlayerIndex];
+
 					UWaterMeshComponent* WaterMeshComponent = WaterZone->GetWaterMeshComponent();
 					check(WaterMeshComponent != nullptr);
 
@@ -602,41 +664,36 @@ void FWaterViewExtension::SetupView(FSceneViewFamily& InViewFamily, FSceneView&
 					if (SceneProxy != nullptr)
 					{
 						// push a quadtree update
-						if (WaterZone->IsLocalOnlyTessellationEnabled() && (bBoundsUpdateNeeded || SceneProxy != WaterZoneViewInfo.OldSceneProxy))
+						if (WaterZone->IsLocalOnlyTessellationEnabled() && (WaterZoneViewInfo.bShouldUpdateQuadtree || SceneProxy != WaterZoneViewInfo.OldSceneProxy))
 						{
-							FScopeLock AddQuadtreeUpdate(&QuadtreeUpdateLock);
-
-							FQuadtreeUpdateInfo QuadtreeUpdate;
-							QuadtreeUpdate.WaterZone = WaterZone;
-							QuadtreeUpdate.Location = WaterZoneViewInfo.UpdateBounds->GetCenter();
-							// we use the actual PlayerIndex instead of ViewPlayerIndex, since the latter can change if more views are added while a quadtree already exists
-							QuadtreeUpdate.Key = InView.PlayerIndex;
-
-							QuadtreeUpdates.Add(QuadtreeUpdate);
-
-
-							if (!QuadTreeKeyLocationMap.Contains(QuadtreeUpdate.Key))
-							{
-								QuadTreeKeyLocationMap.Add(QuadtreeUpdate.Key, QuadtreeUpdate.Location);
-							}
-							else
-							{
-								QuadTreeKeyLocationMap[QuadtreeUpdate.Key] = QuadtreeUpdate.Location;
-							}
+							int32 QuadtreeKey = View->PlayerIndex;
+							FVector2D QuadtreeLocation = WaterZoneViewInfo.UpdateBounds->GetCenter();
+
+							ENQUEUE_RENDER_COMMAND(WaterViewExtensionQuadtreeUpdate)(
+								[SceneProxy, QuadtreeKey, QuadtreeLocation](FRHICommandList& THICmdList)
+								{
+									// TODO: We could add a CreateOrUpdateViewWaterQuadTree to WaterSceneProxy
+									// if the creation fails it means a quadtree for the current view already exists, so we just update it
+									if (!SceneProxy->CreateViewWaterQuadTree(QuadtreeKey, QuadtreeLocation))
+									{
+										bool bResult = SceneProxy->UpdateViewWaterQuadTree(QuadtreeKey, QuadtreeLocation);
+
+										check(bResult);
+									}
+								});
 
 							UE_LOG(LogWater, Verbose, TEXT("Water Zone (%s) queued a quadtree update for view %d updated."), *GetNameSafe(WaterZone), ViewPlayerIndex);
+							
+							WaterZoneViewInfo.bShouldUpdateQuadtree = false;
 						}
 					}
 
 					WaterZoneViewInfo.OldSceneProxy = SceneProxy;
-					WaterZoneViewInfo.bIsDirty = false;
 				}
 			}
 		}
 		else
 		{
-			FScopeLock AddNonDataViewsQuadtreeKeys(&QuadtreeUpdateLock);
-
 			// if !bShouldHaveWaterZoneViewData, we want to store the closest quadtree ID for this view (since it is not
 			// going to have its own quadtree associated with it), so we can assign the WaterInfoTexture index of the view assigned to that quadtree
 			for (AWaterZone* WaterZone : TActorRange<AWaterZone>(WorldPtr.Get()))
@@ -654,38 +711,32 @@ void FWaterViewExtension::SetupView(FSceneViewFamily& InViewFamily, FSceneView&
 						check(WaterMeshComponent != nullptr);
 
 						FWaterMeshSceneProxy* SceneProxy = static_cast<FWaterMeshSceneProxy*>(WaterMeshComponent->GetSceneProxy());
-						if (ensure(SceneProxy != nullptr))
+						if (SceneProxy != nullptr)
 						{
-							NonDataViewsQuadtreeKeys.Add(InView.State, SceneProxy->FindBestQuadTreeForViewLocation(&InView));
+							ENQUEUE_RENDER_COMMAND(WaterViewExtensionNonDataViewsQuadtreeKeysUpdate)(
+								[this, View, SceneProxy](FRHICommandList& THICmdList)
+								{
+									NonDataViewsQuadtreeKeys.Add(View->State, SceneProxy->FindBestQuadTreeForViewLocation(View));
+								});
 
-							UE_LOG(LogWater, Verbose, TEXT("Water Zone (%s) queued a search for the closest quadtree for a View (0x%p) which has no WaterInfo."), *GetNameSafe(WaterZone), &InView);
+							UE_LOG(LogWater, Verbose, TEXT("Water Zone (%s) queued a search for the closest quadtree for a View (0x%p) which has no WaterInfo."), *GetNameSafe(WaterZone), View);
 						}
 					}
 				}
 			}
 		}
 
-		if (bRequestForcedBoundsUpdate)
-		{
-			bForceBoundsUpdate = true;
-			bRequestForcedBoundsUpdate = false;
-
-			UE_LOG(LogWater, Verbose, TEXT("Forced Bounds Update requested for view %d."), ViewPlayerIndex);
-		}
-		else
-		{
-			bForceBoundsUpdate = false;
-		}
 	}
-	
+}
 
-	// The logic in UpdateGPUBuffers() used to be done in SetupViewFamily(). However, SetupView() (which is responsible for water info rendering) potentially modifies the WaterZone but is called after SetupViewFamily().
-	// This can lead to visual artifacts due to outdated data in the GPU buffers.
-	
-	UpdateGPUBuffers();
+bool FWaterViewExtension::ShouldHaveWaterZoneViewData(const FSceneView& InView) const
+{
+	// Don't dirty the water info texture when we're rendering from a scene capture. Due to the frame delay after marking the texture as dirty, scene captures wouldn't have the right texture anyways.
+	// #todo_water [roey]: Once we have no frame-delay for updating the texture and lesser performance impact, we can re-enable updates within scene captures.
+	return !InView.bIsSceneCapture && !InView.bIsSceneCaptureCube && !InView.bIsReflectionCapture && !InView.bIsPlanarReflection && !InView.bIsVirtualTexture;
 }
 
-void FWaterViewExtension::DrawDebugInfo(FSceneView& InView, AWaterZone* WaterZone)
+void FWaterViewExtension::DrawDebugInfo(const FSceneView& InView, AWaterZone* WaterZone)
 {
 	if (GEngine)
 	{
@@ -720,33 +771,6 @@ void FWaterViewExtension::DrawDebugInfo(FSceneView& InView, AWaterZone* WaterZon
 
 void FWaterViewExtension::PreRenderViewFamily_RenderThread(FRDGBuilder& GraphBuilder, FSceneViewFamily& InViewFamily)
 {
-	// consume quadtree updates
-	{
-		FScopeLock ConsumeQuadtreeUpdates(&QuadtreeUpdateLock);
-
-		for (const FQuadtreeUpdateInfo& QuadtreeUpdate : QuadtreeUpdates)
-		{
-			UWaterMeshComponent* WaterMeshComponent = QuadtreeUpdate.WaterZone->GetWaterMeshComponent();
-			if (ensure(WaterMeshComponent != nullptr))
-			{
-				FWaterMeshSceneProxy* SceneProxy = static_cast<FWaterMeshSceneProxy*>(WaterMeshComponent->GetSceneProxy());
-
-				if (SceneProxy != nullptr)
-				{
-					// TODO: We could add a CreateOrUpdateViewWaterQuadTree to WaterSceneProxy
-					// if the creation fails it means a quadtree for the current view already exists, so we just update it
-					if (!SceneProxy->CreateViewWaterQuadTree(QuadtreeUpdate.Key, QuadtreeUpdate.Location))
-					{
-						bool bResult = SceneProxy->UpdateViewWaterQuadTree(QuadtreeUpdate.Key, QuadtreeUpdate.Location);
-
-						check(bResult);
-					}
-				}
-			}
-		}
-
-		QuadtreeUpdates.Empty();
-	}
 	
 }
 
@@ -760,8 +784,6 @@ void FWaterViewExtension::PreRenderView_RenderThread(FRDGBuilder& GraphBuilder,
 	}
 
 	{
-		FScopeLock ReadNonDataViewsQuadtreeKeys(&QuadtreeUpdateLock);
-
 		int32 WaterInfoTextureIndex = GetViewIndex(InView);
 
 		// check if this view is one of the queues without its own WaterZoneViewData, if it is assign the 
diff --git a/Engine/Plugins/Experimental/Water/Source/Runtime/Public/WaterViewExtension.h b/Engine/Plugins/Experimental/Water/Source/Runtime/Public/WaterViewExtension.h
--- a/Engine/Plugins/Experimental/Water/Source/Runtime/Public/WaterViewExtension.h
+++ b/Engine/Plugins/Experimental/Water/Source/Runtime/Public/WaterViewExtension.h
@@ -22,6 +22,7 @@ public:
 	// FSceneViewExtensionBase implementation : 
 	virtual void SetupView(FSceneViewFamily& InViewFamily, FSceneView& InView) override;
 	virtual void SetupViewFamily(FSceneViewFamily& InViewFamily) override;
+	virtual void BeginRenderViewFamily(FSceneViewFamily& InViewFamily) override;
 	virtual void PreRenderViewFamily_RenderThread(FRDGBuilder& GraphBuilder, FSceneViewFamily& InViewFamily) override;
 	virtual void PreRenderView_RenderThread(FRDGBuilder& GraphBuilder, FSceneView& InView) override;
 	virtual void PreRenderBasePass_RenderThread(FRDGBuilder& GraphBuilder, bool bDepthBufferIsPopulated) override;
@@ -51,7 +52,11 @@ public:
 			TOptional<FBox2D> UpdateBounds = FBox2D(ForceInit);
 			FVector Center = FVector(ForceInit);
 			FWaterMeshSceneProxy* OldSceneProxy = nullptr;
+			// bIsDirty is used to force water info texture update for this WaterZone
 			bool bIsDirty = true;
+			// bShouldUpdateQuadtree is set to true when the water info texture is update, 
+			// to make sure the quadtree for this WaterZone is also updated
+			bool bShouldUpdateQuadtree = false;
 		};
 		TArray<FWaterZoneViewInfo, TInlineAllocator<4>> ViewInfos;
 	};
@@ -76,17 +81,6 @@ private:
 
 	TArray<int32, TInlineAllocator<4>> ViewPlayerIndices;
 
-	struct FQuadtreeUpdateInfo
-	{
-		AWaterZone* WaterZone;
-		FVector2D Location;
-		int32 Key;
-	};
-
-	TArray<FQuadtreeUpdateInfo, TInlineAllocator<4>> QuadtreeUpdates;
-
-	FCriticalSection QuadtreeUpdateLock;
-
 	TMap<FSceneViewStateInterface*, int32> NonDataViewsQuadtreeKeys;
 
 	bool bWaterInfoTextureRebuildPending = true;
@@ -101,7 +95,7 @@ private:
 
 	void UpdateGPUBuffers();
 
-	void UpdateViewInfo(AWaterZone* WaterZone, FSceneView& InView);
+	void UpdateViewInfo(AWaterZone* WaterZone, const FSceneView& InView);
 	
 	void RenderWaterInfoTexture(FSceneViewFamily& InViewFamily, FSceneView& InView, const FWaterZoneInfo* WaterZoneInfo, FSceneInterface* Scene, const FVector& ZoneCenter);
 
@@ -111,7 +105,9 @@ private:
 	int32 GetViewIndex(int32 PlayerIndex) const;
 	int32 GetViewIndex(const FSceneView& InView) const;
 
-	void DrawDebugInfo(FSceneView& InView, AWaterZone* WaterZone);
+	bool ShouldHaveWaterZoneViewData(const FSceneView& InView) const;
+
+	void DrawDebugInfo(const FSceneView& InView, AWaterZone* WaterZone);
 
 	void OnWorldDestroyed(UWorld* InWorld);
 };
