[Water]

Deprecated the custom water physical material property to instead prefer using the builtin phys material override on the UPrimitiveComponent.

diff --git a/Engine/Plugins/Experimental/Water/Source/Runtime/Private/WaterBodyActor.cpp b/Engine/Plugins/Experimental/Water/Source/Runtime/Private/WaterBodyActor.cpp
--- a/Engine/Plugins/Experimental/Water/Source/Runtime/Private/WaterBodyActor.cpp
+++ b/Engine/Plugins/Experimental/Water/Source/Runtime/Private/WaterBodyActor.cpp
@@ -581,7 +581,7 @@ PRAGMA_ENABLE_DEPRECATION_WARNINGS
 		check(SplineComp);
 
 		WaterBodyComponent->Mobility = SplineComp->Mobility;
-		WaterBodyComponent->PhysicalMaterial = PhysicalMaterial_DEPRECATED;
+		WaterBodyComponent->PhysicalMaterial_DEPRECATED = PhysicalMaterial_DEPRECATED;
 		WaterBodyComponent->TargetWaveMaskDepth = TargetWaveMaskDepth_DEPRECATED;
 		WaterBodyComponent->MaxWaveHeightOffset = MaxWaveHeightOffset_DEPRECATED;
 		WaterBodyComponent->bFillCollisionUnderWaterBodiesForNavmesh_DEPRECATED = bFillCollisionUnderWaterBodiesForNavmesh_DEPRECATED;
diff --git a/Engine/Plugins/Experimental/Water/Source/Runtime/Private/WaterBodyComponent.cpp b/Engine/Plugins/Experimental/Water/Source/Runtime/Private/WaterBodyComponent.cpp
--- a/Engine/Plugins/Experimental/Water/Source/Runtime/Private/WaterBodyComponent.cpp
+++ b/Engine/Plugins/Experimental/Water/Source/Runtime/Private/WaterBodyComponent.cpp
@@ -180,6 +180,11 @@ FBox UWaterBodyComponent::GetCollisionComponentBounds() const
 	return Box;
 }
 
+UPhysicalMaterial* UWaterBodyComponent::GetPhysicalMaterial() const
+{
+	return BodyInstance.GetPhysMaterialOverride();
+}
+
 void UWaterBodyComponent::SetWaterMeshOverride(UStaticMesh* InMesh)
 {
 	WaterMeshOverride = InMesh;
@@ -1842,6 +1847,11 @@ void UWaterBodyComponent::DeprecateData()
 	{
 		WaterStaticMeshMaterial = WaterLODMaterial_DEPRECATED;
 	}
+
+	if (GetLinkerCustomVersion(FFortniteMainBranchObjectVersion::GUID) < FFortniteMainBranchObjectVersion::WaterBodyPhysicalMaterialPropertyRemoval)
+	{
+		SetPhysMaterialOverride(PhysicalMaterial_DEPRECATED);
+	}
 #endif // WITH_EDITORONLY_DATA
 }
 
@@ -2000,7 +2010,7 @@ void UWaterBodyComponent::CopySharedCollisionSettingsToComponent(UPrimitiveCompo
 	InComponent->SetCollisionResponseToChannels(BodyInstance.GetResponseToChannels());
 	InComponent->SetCollisionProfileName(GetCollisionProfileName(), /*bUpdateOverlaps=*/ true);
 	InComponent->SetGenerateOverlapEvents(GetGenerateOverlapEvents());
-	InComponent->SetPhysMaterialOverride(PhysicalMaterial);
+	InComponent->SetPhysMaterialOverride(BodyInstance.GetPhysMaterialOverride());
 }
 
 void UWaterBodyComponent::CopySharedNavigationSettingsToComponent(UPrimitiveComponent* InComponent)
diff --git a/Engine/Plugins/Experimental/Water/Source/Runtime/Public/WaterBodyComponent.h b/Engine/Plugins/Experimental/Water/Source/Runtime/Public/WaterBodyComponent.h
--- a/Engine/Plugins/Experimental/Water/Source/Runtime/Public/WaterBodyComponent.h
+++ b/Engine/Plugins/Experimental/Water/Source/Runtime/Public/WaterBodyComponent.h
@@ -171,6 +171,9 @@ public:
 	/** Returns collision half-extents */
 	virtual FVector GetCollisionExtents() const { return FVector::ZeroVector; }
 
+	/** Returns the physical material of this water body. */
+	UPhysicalMaterial* GetPhysicalMaterial() const;
+
 	/** Sets an additional water height (For internal use. Please use AWaterBodyOcean instead.) */
 	virtual void SetHeightOffset(float InHeightOffset) { check(false); }
 
@@ -576,9 +579,6 @@ public:
 	static const FName GlobalOceanHeightName;
 	static const FName MaxFlowVelocityParamName;
 
-	UPROPERTY(EditDefaultsOnly, Category = Collision, meta = (EditCondition = "bGenerateCollisions"))
-	TObjectPtr<UPhysicalMaterial> PhysicalMaterial;
-
 	/** Water depth at which waves start being attenuated. */
 	UPROPERTY(EditAnywhere, BlueprintReadWrite, Category = Wave, DisplayName = "Wave Attenuation Water Depth", meta = (UIMin = 0, ClampMin = 0, UIMax = 10000.0))
 	float TargetWaveMaskDepth;
@@ -587,9 +587,9 @@ public:
 	UPROPERTY(EditAnywhere, Category = Wave)
 	float MaxWaveHeightOffset = 0.0f;
 
-	/** Post process settings to apply when the camera goes underwater (only available when bGenerateCollisions is true because collisions are needed to detect if it's under water).
+	/** Post process settings to apply when the camera goes underwater (only available when collisions are enabled as they are needed to detect if the camera is under water).
 	Note: Underwater post process material is setup using UnderwaterPostProcessMaterial. */
-	UPROPERTY(Category = Rendering, EditAnywhere, BlueprintReadWrite, AdvancedDisplay, meta = (EditCondition = "bGenerateCollisions && UnderwaterPostProcessMaterial != nullptr", DisplayAfter = "UnderwaterPostProcessMaterial"))
+	UPROPERTY(Category = Rendering, EditAnywhere, BlueprintReadWrite, AdvancedDisplay, meta = (EditCondition = "UnderwaterPostProcessMaterial != nullptr", DisplayAfter = "UnderwaterPostProcessMaterial"))
 	FUnderwaterPostProcessSettings UnderwaterPostProcessSettings;
 
 	UPROPERTY(Category = Terrain, EditAnywhere, BlueprintReadWrite, meta = (EditCondition = "bAffectsLandscape"))
@@ -710,6 +710,9 @@ protected:
 	UPROPERTY()
 	TObjectPtr<UMaterialInterface> WaterLODMaterial_DEPRECATED;
 
+	UPROPERTY()
+	TObjectPtr<UPhysicalMaterial> PhysicalMaterial_DEPRECATED;
+
 	UPROPERTY()
 	bool bFillCollisionUnderWaterBodiesForNavmesh_DEPRECATED;
 
diff --git a/Engine/Source/Runtime/Core/Public/UObject/FortniteMainBranchObjectVersions.inl b/Engine/Source/Runtime/Core/Public/UObject/FortniteMainBranchObjectVersions.inl
--- a/Engine/Source/Runtime/Core/Public/UObject/FortniteMainBranchObjectVersions.inl
+++ b/Engine/Source/Runtime/Core/Public/UObject/FortniteMainBranchObjectVersions.inl
@@ -592,5 +592,8 @@ FORTNITE_MAIN_VERSION(SoftObjectPathUtf8SubPaths, 192)
 // FSoftObjectPath::SubPathString could be saved with trailing NULs and need truncating
 FORTNITE_MAIN_VERSION(SoftObjectPathTrailingNULsMaintained, 193)
 
+// Water body components no longer need to maintain their own PhysicalMaterial property since they are primitive components. After this version, leverage that one instead.
+FORTNITE_MAIN_VERSION(WaterBodyPhysicalMaterialPropertyRemoval, 194)
+
 // -----<new versions can be added above this line>-------------------------------------------------
 #undef FORTNITE_MAIN_VERSION
\ No newline at end of file
