WaterViewExtension ASAN fixes

diff --git a/Engine/Plugins/Experimental/Water/Source/Runtime/Private/WaterMeshSceneProxy.cpp b/Engine/Plugins/Experimental/Water/Source/Runtime/Private/WaterMeshSceneProxy.cpp
--- a/Engine/Plugins/Experimental/Water/Source/Runtime/Private/WaterMeshSceneProxy.cpp
+++ b/Engine/Plugins/Experimental/Water/Source/Runtime/Private/WaterMeshSceneProxy.cpp
@@ -853,13 +853,13 @@ int32 FWaterMeshSceneProxy::FindBestQuadTreeForView(const FSceneView* View) cons
 	// Otherwise just use the closest quadtree.
 	else
 	{
-		return FindBestQuadTreeForViewLocation(View);
+		const FVector2D ViewPosition2D = FVector2D(View->ViewLocation);
+		return FindBestQuadTreeForViewLocation(ViewPosition2D);
 	}
 }
 
-int32 FWaterMeshSceneProxy::FindBestQuadTreeForViewLocation(const FSceneView* View) const
+int32 FWaterMeshSceneProxy::FindBestQuadTreeForViewLocation(const FVector2D& ViewPosition2D) const
 {
-	const FVector2D ViewPosition2D = FVector2D(View->ViewLocation);
 	double ClosestDistSquared = DBL_MAX;
 	int32 ClosestQuadTreeKey = INDEX_NONE;
 	for (const auto& Pair : ViewQuadTrees)
diff --git a/Engine/Plugins/Experimental/Water/Source/Runtime/Private/WaterViewExtension.cpp b/Engine/Plugins/Experimental/Water/Source/Runtime/Private/WaterViewExtension.cpp
--- a/Engine/Plugins/Experimental/Water/Source/Runtime/Private/WaterViewExtension.cpp
+++ b/Engine/Plugins/Experimental/Water/Source/Runtime/Private/WaterViewExtension.cpp
@@ -638,6 +638,9 @@ void FWaterViewExtension::BeginRenderViewFamily(FSceneViewFamily& InViewFamily)
 
 	for (const FSceneView* View : InViewFamily.Views)
 	{
+		// Waning: Do not capture View in ENQUEUE_RENDER_COMMAND lambdas, since the RT's view hasn't been created yet,
+		// so we would pass the GT view, which then would end up being a dangling pointer when accessed by the RT.
+
 		if (ShouldHaveWaterZoneViewData(*View))
 		{
 			for (AWaterZone* WaterZone : TActorRange<AWaterZone>(WorldPtr.Get()))
@@ -713,10 +716,13 @@ void FWaterViewExtension::BeginRenderViewFamily(FSceneViewFamily& InViewFamily)
 						FWaterMeshSceneProxy* SceneProxy = static_cast<FWaterMeshSceneProxy*>(WaterMeshComponent->GetSceneProxy());
 						if (SceneProxy != nullptr)
 						{
+							const FVector2D ViewPosition2D = FVector2D(View->ViewLocation);
+							FSceneViewStateInterface* ViewState = View->State;
+
 							ENQUEUE_RENDER_COMMAND(WaterViewExtensionNonDataViewsQuadtreeKeysUpdate)(
-								[this, View, SceneProxy](FRHICommandList& THICmdList)
+								[this, SceneProxy, ViewPosition2D, ViewState](FRHICommandList& THICmdList)
 								{
-									NonDataViewsQuadtreeKeys.Add(View->State, SceneProxy->FindBestQuadTreeForViewLocation(View));
+									NonDataViewsQuadtreeKeys.Add(ViewState, SceneProxy->FindBestQuadTreeForViewLocation(ViewPosition2D));
 								});
 
 							UE_LOG(LogWater, Verbose, TEXT("Water Zone (%s) queued a search for the closest quadtree for a View (0x%p) which has no WaterInfo."), *GetNameSafe(WaterZone), View);
diff --git a/Engine/Plugins/Experimental/Water/Source/Runtime/Public/WaterMeshSceneProxy.h b/Engine/Plugins/Experimental/Water/Source/Runtime/Public/WaterMeshSceneProxy.h
--- a/Engine/Plugins/Experimental/Water/Source/Runtime/Public/WaterMeshSceneProxy.h
+++ b/Engine/Plugins/Experimental/Water/Source/Runtime/Public/WaterMeshSceneProxy.h
@@ -164,7 +164,7 @@ public:
 	// Destroys the quadtree associated with Key.
 	void DestroyViewWaterQuadTree(int32 Key);
 
-	int32 FindBestQuadTreeForViewLocation(const FSceneView* View) const;
+	int32 FindBestQuadTreeForViewLocation(const FVector2D& ViewPosition2D) const;
 
 private:
 
