Make NiagaraDataInterfaceWater take a weak ptr to the source water body component - to prevent keeping a garbage reference if the water body is destroyed ( eg. if the particle component is added to a pool, and then we stream out the water body ).

diff --git a/Engine/Plugins/Experimental/Water/Source/Runtime/Private/NiagaraDataInterfaceWater.cpp b/Engine/Plugins/Experimental/Water/Source/Runtime/Private/NiagaraDataInterfaceWater.cpp
--- a/Engine/Plugins/Experimental/Water/Source/Runtime/Private/NiagaraDataInterfaceWater.cpp
+++ b/Engine/Plugins/Experimental/Water/Source/Runtime/Private/NiagaraDataInterfaceWater.cpp
@@ -351,20 +351,21 @@ bool UNiagaraDataInterfaceWater::PerInstanceTick(void* PerInstanceData, FNiagara
 	{
 		// If the search for closest was enabled, perform the search
 		// Note: we do this here rather than in Init as the system might be auto activate and the user parameter not set until after the spawn
-		if (SourceBodyComponent == nullptr && InstData->bFindClosestBody)
+		if (SourceBodyComponent.IsNull() && InstData->bFindClosestBody)
 		{
 			const FVector QueryLocation = SystemInstance->GetWorldTransform().GetTranslation();
 			InstData->WaterBodyComponent = NDIWaterPrivate::FindClosestWaterBody(SystemInstance->GetWorld(), QueryLocation);
 		}
 		else
 		{
-			if (AActor* SourceActor = Cast<AActor>(SourceBodyComponent))
+			UObject* RawSourceBodyComponent = SourceBodyComponent.Get();
+			if (AActor* SourceActor = Cast<AActor>(RawSourceBodyComponent))
 			{
 				InstData->WaterBodyComponent = SourceActor->GetComponentByClass<UWaterBodyComponent>();
 			}
 			else
 			{
-				InstData->WaterBodyComponent = Cast<UWaterBodyComponent>(SourceBodyComponent);
+				InstData->WaterBodyComponent = Cast<UWaterBodyComponent>(RawSourceBodyComponent);
 			}
 		}
 		InstData->bFindClosestBody = false;
diff --git a/Engine/Plugins/Experimental/Water/Source/Runtime/Public/NiagaraDataInterfaceWater.h b/Engine/Plugins/Experimental/Water/Source/Runtime/Public/NiagaraDataInterfaceWater.h
--- a/Engine/Plugins/Experimental/Water/Source/Runtime/Public/NiagaraDataInterfaceWater.h
+++ b/Engine/Plugins/Experimental/Water/Source/Runtime/Public/NiagaraDataInterfaceWater.h
@@ -60,7 +60,7 @@ private:
 	bool bEvaluateSystemDepthPerFrame = true;
 
 	UPROPERTY(EditAnywhere, Category = "Water", meta = (DisplayName = "Source Actor Or Component", AllowedClasses = "/Script/Engine.WaterBodyComponent,/Script/Engine.Actor"))
-	TObjectPtr<UObject> SourceBodyComponent;
+	TSoftObjectPtr<UObject> SourceBodyComponent;
 
 	uint32 SourceBodyChangeId = 0;
 };
