Made AWaterLandscapeBrush abstract, since it's useless in itself so we don't want people to be able to spawn it

#rnx

diff --git a/Engine/Plugins/Experimental/Water/Source/Editor/Private/WaterLandscapeBrush.h b/Engine/Plugins/Experimental/Water/Source/Editor/Private/WaterLandscapeBrush.h
--- a/Engine/Plugins/Experimental/Water/Source/Editor/Private/WaterLandscapeBrush.h
+++ b/Engine/Plugins/Experimental/Water/Source/Editor/Private/WaterLandscapeBrush.h
@@ -15,7 +15,7 @@ template <typename InterfaceType> class TScriptInterface;
 class AWaterBody;
 class AWaterBodyIsland;
 
-UCLASS(Blueprintable, hidecategories = (Collision, Replication, Input, LOD, Actor, Cooking, Rendering))
+UCLASS(Blueprintable, Abstract, hidecategories = (Collision, Replication, Input, LOD, Actor, Cooking, Rendering))
 class AWaterLandscapeBrush : public ALandscapeBlueprintBrush
 {
 	GENERATED_BODY()
