Made water body exclusion volumes not spatially loaded by default as they tend to be costly when spawned dynamically

diff --git a/Engine/Plugins/Experimental/Water/Source/Runtime/Private/WaterBodyExclusionVolume.cpp b/Engine/Plugins/Experimental/Water/Source/Runtime/Private/WaterBodyExclusionVolume.cpp
--- a/Engine/Plugins/Experimental/Water/Source/Runtime/Private/WaterBodyExclusionVolume.cpp
+++ b/Engine/Plugins/Experimental/Water/Source/Runtime/Private/WaterBodyExclusionVolume.cpp
@@ -24,6 +24,9 @@ AWaterBodyExclusionVolume::AWaterBodyExclusionVolume(const FObjectInitializer& O
 {
 #if WITH_EDITOR
 	ActorIcon = FWaterIconHelper::EnsureSpriteComponentCreated(this, TEXT("/Water/Icons/WaterBodyExclusionVolumeSprite"));
+	// By default, exclusion volume are not spatially loaded to avoid recomputing water body collision at runtime when the actor is streamed in,
+	//  which is allowed, but tends to be expensive
+	bIsSpatiallyLoaded = false;
 #endif
 }
 
