[Water Plugin] Fix validation error on mismatched rendertarget format when rendering the water info texture. Certain platforms require that SetRenderTargetOutputFormat be used to define the format of the render target if it differs from RGBA16F. FWaterInfoMergePS always set the format to RGBA32F even though it actually depends on the final WIT rendertarget format, which can be controlled by the user to be either 16bit or 32bit. An older version of this code probably actually always used RGBA32F as output texture for this shader, but that is no longer the case and I don't see a good reason to revert to that (assumed) older behavior.

Also removed some dead code that was related to WIT render method 0.

diff --git a/Engine/Plugins/Experimental/Water/Source/Runtime/Private/WaterInfoRendering.cpp b/Engine/Plugins/Experimental/Water/Source/Runtime/Private/WaterInfoRendering.cpp
--- a/Engine/Plugins/Experimental/Water/Source/Runtime/Private/WaterInfoRendering.cpp
+++ b/Engine/Plugins/Experimental/Water/Source/Runtime/Private/WaterInfoRendering.cpp
@@ -57,13 +57,16 @@ struct FUpdateWaterInfoParams
 
 // ---------------------------------------------------------------------------------------------------------------------
 
-/** A pixel shader for capturing a component of the rendered scene for a scene capture.*/
+/** Pixel shader for merging water flow velocity (XY), normalized water height (Z) and normalized ground height (W).*/
 class FWaterInfoMergePS : public FGlobalShader
 {
 public:
 	DECLARE_GLOBAL_SHADER(FWaterInfoMergePS);
 	SHADER_USE_PARAMETER_STRUCT(FWaterInfoMergePS, FGlobalShader);
 
+	class FEnable128BitRT : SHADER_PERMUTATION_BOOL("ENABLE_128_BIT");
+	using FPermutationDomain = TShaderPermutationDomain<FEnable128BitRT>;
+
 	BEGIN_SHADER_PARAMETER_STRUCT(FParameters, )
 		SHADER_PARAMETER_STRUCT_REF(FViewUniformShaderParameters, View)
 		SHADER_PARAMETER_STRUCT_INCLUDE(FSceneTextureShaderParameters, SceneTextures)
@@ -81,11 +84,24 @@ public:
 		RENDER_TARGET_BINDING_SLOTS()
 	END_SHADER_PARAMETER_STRUCT()
 
+	static FPermutationDomain GetPermutationVector(bool bUse128BitRT)
+	{
+		FPermutationDomain PermutationVector;
+		PermutationVector.Set<FEnable128BitRT>(bUse128BitRT);
+		return PermutationVector;
+	}
+
+	static bool ShouldCompilePermutation(const FGlobalShaderPermutationParameters& Parameters)
+	{
+		const FPermutationDomain PermutationVector(Parameters.PermutationId);
+		const bool bPlatformRequiresExplicit128bitRT = FDataDrivenShaderPlatformInfo::GetRequiresExplicit128bitRT(Parameters.Platform);
+		return (!PermutationVector.Get<FEnable128BitRT>() || bPlatformRequiresExplicit128bitRT);
+	}
+
 	static void ModifyCompilationEnvironment(const FGlobalShaderPermutationParameters& Parameters, FShaderCompilerEnvironment& OutEnvironment)
 	{
-		// Water info merge unconditionally requires a 128 bit render target. Some platforms require explicitly enabling this output mode.
-		bool bPlatformRequiresExplicit128bitRT = FDataDrivenShaderPlatformInfo::GetRequiresExplicit128bitRT(Parameters.Platform);
-		if (bPlatformRequiresExplicit128bitRT)
+		const FPermutationDomain PermutationVector(Parameters.PermutationId);
+		if (PermutationVector.Get<FEnable128BitRT>())
 		{
 			OutEnvironment.SetRenderTargetOutputFormat(0, PF_A32B32G32R32F);
 		}
@@ -113,6 +129,9 @@ static void MergeWaterInfoAndDepth(
 	GraphicsPSOInit.DepthStencilState = TStaticDepthStencilState<false, CF_Always>::GetRHI();
 	GraphicsPSOInit.BlendState = TStaticBlendState<>::GetRHI();
 
+	const bool bUse128BitRT = PlatformRequires128bitRT(OutputTexture->Desc.Format);
+	const FWaterInfoMergePS::FPermutationDomain PixelPermutationVector = FWaterInfoMergePS::GetPermutationVector(bUse128BitRT);
+
 	{
 		static auto* CVarDilationOverwriteMinimumDistance = IConsoleManager::Get().FindTConsoleVariableDataFloat(TEXT("r.Water.WaterInfo.DilationOverwriteMinimumDistance"));
 		static auto* CVarUndergroundDilationDepthOffset = IConsoleManager::Get().FindTConsoleVariableDataFloat(TEXT("r.Water.WaterInfo.UndergroundDilationDepthOffset"));
@@ -134,7 +153,7 @@ static void MergeWaterInfoAndDepth(
 		PassParameters->UndergroundDilationDepthOffset = CVarUndergroundDilationDepthOffset ? CVarUndergroundDilationDepthOffset->GetValueOnRenderThread() : 64.0f;
 
 		TShaderMapRef<FScreenVS> VertexShader(ShaderMap);
-		TShaderMapRef<FWaterInfoMergePS> PixelShader(ShaderMap);
+		TShaderMapRef<FWaterInfoMergePS> PixelShader(ShaderMap, PixelPermutationVector);
 
 		GraphicsPSOInit.BoundShaderState.VertexDeclarationRHI = GFilterVertexDeclaration.VertexDeclarationRHI;
 		GraphicsPSOInit.BoundShaderState.VertexShaderRHI = VertexShader.GetVertexShader();
@@ -160,13 +179,16 @@ static void MergeWaterInfoAndDepth(
 // ---------------------------------------------------------------------------------------------------------------------
 
 
-/** A pixel shader for capturing a component of the rendered scene for a scene capture.*/
+/** Pixel shader for blurring the water flow velocity component (XY) of the output of the water info merge pass. */
 class FWaterInfoFinalizePS : public FGlobalShader
 {
 public:
 	DECLARE_GLOBAL_SHADER(FWaterInfoFinalizePS);
 	SHADER_USE_PARAMETER_STRUCT(FWaterInfoFinalizePS, FGlobalShader);
 
+	class FEnable128BitRT : SHADER_PERMUTATION_BOOL("ENABLE_128_BIT");
+	using FPermutationDomain = TShaderPermutationDomain<FEnable128BitRT>;
+
 	BEGIN_SHADER_PARAMETER_STRUCT(FParameters, )
 		SHADER_PARAMETER_STRUCT_REF(FViewUniformShaderParameters, View)
 		SHADER_PARAMETER_STRUCT_INCLUDE(FSceneTextureShaderParameters, SceneTextures)
@@ -179,9 +201,6 @@ public:
 		RENDER_TARGET_BINDING_SLOTS()
 	END_SHADER_PARAMETER_STRUCT()
 
-	class FEnable128BitRT : SHADER_PERMUTATION_BOOL("ENABLE_128_BIT");
-	using FPermutationDomain = TShaderPermutationDomain<FEnable128BitRT>;
-
 	static FPermutationDomain GetPermutationVector(bool bUse128BitRT)
 	{
 		FPermutationDomain PermutationVector;
@@ -191,9 +210,8 @@ public:
 
 	static bool ShouldCompilePermutation(const FGlobalShaderPermutationParameters& Parameters) 
 	{
-		FPermutationDomain PermutationVector(Parameters.PermutationId);
-
-		bool bPlatformRequiresExplicit128bitRT = FDataDrivenShaderPlatformInfo::GetRequiresExplicit128bitRT(Parameters.Platform);
+		const FPermutationDomain PermutationVector(Parameters.PermutationId);
+		const bool bPlatformRequiresExplicit128bitRT = FDataDrivenShaderPlatformInfo::GetRequiresExplicit128bitRT(Parameters.Platform);
 		return (!PermutationVector.Get<FEnable128BitRT>() || bPlatformRequiresExplicit128bitRT);
 	}
 
@@ -289,156 +307,6 @@ static FMatrix BuildOrthoMatrix(float InOrthoWidth, float InOrthoHeight)
 
 // ---------------------------------------------------------------------------------------------------------------------
 
-struct FEnableWaterInfoMeshesScoped
-{
-public:
-	FEnableWaterInfoMeshesScoped(FSceneInterface* Scene, const FSceneView& View)
-	{
-		TRACE_CPUPROFILER_EVENT_SCOPE(WaterInfo::EnableWaterInfoMeshes);
-		if (View.ShowOnlyPrimitives.IsSet())
-		{
-			for (FPrimitiveSceneProxy* PrimProxy : Scene->GetPrimitiveSceneProxies())
-			{
-				if (PrimProxy && View.ShowOnlyPrimitives->Contains(PrimProxy->GetPrimitiveComponentId()))
-				{
-					FWaterBodyInfoMeshSceneProxy* WaterProxy = (FWaterBodyInfoMeshSceneProxy*)PrimProxy;
-					WaterProxy->SetEnabled(true);
-					Proxies.Add(WaterProxy);
-				}
-			}
-		}
-	}
-
-	~FEnableWaterInfoMeshesScoped()
-	{
-		TRACE_CPUPROFILER_EVENT_SCOPE(WaterInfo::DisableWaterInfoMeshes);
-		for (FWaterBodyInfoMeshSceneProxy* Proxy : Proxies)
-		{
-			Proxy->SetEnabled(false);
-		}
-	}
-private:
-	TArray<FWaterBodyInfoMeshSceneProxy*> Proxies;
-};
-
-/**
- * Sets up the landscape LOD override for water info rendering and restores it when the instance goes out of scope.
- */
-struct FScopedLandscapeLODOverride
-{
-public:
-	FScopedLandscapeLODOverride(const FUpdateWaterInfoParams& InParams)	
-		: Params(InParams)
-	{
-		if (ILandscapeModule* LandscapeModule = FModuleManager::GetModulePtr<ILandscapeModule>("Landscape"))
-		{
-			if (TSharedPtr<FLandscapeSceneViewExtension, ESPMode::ThreadSafe> LandscapeViewExtension = LandscapeModule->GetLandscapeSceneViewExtension())
-			{
-				LandscapeRenderSystems = &LandscapeViewExtension->GetLandscapeRenderSystems();
-			}
-		}
-
-		SetOptimalLandscapeLODOverrides();
-	}
-
-	~FScopedLandscapeLODOverride()
-	{
-		RestoreLandscapeLODOverrides();
-	}
-
-private:
-	FScopedLandscapeLODOverride(const FScopedLandscapeLODOverride&) = delete;
-	FScopedLandscapeLODOverride& operator=(const FScopedLandscapeLODOverride&) = delete;
-
-	void SetOptimalLandscapeLODOverrides()
-	{
-		if (!LandscapeRenderSystems)
-		{
-			return;
-		}
-		
-		// In order to prevent overdrawing the landscape components, we compute the lowest-detailed LOD level which satisfies the pixel coverage of the Water Info texture
-		// and force it on all landscape components. This override is set different per Landscape actor in case there are multiple under the same water zone.
-		//
-		// Ex: If the WaterInfoTexture only has 1 pixel per 100 units, and the highest landscape LOD has 1 vertex per 20 units, we don't need to use the maximum landscape LOD
-		// and can force a lower level of detail (in this case LOD2) while still satisfying the resolution of the water info texture.
-
-		const double MinWaterInfoTextureExtent = (double)FMath::Min(Params.OutputTexture->GetSizeX(), Params.OutputTexture->GetSizeY());
-		const double MaxWaterZoneExtent = FMath::Max(Params.WaterZoneExtents.X, Params.WaterZoneExtents.Y);
-		const double WaterInfoUnitsPerPixel =  MaxWaterZoneExtent / MinWaterInfoTextureExtent;
-
-		for (const TPair<uint32, FLandscapeRenderSystem*>& Pair : *LandscapeRenderSystems)
-		{
-			FLandscapeRenderSystem* LandscapeRenderSystem = Pair.Value;
-			int32 OptimalLODLevel = INDEX_NONE;
-
-			// All components within the same landscape (and thus its render system) should have the same number of quads and the same extent.
-			// therefore we can simply find the first component and compute its optimal LOD level.
-			for (FLandscapeSectionInfo* LandscapeSectionInfo : LandscapeRenderSystem->SectionInfos)
-			{
-				if (LandscapeSectionInfo != nullptr)
-				{
-					// Double the required landscape resolution to achieve 2 quads per pixel.
-					const double LandscapeComponentUnitsPerQuad = 2.0 * LandscapeSectionInfo->ComputeSectionResolution();
-					if (LandscapeComponentUnitsPerQuad <= 0.f)
-					{
-						// No section resolution probably means the section is a mesh proxy, which might not have regular units per vertex.
-						// Avoid computing optimal LOD in this case.
-						continue;
-					}
-
-					// Derived from:
-					// (ComponentWorldExtent / WaterInfoWorldspaceExtent) * WaterInfoTextureResolution = (NumComponentQuads / 2 ^ (LodLevel))
-					OptimalLODLevel = FMath::Max(WaterInfoRenderLandscapeMinimumMipLevel, FMath::FloorToInt(FMath::Log2(WaterInfoUnitsPerPixel / LandscapeComponentUnitsPerQuad)));
-
-					break;
-				}
-			}
-
-			LandscapeLODOverridesToRestore.Add(LandscapeRenderSystem, LandscapeRenderSystem->ForcedLODOverride);
-			LandscapeRenderSystem->ForcedLODOverride = OptimalLODLevel;
-		}
-	}
-
-	void RestoreLandscapeLODOverrides()
-	{
-		if (!LandscapeRenderSystems) return;
-		
-		for (const TPair<uint32, FLandscapeRenderSystem*>& Pair : *LandscapeRenderSystems)
-		{
-			FLandscapeRenderSystem* LandscapeRenderSystem = Pair.Value;
-
-			// Ignore landscape systems that were removed within the scope.
-			if (int8* ForcedLODOverride = LandscapeLODOverridesToRestore.Find(LandscapeRenderSystem))
-			{
-				LandscapeRenderSystem->ForcedLODOverride = *ForcedLODOverride;
-			}
-		}
-	}
-
-	const FUpdateWaterInfoParams& Params;
-	TMap<uint32, FLandscapeRenderSystem*> const* LandscapeRenderSystems = nullptr;
-	TMap<FLandscapeRenderSystem*, int8> LandscapeLODOverridesToRestore;
-};
-
-struct FCreateWaterInfoSceneRendererParams
-{
-public:
-	FCreateWaterInfoSceneRendererParams(const WaterInfo::FRenderingContext& InContext) : Context(InContext) {}
-
-	const WaterInfo::FRenderingContext& Context;
-
-	FSceneInterface* Scene = nullptr;
-	FRenderTarget* RenderTarget = nullptr;
-
-	FIntPoint RenderTargetSize = FIntPoint(EForceInit::ForceInit);
-	FMatrix ViewRotationMatrix = FMatrix(EForceInit::ForceInit);
-	FVector ViewLocation = FVector::Zero();
-	FMatrix ProjectionMatrix = FMatrix(EForceInit::ForceInit);
-	ESceneCaptureSource CaptureSource = ESceneCaptureSource::SCS_MAX;
-	TSet<FPrimitiveComponentId> ShowOnlyPrimitives;
-};
-
 void UpdateWaterInfoRendering2(FSceneView& InView, const FRenderingContext& Context, int32 RenderTargetArrayLayer, const FVector& WaterInfoCenter)
 {
 	TRACE_CPUPROFILER_EVENT_SCOPE(WaterInfo::UpdateWaterInfoRendering2);
@@ -591,8 +459,7 @@ public:
 	
 	virtual void OnPostRender(FRDGBuilder& GraphBuilder) override
 	{
-		FRDGTextureDesc TextureDesc(RenderTargetTexture->Desc);		
-		FRDGTextureRef MergeTargetTexture = GraphBuilder.CreateTexture(TextureDesc, TEXT("WaterInfoMerged"));
+		FRDGTextureRef MergeTargetTexture = GraphBuilder.CreateTexture(RenderTargetTexture->Desc, TEXT("WaterInfoMerged"));
 		MergeWaterInfoAndDepth(GraphBuilder, *Views[0]->Family, *Views[0], MergeTargetTexture, DepthPass->GetRenderTargetTexture(), ColorPass->GetRenderTargetTexture(), RenderTargetTexture, Params);
 	
 		FRDGTextureRef FinalizedTexture = GraphBuilder.CreateTexture(RenderTargetTexture->Desc, TEXT("WaterInfoFinalized"));
