Fix for disappearing water when hit proxy renders are those that trigger the water info texture render.
The problem is that custom render passes, when rendering hit proxies, are skipped, so the texture is cleared but never rendered

diff --git a/Engine/Plugins/Experimental/Water/Source/Runtime/Private/WaterViewExtension.cpp b/Engine/Plugins/Experimental/Water/Source/Runtime/Private/WaterViewExtension.cpp
--- a/Engine/Plugins/Experimental/Water/Source/Runtime/Private/WaterViewExtension.cpp
+++ b/Engine/Plugins/Experimental/Water/Source/Runtime/Private/WaterViewExtension.cpp
@@ -485,7 +485,7 @@ void FWaterViewExtension::SetupView(FSceneViewFamily& InViewFamily, FSceneView&
 
 	const TWeakObjectPtr<UWorld> WorldPtr = GetWorld();
 
-	if (!ensureMsgf(WorldPtr.IsValid(), TEXT("FWaterViewExtension::SetupView was called while it's owning world is not valid! Lifetime of the WaterViewExtension is tied to the world, this should be impossible!")))
+	if (!ensureMsgf(WorldPtr.IsValid(), TEXT("FWaterViewExtension::SetupView was called while its owning world is not valid! Lifetime of the WaterViewExtension is tied to the world, this should be impossible!")))
 	{
 		return;
 	}
@@ -638,7 +638,7 @@ void FWaterViewExtension::BeginRenderViewFamily(FSceneViewFamily& InViewFamily)
 
 	for (const FSceneView* View : InViewFamily.Views)
 	{
-		// Waning: Do not capture View in ENQUEUE_RENDER_COMMAND lambdas, since the RT's view hasn't been created yet,
+		// Warning: Do not capture View in ENQUEUE_RENDER_COMMAND lambdas, since the RT's view hasn't been created yet,
 		// so we would pass the GT view, which then would end up being a dangling pointer when accessed by the RT.
 
 		if (ShouldHaveWaterZoneViewData(*View))
@@ -739,7 +739,9 @@ bool FWaterViewExtension::ShouldHaveWaterZoneViewData(const FSceneView& InView)
 {
 	// Don't dirty the water info texture when we're rendering from a scene capture. Due to the frame delay after marking the texture as dirty, scene captures wouldn't have the right texture anyways.
 	// #todo_water [roey]: Once we have no frame-delay for updating the texture and lesser performance impact, we can re-enable updates within scene captures.
-	return !InView.bIsSceneCapture && !InView.bIsSceneCaptureCube && !InView.bIsReflectionCapture && !InView.bIsPlanarReflection && !InView.bIsVirtualTexture;
+	return !InView.bIsSceneCapture && !InView.bIsSceneCaptureCube && !InView.bIsReflectionCapture && !InView.bIsPlanarReflection && !InView.bIsVirtualTexture 
+		// Also don't update water info texture when rendering hit proxies as it bypasses custom render passes 
+		&& !InView.Family->EngineShowFlags.HitProxies;
 }
 
 void FWaterViewExtension::DrawDebugInfo(const FSceneView& InView, AWaterZone* WaterZone)
diff --git a/Engine/Source/Editor/UnrealEd/Private/EditorViewportClient.cpp b/Engine/Source/Editor/UnrealEd/Private/EditorViewportClient.cpp
--- a/Engine/Source/Editor/UnrealEd/Private/EditorViewportClient.cpp
+++ b/Engine/Source/Editor/UnrealEd/Private/EditorViewportClient.cpp
@@ -4307,6 +4307,12 @@ void FEditorViewportClient::Draw(FViewport* InViewport, FCanvas* Canvas)
 
 	ViewFamily.EngineShowFlags = UseEngineShowFlags;
 
+	// Set the view hit proxies flag as early as possible so that view extensions can check whether they render 
+	if (Canvas->IsHitTesting())
+	{
+		ViewFamily.EngineShowFlags.SetHitProxies(1);
+	}
+
 	ViewFamily.bIsHDR = Viewport->IsHDRViewport();
 	
 	// The view is in focus if it is currently in editing
diff --git a/Engine/Source/Runtime/Renderer/Private/SceneRendering.cpp b/Engine/Source/Runtime/Renderer/Private/SceneRendering.cpp
--- a/Engine/Source/Runtime/Renderer/Private/SceneRendering.cpp
+++ b/Engine/Source/Runtime/Renderer/Private/SceneRendering.cpp
@@ -2836,12 +2836,6 @@ FSceneRenderer::FSceneRenderer(const FSceneViewFamily* InViewFamily, FHitProxyCo
 		ViewFamily.Time = FGameTime::CreateDilated(0.0, ViewFamily.Time.GetDeltaRealTimeSeconds(), 0.0, ViewFamily.Time.GetDeltaWorldTimeSeconds());
 	}
 
-	if (HitProxyConsumer)
-	{
-		// Set the hit proxies show flag.
-		ViewFamily.EngineShowFlags.SetHitProxies(1);
-	}
-
 	// copy off the requests
 	if (ensure(InViewFamily->RenderTarget))
 	{
