Adding MinimalAPI specifier to UWaterEditorSettings to fix WorldBuildingTest access error

diff --git a/Engine/Plugins/Experimental/Water/Source/Editor/Public/WaterEditorSettings.h b/Engine/Plugins/Experimental/Water/Source/Editor/Public/WaterEditorSettings.h
--- a/Engine/Plugins/Experimental/Water/Source/Editor/Public/WaterEditorSettings.h
+++ b/Engine/Plugins/Experimental/Water/Source/Editor/Public/WaterEditorSettings.h
@@ -183,7 +183,7 @@ struct WATEREDITOR_API FWaterBodyIslandDefaults
 /**
  * Implements the editor settings for the Water plugin.
  */
-UCLASS(config = Engine, defaultconfig, meta=(DisplayName="Water Editor"))
+UCLASS(MinimalAPI, config = Engine, defaultconfig, meta=(DisplayName="Water Editor"))
 class UWaterEditorSettings : public UDeveloperSettings
 {
 	GENERATED_BODY()
