store water bodies in baked shallow water sim grid

diff --git a/Engine/Plugins/Experimental/Water/Source/Runtime/Private/BakedShallowWaterSimulationComponent.cpp b/Engine/Plugins/Experimental/Water/Source/Runtime/Private/BakedShallowWaterSimulationComponent.cpp
--- a/Engine/Plugins/Experimental/Water/Source/Runtime/Private/BakedShallowWaterSimulationComponent.cpp
+++ b/Engine/Plugins/Experimental/Water/Source/Runtime/Private/BakedShallowWaterSimulationComponent.cpp
@@ -1,5 +1,6 @@
 // Copyright Epic Games, Inc. All Rights Reserved.
 #include "BakedShallowWaterSimulationComponent.h"
+#include "WaterBodyActor.h"
 
 void FShallowWaterSimulationGrid::SampleShallowWaterSimulationAtIndex(const FVector2D &QueryFloatIndex, FVector& OutWaterVelocity, float& OutWaterHeight, float& OutWaterDepth) const
 {
diff --git a/Engine/Plugins/Experimental/Water/Source/Runtime/Public/BakedShallowWaterSimulationComponent.h b/Engine/Plugins/Experimental/Water/Source/Runtime/Public/BakedShallowWaterSimulationComponent.h
--- a/Engine/Plugins/Experimental/Water/Source/Runtime/Public/BakedShallowWaterSimulationComponent.h
+++ b/Engine/Plugins/Experimental/Water/Source/Runtime/Public/BakedShallowWaterSimulationComponent.h
@@ -7,6 +7,8 @@
 #include "Components/PrimitiveComponent.h"
 #include "BakedShallowWaterSimulationComponent.generated.h"
 
+class AWaterBody;
+
 USTRUCT(BlueprintType)
 struct WATER_API FShallowWaterSimulationGrid
 {
@@ -119,6 +121,10 @@ class UBakedShallowWaterSimulationComponent : public UPrimitiveComponent
 {
 	GENERATED_UCLASS_BODY()
 
+public:
 	UPROPERTY()
 	FShallowWaterSimulationGrid SimulationData;
+
+	UPROPERTY()
+	TSet<TSoftObjectPtr<AWaterBody>> WaterBodies;
 };
\ No newline at end of file
