[Water]

Only bind the callback to OnMarkRenderStateDirty when there are registered WaterTerrainComponents.

diff --git a/Engine/Plugins/Experimental/Water/Source/Runtime/Private/WaterSubsystem.cpp b/Engine/Plugins/Experimental/Water/Source/Runtime/Private/WaterSubsystem.cpp
--- a/Engine/Plugins/Experimental/Water/Source/Runtime/Private/WaterSubsystem.cpp
+++ b/Engine/Plugins/Experimental/Water/Source/Runtime/Private/WaterSubsystem.cpp
@@ -308,8 +308,6 @@ void UWaterSubsystem::PostInitialize()
 		GEngine->OnActorMoved().AddUObject(this, &UWaterSubsystem::OnActorMoved);
 	}
 #endif // WITH_EDITOR
-
-	UActorComponent::MarkRenderStateDirtyEvent.AddUObject(this, &UWaterSubsystem::OnMarkRenderStateDirty);
 }
 
 void UWaterSubsystem::Deinitialize()
@@ -324,7 +322,10 @@ void UWaterSubsystem::Deinitialize()
 	}
 #endif // WITH_EDITOR
 
-	UActorComponent::MarkRenderStateDirtyEvent.RemoveAll(this);
+	if (OnMarkRenderStateDirtyHandle.IsValid())
+	{
+		UActorComponent::MarkRenderStateDirtyEvent.Remove(OnMarkRenderStateDirtyHandle);
+	}
 
 	FConsoleVariableDelegate NullCallback;
 	CVarShallowWaterSimulationRenderTargetSize->SetOnChangedCallback(NullCallback);
@@ -667,6 +668,11 @@ void UWaterSubsystem::RegisterWaterTerrainComponent(UWaterTerrainComponent* InWa
 	{
 		 WaterTerrainActors.Add(TerrainActor,  InWaterTerrainComponent);
 	}
+
+	if (!WaterTerrainActors.IsEmpty() && !OnMarkRenderStateDirtyHandle.IsValid())
+	{
+		OnMarkRenderStateDirtyHandle = UActorComponent::MarkRenderStateDirtyEvent.AddUObject(this, &UWaterSubsystem::OnMarkRenderStateDirty);
+	}
 }
 
 void UWaterSubsystem::UnregisterWaterTerrainComponent(UWaterTerrainComponent* InWaterTerrainComponent)
@@ -676,6 +682,12 @@ void UWaterSubsystem::UnregisterWaterTerrainComponent(UWaterTerrainComponent* In
 	{
 		WaterTerrainActors.RemoveSingle(TerrainActor, InWaterTerrainComponent);
 	}
+
+	if (WaterTerrainActors.IsEmpty() && ensure(OnMarkRenderStateDirtyHandle.IsValid()))
+	{
+		UActorComponent::MarkRenderStateDirtyEvent.Remove(OnMarkRenderStateDirtyHandle);
+		OnMarkRenderStateDirtyHandle.Reset();
+	}
 }
 
 void UWaterSubsystem::GetWaterTerrainComponents(TArray<UWaterTerrainComponent*>& OutWaterTerrainComponents) const
diff --git a/Engine/Plugins/Experimental/Water/Source/Runtime/Public/WaterSubsystem.h b/Engine/Plugins/Experimental/Water/Source/Runtime/Public/WaterSubsystem.h
--- a/Engine/Plugins/Experimental/Water/Source/Runtime/Public/WaterSubsystem.h
+++ b/Engine/Plugins/Experimental/Water/Source/Runtime/Public/WaterSubsystem.h
@@ -267,6 +267,7 @@ private:
 	 * every actors components to find one when global events are triggered.
 	 */
 	TMultiMap<const AActor*, TWeakObjectPtr<UWaterTerrainComponent>> WaterTerrainActors;
+	FDelegateHandle OnMarkRenderStateDirtyHandle{};
 
 #if WITH_EDITOR
 	/** By default, there is no water subsystem allowed on preview worlds except when explicitly requested : */
