Truly prevent underwater post process processing when anything disables it (post processing show flag, r.Water.EnableUnderwaterPostProcess or path tracing)

diff --git a/Engine/Plugins/Experimental/Water/Source/Runtime/Private/WaterSubsystem.cpp b/Engine/Plugins/Experimental/Water/Source/Runtime/Private/WaterSubsystem.cpp
--- a/Engine/Plugins/Experimental/Water/Source/Runtime/Private/WaterSubsystem.cpp
+++ b/Engine/Plugins/Experimental/Water/Source/Runtime/Private/WaterSubsystem.cpp
@@ -133,6 +133,9 @@ struct FUnderwaterPostProcessDebugInfo
 	TArray<TWeakObjectPtr<UWaterBodyComponent>> OverlappedWaterBodyComponents;
 	TWeakObjectPtr<UWaterBodyComponent> ActiveWaterBodyComponent;
 	FWaterBodyQueryResult ActiveWaterBodyQueryResult;
+	bool bIsPostProcessingEnabled = false;
+	bool bIsUnderwaterPostProcessEnabled = false;
+	bool bIsPathTracingEnabled = false;
 };
 #endif // !(UE_BUILD_SHIPPING || UE_BUILD_TEST)
 
@@ -754,14 +757,32 @@ void UWaterSubsystem::ComputeUnderwaterPostProcess(FVector ViewLocation, FSceneV
 	SCOPE_CYCLE_COUNTER(STAT_WaterIsUnderwater);
 
 	UWorld* World = GetWorld();
-	if ((World == nullptr) || (SceneView->Family->EngineShowFlags.PostProcessing == 0))
-	{
-		return;
-	}
+	const bool bIsPostProcessingEnabled = (SceneView->Family->EngineShowFlags.PostProcessing != 0);
+	const bool bIsUnderwaterPostProcessEnabled = IsUnderwaterPostProcessEnabled();
+	const bool bIsPathTracingEnabled = (SceneView->Family->EngineShowFlags.PathTracing != 0);
 
 	const float PrevDepthUnderwater = CachedDepthUnderwater;
 	CachedDepthUnderwater = -1;
 
+	// Set all that needs to be set before an eventual early-out
+	UnderwaterPostProcessVolume.PostProcessProperties.bIsEnabled = false;
+	UnderwaterPostProcessVolume.PostProcessProperties.Settings = nullptr;
+	SceneView->UnderwaterDepth = CachedDepthUnderwater;
+	SceneView->WaterIntersection = EViewWaterIntersection::OutsideWater;
+
+#if !(UE_BUILD_SHIPPING || UE_BUILD_TEST)
+	FUnderwaterPostProcessDebugInfo UnderwaterPostProcessDebugInfo;
+	UnderwaterPostProcessDebugInfo.bIsPostProcessingEnabled = bIsPostProcessingEnabled;
+	UnderwaterPostProcessDebugInfo.bIsUnderwaterPostProcessEnabled = bIsUnderwaterPostProcessEnabled;
+	UnderwaterPostProcessDebugInfo.bIsPathTracingEnabled = bIsPathTracingEnabled;
+	ON_SCOPE_EXIT { ShowOnScreenDebugInfo(ViewLocation, UnderwaterPostProcessDebugInfo); };
+#endif // !(UE_BUILD_SHIPPING || UE_BUILD_TEST)
+
+	if ((World == nullptr) || !bIsPostProcessingEnabled || !bIsUnderwaterPostProcessEnabled || bIsPathTracingEnabled)
+	{
+		return;
+	}
+
 	// Compute distance from view origin to the corner of the near plane. This distance needs to be taken into account when computing whether the view intersects the water surface.
 	const FVector4f NearPlaneCornerViewSpace = FVector4f(SceneView->ViewMatrices.GetInvProjectionMatrix().TransformFVector4(FVector4(1.0f, 1.0f, (bool)ERHIZBuffer::IsInverted ? 1.0f : 0.0f, 1.0f)));
 	const float ViewToNearPlaneCornerDistance = FVector2f(NearPlaneCornerViewSpace / NearPlaneCornerViewSpace.W).Length();
@@ -776,10 +797,6 @@ void UWaterSubsystem::ComputeUnderwaterPostProcess(FVector ViewLocation, FSceneV
 	// Always force simple collision traces
 	static FCollisionQueryParams TraceSimple(SCENE_QUERY_STAT(DefaultQueryParam), false);
 
-#if !(UE_BUILD_SHIPPING || UE_BUILD_TEST)
-	FUnderwaterPostProcessDebugInfo UnderwaterPostProcessDebugInfo;
-#endif // !(UE_BUILD_SHIPPING || UE_BUILD_TEST)
-
 	TArray<FHitResult> Hits;
 	TArray<FWaterBodyPostProcessQuery, TInlineAllocator<4>> WaterBodyQueriesToProcess;
 	const bool bWorldHasWater  = WaterBodyManager.HasAnyWaterBodies();
@@ -890,16 +907,6 @@ void UWaterSubsystem::ComputeUnderwaterPostProcess(FVector ViewLocation, FSceneV
 	{
 		SceneView->WaterIntersection = EViewWaterIntersection::OutsideWater;
 	}
-
-	if (!bUnderwaterForPostProcess || !IsUnderwaterPostProcessEnabled() || SceneView->Family->EngineShowFlags.PathTracing)
-	{
-		UnderwaterPostProcessVolume.PostProcessProperties.bIsEnabled = false;
-		UnderwaterPostProcessVolume.PostProcessProperties.Settings = nullptr;
-	}
-
-#if !(UE_BUILD_SHIPPING || UE_BUILD_TEST)
-	ShowOnScreenDebugInfo(ViewLocation, UnderwaterPostProcessDebugInfo);
-#endif // !(UE_BUILD_SHIPPING || UE_BUILD_TEST)						
 }
 
 void UWaterSubsystem::SetMPCTime(float Time, float PrevTime)
@@ -996,6 +1003,21 @@ void UWaterSubsystem::ShowOnScreenDebugInfo(const FVector& InViewLocation, const
 				}
 			}
 		}
+
+		if (!InDebugInfo.bIsPostProcessingEnabled)
+		{
+			OutputStrings.Add(LOCTEXT("VisualizeActiveUnderwaterPostProcess_PostProcessingDisabled", "Post processing is disabled"));
+		}
+
+		if (!InDebugInfo.bIsUnderwaterPostProcessEnabled)
+		{
+			OutputStrings.Add(LOCTEXT("VisualizeActiveUnderwaterPostProcess_UnderwaterPostProcessDisabled", "Underwater post process is disabled"));
+		}
+
+		if (InDebugInfo.bIsPathTracingEnabled)
+		{
+			OutputStrings.Add(LOCTEXT("VisualizeActiveUnderwaterPostProcess_PathTracingEnabled", "Path tracing is enabled"));
+		}
 	}
 
 	// Output a single message because multi-line texts end up overlapping over messages
