WaterViewExtension multithreading timing issue fix

diff --git a/Engine/Plugins/Experimental/Water/Source/Runtime/Private/WaterSubsystem.cpp b/Engine/Plugins/Experimental/Water/Source/Runtime/Private/WaterSubsystem.cpp
--- a/Engine/Plugins/Experimental/Water/Source/Runtime/Private/WaterSubsystem.cpp
+++ b/Engine/Plugins/Experimental/Water/Source/Runtime/Private/WaterSubsystem.cpp
@@ -203,6 +203,15 @@ FWaterViewExtension* UWaterSubsystem::GetWaterViewExtension(const UWorld* InWorl
 	return {};
 }
 
+TWeakPtr<FWaterViewExtension, ESPMode::ThreadSafe> UWaterSubsystem::GetWaterViewExtensionWeakPtr(const UWorld* InWorld)
+{
+	if (FWaterBodyManager* Manager = GetWaterBodyManager(InWorld))
+	{
+		return Manager->GetWaterViewExtensionWeakPtr();
+	}
+	return {};
+}
+
 void UWaterSubsystem::Tick(float DeltaTime)
 {
 	TRACE_CPUPROFILER_EVENT_SCOPE(UWaterSubsystem::Tick);
diff --git a/Engine/Plugins/Experimental/Water/Source/Runtime/Private/WaterViewExtension.cpp b/Engine/Plugins/Experimental/Water/Source/Runtime/Private/WaterViewExtension.cpp
--- a/Engine/Plugins/Experimental/Water/Source/Runtime/Private/WaterViewExtension.cpp
+++ b/Engine/Plugins/Experimental/Water/Source/Runtime/Private/WaterViewExtension.cpp
@@ -517,10 +517,15 @@ void FWaterViewExtension::SetupView(FSceneViewFamily& InViewFamily, FSceneView&
 	{
 		ViewPlayerIndices.Empty();
 
+		TWeakPtr<FWaterViewExtension, ESPMode::ThreadSafe> WaterViewExtension = UWaterSubsystem::GetWaterViewExtensionWeakPtr(WorldPtr.Get());
+
 		ENQUEUE_RENDER_COMMAND(WaterViewExtensionNonDataViewsQuadtreeKeysReset)(
-			[this](FRHICommandList& THICmdList)
+			[WaterViewExtension](FRHICommandList& THICmdList)
 			{
-				NonDataViewsQuadtreeKeys.Empty();
+				if (TSharedPtr<FWaterViewExtension, ESPMode::ThreadSafe> WaterViewExtensionPtr = WaterViewExtension.Pin(); WaterViewExtensionPtr.IsValid())
+				{
+					WaterViewExtensionPtr->NonDataViewsQuadtreeKeys.Empty();
+				}
 			});
 	}
 	else
@@ -719,10 +724,15 @@ void FWaterViewExtension::BeginRenderViewFamily(FSceneViewFamily& InViewFamily)
 							const FVector2D ViewPosition2D = FVector2D(View->ViewLocation);
 							FSceneViewStateInterface* ViewState = View->State;
 
+							TWeakPtr<FWaterViewExtension, ESPMode::ThreadSafe> WaterViewExtension = UWaterSubsystem::GetWaterViewExtensionWeakPtr(WorldPtr.Get());
+
 							ENQUEUE_RENDER_COMMAND(WaterViewExtensionNonDataViewsQuadtreeKeysUpdate)(
-								[this, SceneProxy, ViewPosition2D, ViewState](FRHICommandList& THICmdList)
+								[WaterViewExtension, SceneProxy, ViewPosition2D, ViewState](FRHICommandList& THICmdList)
 								{
-									NonDataViewsQuadtreeKeys.Add(ViewState, SceneProxy->FindBestQuadTreeForViewLocation(ViewPosition2D));
+									if (TSharedPtr<FWaterViewExtension, ESPMode::ThreadSafe> WaterViewExtensionPtr = WaterViewExtension.Pin(); WaterViewExtensionPtr.IsValid())
+									{
+										WaterViewExtensionPtr->NonDataViewsQuadtreeKeys.Add(ViewState, SceneProxy->FindBestQuadTreeForViewLocation(ViewPosition2D));
+									}
 								});
 
 							UE_LOG(LogWater, Verbose, TEXT("Water Zone (%s) queued a search for the closest quadtree for a View (0x%p) which has no WaterInfo."), *GetNameSafe(WaterZone), View);
diff --git a/Engine/Plugins/Experimental/Water/Source/Runtime/Public/WaterBodyManager.h b/Engine/Plugins/Experimental/Water/Source/Runtime/Public/WaterBodyManager.h
--- a/Engine/Plugins/Experimental/Water/Source/Runtime/Public/WaterBodyManager.h
+++ b/Engine/Plugins/Experimental/Water/Source/Runtime/Public/WaterBodyManager.h
@@ -63,6 +63,8 @@ public:
 
 	FWaterViewExtension* GetWaterViewExtension() { return WaterViewExtension.Get(); }
 
+	TWeakPtr<FWaterViewExtension, ESPMode::ThreadSafe> GetWaterViewExtensionWeakPtr() { return WaterViewExtension; }
+
 	FWaterBodyEvent OnWaterBodyAdded;
 
 	FWaterBodyEvent OnWaterBodyRemoved;
diff --git a/Engine/Plugins/Experimental/Water/Source/Runtime/Public/WaterSubsystem.h b/Engine/Plugins/Experimental/Water/Source/Runtime/Public/WaterSubsystem.h
--- a/Engine/Plugins/Experimental/Water/Source/Runtime/Public/WaterSubsystem.h
+++ b/Engine/Plugins/Experimental/Water/Source/Runtime/Public/WaterSubsystem.h
@@ -106,6 +106,9 @@ public:
 	/** Static helper function to get a weak ptr to the water scene view extension for a given world. */
 	static UE_API FWaterViewExtension* GetWaterViewExtension(const UWorld* InWorld);
 
+	/** Static helper function to get a TWeakPtr to the water scene view extension for a given world. */
+	static TWeakPtr<FWaterViewExtension, ESPMode::ThreadSafe> GetWaterViewExtensionWeakPtr(const UWorld* InWorld);
+
 	ABuoyancyManager* GetBuoyancyManager() const { return BuoyancyManager; }
 
 	TWeakObjectPtr<UWaterBodyComponent> GetOceanBodyComponent() { return OceanBodyComponent; }
