WaterViewExtension optimization: only loop through WaterZones for quadtree updates when at least one update is actually needed

diff --git a/Engine/Plugins/Experimental/Water/Source/Runtime/Private/WaterViewExtension.cpp b/Engine/Plugins/Experimental/Water/Source/Runtime/Private/WaterViewExtension.cpp
--- a/Engine/Plugins/Experimental/Water/Source/Runtime/Private/WaterViewExtension.cpp
+++ b/Engine/Plugins/Experimental/Water/Source/Runtime/Private/WaterViewExtension.cpp
@@ -600,6 +600,19 @@ void FWaterViewExtension::SetupView(FSceneViewFamily& InViewFamily, FSceneView&
 						}
 
 						WaterZoneViewInfo.bShouldUpdateQuadtree = true;
+						bAnyQuadTreeUpdateRequired = true;
+					}
+
+					{
+						UWaterMeshComponent* WaterMeshComponent = WaterZone->GetWaterMeshComponent();
+						check(WaterMeshComponent != nullptr);
+
+						FWaterMeshSceneProxy* SceneProxy = static_cast<FWaterMeshSceneProxy*>(WaterMeshComponent->GetSceneProxy());
+
+						if (SceneProxy != nullptr)
+						{
+							bAnyQuadTreeUpdateRequired |= (SceneProxy != WaterZoneViewInfo.OldSceneProxy);
+						}
 					}
 
 					WaterZoneViewInfo.bIsDirty = false;
@@ -648,6 +661,11 @@ void FWaterViewExtension::BeginRenderViewFamily(FSceneViewFamily& InViewFamily)
 
 		if (ShouldHaveWaterZoneViewData(*View))
 		{
+			if (!bAnyQuadTreeUpdateRequired)
+			{
+				continue;
+			}
+
 			for (AWaterZone* WaterZone : TActorRange<AWaterZone>(WorldPtr.Get()))
 			{
 				int32 ViewPlayerIndex = GetViewIndex(*View);
@@ -743,6 +761,8 @@ void FWaterViewExtension::BeginRenderViewFamily(FSceneViewFamily& InViewFamily)
 		}
 
 	}
+
+	bAnyQuadTreeUpdateRequired = false;
 }
 
 bool FWaterViewExtension::ShouldHaveWaterZoneViewData(const FSceneView& InView) const
diff --git a/Engine/Plugins/Experimental/Water/Source/Runtime/Public/WaterViewExtension.h b/Engine/Plugins/Experimental/Water/Source/Runtime/Public/WaterViewExtension.h
--- a/Engine/Plugins/Experimental/Water/Source/Runtime/Public/WaterViewExtension.h
+++ b/Engine/Plugins/Experimental/Water/Source/Runtime/Public/WaterViewExtension.h
@@ -90,6 +90,9 @@ private:
 	bool bRequestForcedBoundsUpdate = false;
 	bool bForceBoundsUpdate = false;
 
+	// this flag lets us know if we can skip looping through water zones for quadtrees updates, as an optimization
+	bool bAnyQuadTreeUpdateRequired = false;
+
 	// store the locations of every active water mesh scene proxy quad tree based on the key
 	TMap<int32, FVector2D> QuadTreeKeyLocationMap;
 
