DLL exported some public functions in some of the Water apis.

PR #12325

diff --git a/Engine/Plugins/Experimental/Water/Source/Runtime/Public/GerstnerWaterWaveSubsystem.h b/Engine/Plugins/Experimental/Water/Source/Runtime/Public/GerstnerWaterWaveSubsystem.h
--- a/Engine/Plugins/Experimental/Water/Source/Runtime/Public/GerstnerWaterWaveSubsystem.h
+++ b/Engine/Plugins/Experimental/Water/Source/Runtime/Public/GerstnerWaterWaveSubsystem.h
@@ -10,7 +10,7 @@ struct FResourceSizeEx;
 class FWaterViewExtension;
 
 /** UGerstnerWaterWaveSubsystem manages all UGerstnerWaterWaves objects, regardless of which world they belong to (it's a UEngineSubsystem) */
-UCLASS()
+UCLASS(MinimalAPI)
 class UGerstnerWaterWaveSubsystem : public UEngineSubsystem
 {
 	GENERATED_BODY()
@@ -19,16 +19,16 @@ public:
 	UGerstnerWaterWaveSubsystem();
 
 	// UEngineSubsystem implementation
-	virtual void Initialize(FSubsystemCollectionBase& Collection) override;
-	virtual void Deinitialize() override;
+	WATER_API virtual void Initialize(FSubsystemCollectionBase& Collection) override;
+	WATER_API virtual void Deinitialize() override;
 
-	void Register(FWaterViewExtension* ViewExtension);
-	void Unregister(FWaterViewExtension* ViewExtension);
+	WATER_API void Register(FWaterViewExtension* ViewExtension);
+	WATER_API void Unregister(FWaterViewExtension* ViewExtension);
 
-	void RebuildGPUData() { bRebuildGPUData = true; }
+	WATER_API void RebuildGPUData() { bRebuildGPUData = true; }
 
 	//~ Begin UObject Interface.	
-	virtual void GetResourceSizeEx(FResourceSizeEx& CumulativeResourceSize) override;
+	WATER_API virtual void GetResourceSizeEx(FResourceSizeEx& CumulativeResourceSize) override;
 	//~ End UObject Interface
 
 private:
diff --git a/Engine/Plugins/Experimental/Water/Source/Runtime/Public/GerstnerWaterWaves.h b/Engine/Plugins/Experimental/Water/Source/Runtime/Public/GerstnerWaterWaves.h
--- a/Engine/Plugins/Experimental/Water/Source/Runtime/Public/GerstnerWaterWaves.h
+++ b/Engine/Plugins/Experimental/Water/Source/Runtime/Public/GerstnerWaterWaves.h
@@ -177,17 +177,17 @@ public:
 	UGerstnerWaterWaves();
 	
 	/** Returns the maximum wave height that can be reached by those waves */
-	virtual float GetMaxWaveHeight() const override { return MaxWaveHeight; }
+	WATER_API virtual float GetMaxWaveHeight() const override { return MaxWaveHeight; }
 
 	/** Computes the raw wave perturbation of the water height/normal */
-	virtual float GetWaveHeightAtPosition(const FVector& InPosition, float InWaterDepth, float InTime, FVector& OutNormal) const override;
+	WATER_API virtual float GetWaveHeightAtPosition(const FVector& InPosition, float InWaterDepth, float InTime, FVector& OutNormal) const override;
 
 	/** Computes the raw wave perturbation of the water height only (simple version : faster computation) */
-	virtual float GetSimpleWaveHeightAtPosition(const FVector& InPosition, float InWaterDepth, float InTime) const override;
+	WATER_API virtual float GetSimpleWaveHeightAtPosition(const FVector& InPosition, float InWaterDepth, float InTime) const override;
 
 	/** Computes the attenuation factor to apply to the raw wave perturbation. Attenuates : normal/wave height/max wave height. 
 	Should match the GPU version (ComputeWaveDepthAttenuationFactor) */
-	virtual float GetWaveAttenuationFactor(const FVector& InPosition, float InWaterDepth, float InTargetWaveMaskDepth) const override;
+	WATER_API virtual float GetWaveAttenuationFactor(const FVector& InPosition, float InWaterDepth, float InTargetWaveMaskDepth) const override;
 
 	UPROPERTY(EditAnywhere, BlueprintReadOnly, Instanced, Category = Waves)
 	TObjectPtr<UGerstnerWaterWaveGeneratorBase> GerstnerWaveGenerator;
@@ -200,7 +200,7 @@ protected:
 	float MaxWaveHeight;
 
 public:
-	const TArray<FGerstnerWave>& GetGerstnerWaves() const { return GerstnerWaves; }
+	WATER_API const TArray<FGerstnerWave>& GetGerstnerWaves() const { return GerstnerWaves; }
 
 	virtual void PostDuplicate(EDuplicateMode::Type DuplicateMode) override;
 #if WITH_EDITOR
diff --git a/Engine/Plugins/Experimental/Water/Source/Runtime/Public/WaterViewExtension.h b/Engine/Plugins/Experimental/Water/Source/Runtime/Public/WaterViewExtension.h
--- a/Engine/Plugins/Experimental/Water/Source/Runtime/Public/WaterViewExtension.h
+++ b/Engine/Plugins/Experimental/Water/Source/Runtime/Public/WaterViewExtension.h
@@ -28,14 +28,14 @@ public:
 	virtual void PreRenderBasePass_RenderThread(FRDGBuilder& GraphBuilder, bool bDepthBufferIsPopulated) override;
 	// End FSceneViewExtensionBase implementation
 
-	void MarkWaterInfoTextureForRebuild(const UE::WaterInfo::FRenderingContext& RenderContext);
+	WATER_API void MarkWaterInfoTextureForRebuild(const UE::WaterInfo::FRenderingContext& RenderContext);
 
-	void MarkGPUDataDirty();
+	WATER_API void MarkGPUDataDirty();
 
-	void AddWaterZone(AWaterZone* InWaterZone);
-	void RemoveWaterZone(AWaterZone* InWaterZone);
+	WATER_API void AddWaterZone(AWaterZone* InWaterZone);
+	WATER_API void RemoveWaterZone(AWaterZone* InWaterZone);
 
-	bool GetZoneLocation(const AWaterZone* InWaterZone, int32 PlayerIndex, FVector& OutLocation) const;
+	WATER_API bool GetZoneLocation(const AWaterZone* InWaterZone, int32 PlayerIndex, FVector& OutLocation) const;
 
 	void CreateSceneProxyQuadtrees(FWaterMeshSceneProxy* SceneProxy);
 
