[Water]

Fixed ground min/max code grabbing all component bounds from a landscape component instead of just the ULandscapeComponent bounds. This fixes an issue where components attached to landscape actors could affect the precision of our depth values.

diff --git a/Engine/Plugins/Experimental/Water/Source/Runtime/Private/WaterZoneActor.cpp b/Engine/Plugins/Experimental/Water/Source/Runtime/Private/WaterZoneActor.cpp
--- a/Engine/Plugins/Experimental/Water/Source/Runtime/Private/WaterZoneActor.cpp
+++ b/Engine/Plugins/Experimental/Water/Source/Runtime/Private/WaterZoneActor.cpp
@@ -672,15 +672,26 @@ bool AWaterZone::UpdateWaterInfoTexture()
 		{
 			for (ALandscapeProxy* LandscapeProxy : TActorRange<ALandscapeProxy>(World))
 			{
-				 const FBox LandscapeBox = LandscapeProxy->GetComponentsBoundingBox(/*bIncludeNonColliding = */ true);
-				 // Only consider landscapes which this zone intersects with in XY and if the landscape volume is not zero sized
-				 if (WaterZoneBounds.IntersectXY(LandscapeBox) && LandscapeBox.GetVolume() > 0.0)
-				 {
-					 GroundZMin = FMath::Min(GroundZMin, LandscapeBox.Min.Z);
-					 GroundZMax = FMath::Max(GroundZMax, LandscapeBox.Max.Z);
-					 TInlineComponentArray<ULandscapeComponent*> LandscapeComponents(LandscapeProxy);
-					 GroundPrimitiveComponents.Append(LandscapeComponents);
-				 }
+				FBox LandscapeBox;
+				// Only collect the bounds for the ULandscapeComponents to avoid growing the bounding box unnecessarily
+				// if the landscape proxy has other components attached to it such as Grass or any other user-added components.
+				LandscapeProxy->ForEachComponent<ULandscapeComponent>(true, [&](const ULandscapeComponent* InLandscapeComp)
+					{
+						// Don't include bounds for non-registered components as they will not have the correct worldspace transform
+						if (InLandscapeComp->IsRegistered())
+						{
+							LandscapeBox += InLandscapeComp->Bounds.GetBox();
+						}
+					});
+
+				// Only consider landscapes which this zone intersects with in XY and if the landscape volume is not zero sized
+				if (WaterZoneBounds.IntersectXY(LandscapeBox) && LandscapeBox.GetVolume() > 0.0)
+				{
+					GroundZMin = FMath::Min(GroundZMin, LandscapeBox.Min.Z);
+					GroundZMax = FMath::Max(GroundZMax, LandscapeBox.Max.Z);
+					TInlineComponentArray<ULandscapeComponent*> LandscapeComponents(LandscapeProxy);
+					GroundPrimitiveComponents.Append(LandscapeComponents);
+				}
 			}
 		}
 
