[Water]

Implemented a scalability cvar which controls a cap on the maximum extent in tiles a water quad tree can have. Any water mesh component which attempts to allocate more than this number of tiles will be biased back to generate lower res tiles and emit an error so designers can either tweak the parameters of the water zone or change the cvar.

diff --git a/Engine/Plugins/Experimental/Water/Source/Runtime/Private/WaterMeshComponent.cpp b/Engine/Plugins/Experimental/Water/Source/Runtime/Private/WaterMeshComponent.cpp
--- a/Engine/Plugins/Experimental/Water/Source/Runtime/Private/WaterMeshComponent.cpp
+++ b/Engine/Plugins/Experimental/Water/Source/Runtime/Private/WaterMeshComponent.cpp
@@ -34,6 +34,11 @@ static TAutoConsoleVariable<float> CVarWaterMeshLODScaleBias(
 	TEXT("This value is added to the LOD Scale of each Mesh Component. Negative values will lower the overall density/resolution or the vertex grid and make the LODs smaller, higher values will increase the density/resolution and make the LODs larger. Smallest value is -0.5. That will make the inner LOD as tight and optimized as possible"),
 	ECVF_Scalability);
 
+static TAutoConsoleVariable<int32> CVarWaterMeshMaxDimensionInTiles(
+	TEXT("r.Water.WaterMesh.MaxDimensionInTiles"), 256,
+	TEXT("Maximum dimension of a water mesh in tiles (max of either X or Y). If a water mesh attempts to generate a quadtree with more than this many tiles in X or Y, it will be biased back below this threshold via the same process as the LODScaleBias (iteratively dividing tile size by 2 until compliance). Having too many tiles can create very large GPU allocations."),
+	ECVF_Scalability);
+
 TAutoConsoleVariable<int32> CVarWaterMeshGPUQuadTree(
 	TEXT("r.Water.WaterMesh.GPUQuadTree"),
 	0,
@@ -563,9 +568,25 @@ void UWaterMeshComponent::Update()
 		LODCountBiasScalability = NewLODCountBias;
 		TessFactorBiasScalability = NewTessFactorBias;
 		LODScaleBiasScalability = NewLODScaleBias;
-		const float LODCountBiasFactor = FMath::Pow(2.0f, (float)LODCountBiasScalability);
+		const FIntPoint ExtentInTiles = GetExtentInTiles();
+
+		const int32 MaxDimensionInTiles = CVarWaterMeshMaxDimensionInTiles.GetValueOnGameThread();
+
+		int32 Bias = 0;
+		float LODCountBiasFactor;
+		do
+		{
+			LODCountBiasFactor = FMath::Pow(2.0f, (float)LODCountBiasScalability - Bias);
+		}
+		// Cap the loop at 100 iterations in case users input a bad MaxExtentInTiles.
+		// note: the shortcircuiting of the increment of bias in this condition evaluation is import for detecting if the loop falls through or was executed.
+		while (ExtentInTiles.GetMax() * LODCountBiasFactor > MaxDimensionInTiles && ++Bias < 100);
+
+		if (Bias != 0)
+		{
+			UE_LOG(LogWater, Warning, TEXT("Width of water quad tree tiles (%d) for Water Mesh Component (%s) has exceeded the cap for this platform (%d). Tile sizes have been biased by a factor of %.2f to prevent exceeding the cap. If a high tile count is intended by the user, the scalability cvar `r.Water.WaterMesh.MaxWidthInTiles` should be adjusted accordingly."), ExtentInTiles.GetMax(), *GetPathName(), MaxDimensionInTiles, LODCountBiasFactor);
+		}
 
-		FIntPoint ExtentInTiles = GetExtentInTiles();
 		RebuildWaterMesh(TileSize / LODCountBiasFactor, FIntPoint(FMath::CeilToInt(ExtentInTiles.X * LODCountBiasFactor), FMath::CeilToInt(ExtentInTiles.Y * LODCountBiasFactor)));
 		PrecachePSOs();
 		bNeedsRebuild = false;
