WaterSplineComponent: Fix unnecessary tangent setting when updating water splines due to a regression in the tangent clamping logic.

diff --git a/Engine/Plugins/Experimental/Water/Source/Runtime/Private/WaterSplineComponent.cpp b/Engine/Plugins/Experimental/Water/Source/Runtime/Private/WaterSplineComponent.cpp
--- a/Engine/Plugins/Experimental/Water/Source/Runtime/Private/WaterSplineComponent.cpp
+++ b/Engine/Plugins/Experimental/Water/Source/Runtime/Private/WaterSplineComponent.cpp
@@ -253,13 +253,20 @@ PRAGMA_ENABLE_DEPRECATION_WARNINGS
 				Result.Z = FMath::Clamp(Vec3.Z, Min, Max);
 				return Result;
 			};
-			FVector ArriveTangent = GetArriveTangentAtSplinePoint(Point, ESplineCoordinateSpace::Local);
-			FVector LeaveTangent = GetLeaveTangentAtSplinePoint(Point, ESplineCoordinateSpace::Local);
-			
+
 			constexpr double MaxTangentValue = 1.e10L;
-			ArriveTangent = ClampVec3(ArriveTangent, -MaxTangentValue, MaxTangentValue);
-			LeaveTangent = ClampVec3(LeaveTangent, -MaxTangentValue, MaxTangentValue);
-			SetTangentsAtSplinePoint(Point, ArriveTangent, LeaveTangent, ESplineCoordinateSpace::Local, false);
+
+			const FVector ArriveTangent = GetArriveTangentAtSplinePoint(Point, ESplineCoordinateSpace::Local);
+			const FVector LeaveTangent = GetLeaveTangentAtSplinePoint(Point, ESplineCoordinateSpace::Local);
+			
+			const FVector ClampedArriveTangent = ClampVec3(ArriveTangent, -MaxTangentValue, MaxTangentValue);
+			const FVector ClampedLeaveTangent = ClampVec3(LeaveTangent, -MaxTangentValue, MaxTangentValue);
+
+			// Only change if clamping is actually doing anything so that we don't unnecessarily force the point into custom tangent mode.
+			if (!ArriveTangent.Equals(ClampedArriveTangent) || !LeaveTangent.Equals(ClampedLeaveTangent))
+			{
+				SetTangentsAtSplinePoint(Point, ArriveTangent, LeaveTangent, ESplineCoordinateSpace::Local, false);
+			}
 		}
 	}
 
