[Water] Improved WaterInfoTexture depth precision by switching CRPs from DeviceDepth encoded in RGB8 to SceneDepth in R32F.

diff --git a/Engine/Plugins/Experimental/Water/Shaders/Private/WaterInfoMerge.usf b/Engine/Plugins/Experimental/Water/Shaders/Private/WaterInfoMerge.usf
--- a/Engine/Plugins/Experimental/Water/Shaders/Private/WaterInfoMerge.usf
+++ b/Engine/Plugins/Experimental/Water/Shaders/Private/WaterInfoMerge.usf
@@ -31,12 +31,6 @@ FWaterColorPassResult MakeWaterColorPassResult(float4 InSampleData)
 	return Result;
 }
 
-float DecodeFloatRGB(float3 Enc)
-{
-	float3 Encoding = float3(1.0, 1.0 / 255.0, 1.0 / 65025.0);
-	return 1 - dot(Enc, Encoding);
-}
-
 float NormalizeHeight(float Height, float Min, float Max)
 {
 	float NormalizedHeight = 1.f;
@@ -55,43 +49,34 @@ void Main(
 	)
 {
 	// All input textures MUST be sampled with a point sampler to avoid a driver issue on older Android devices which occurs when using Load instead.
-	float4 ColorSample = ColorTexture.SampleLevel(ColorTextureSampler, Input.UV, 0);
-	float4 DepthSample = DepthTexture.SampleLevel(DepthTextureSampler, Input.UV, 0);
-	float4 DilationSample = DilationTexture.SampleLevel(DilationTextureSampler, Input.UV, 0);
+	const float4 ColorSample = ColorTexture.SampleLevel(ColorTextureSampler, Input.UV, 0);
+	const float DepthSample = DepthTexture.SampleLevel(DepthTextureSampler, Input.UV, 0).x;
+	const float DilationSample = DilationTexture.SampleLevel(DilationTextureSampler, Input.UV, 0).x;
 	
 	const FWaterColorPassResult WaterInfo = MakeWaterColorPassResult(ColorSample);
-	const float GroundDeviceDepth = DecodeFloatRGB(DepthSample.rgb);
-	const float GroundSceneDepth = ConvertFromDeviceZ(GroundDeviceDepth);
-	float GroundZ = CaptureZ - GroundSceneDepth;
+	const float GroundSceneDepth = DepthSample;
+	const float GroundZ = max(CaptureZ - GroundSceneDepth, GroundZMin);
 	
 	const float WaterSceneDepth = WaterInfo.SceneDepth;
 	const float WaterZ = CaptureZ - WaterSceneDepth;
 
-	const float DilationDeviceDepth = DecodeFloatRGB(DilationSample.rgb);
-	const float DilationSceneDepth = ConvertFromDeviceZ(DilationDeviceDepth);
-	const bool bPixelHasDilationData = DilationSceneDepth != 0.0;
+	const float DilationSceneDepth = DilationSample;
+	const bool bPixelHasDilationData = DilationSceneDepth != 0.0f;
 		
 	float DilatedZ = CaptureZ - DilationSceneDepth;
 
-
 	// Always initialize output to avoid bugs in old shader compilers 
-	OutColor = 0.f;
-
-	// If the GroundDepth is further than the far plane of the projection, set the GroundZ to some really far away value
-	if (GroundDeviceDepth == 1.0)
-	{
-		GroundZ = GroundZMin;
-	}
+	OutColor = 0.0f;
 
 	const float WaterZMin = WaterHeightExtents.x;
 	const float WaterZMax = WaterHeightExtents.y;
 	
-	// if the measured WaterZ is below WaterZMin it means that there is no water at this pixel because it is not possible
+	// If the measured WaterZ is below WaterZMin it means that there is no water at this pixel because it is not possible
 	// for water to be further than ZMin.
 	bool bPixelHasValidWaterData = WaterZ >= WaterZMin;
 
 	// Additionally, if we have dilation data here and the water data would be placed below the landscape, instead draw the dilation data
-	if (bPixelHasDilationData && (WaterZ + 128.f < GroundZ))
+	if (bPixelHasDilationData && (WaterZ + 128.0f < GroundZ))
 	{
 		bPixelHasValidWaterData = false;
 	}
@@ -105,7 +90,6 @@ void Main(
 	}
 	else if (bPixelHasDilationData)
 	{
-
 		// If the dilated water would render on top of terrain, squash it down below
 		if (DilatedZ > GroundZ)
 		{
@@ -115,13 +99,12 @@ void Main(
 		const float NormalizedWaterZ = NormalizeHeight(DilatedZ, WaterZMin, WaterZMax);
 		const float NormalizedGroundZ = NormalizeHeight(GroundZ, GroundZMin, WaterZMax);
 
-		OutColor = float4(0.0, 0.0, NormalizedWaterZ, NormalizedGroundZ);
+		OutColor = float4(0.0f, 0.0f, NormalizedWaterZ, NormalizedGroundZ);
 	}
 	else
 	{
 		// Pixel may not have any dilation data nor water data. In this case just write 0s to the texture
-
 		const float NormalizedGroundZ = NormalizeHeight(GroundZ, GroundZMin, WaterZMax);
-		OutColor = float4(0.0, 0.0, 0.0, NormalizedGroundZ);
+		OutColor = float4(0.0f, 0.0f, 0.0f, NormalizedGroundZ);
 	}
 }
diff --git a/Engine/Plugins/Experimental/Water/Source/Runtime/Private/WaterInfoRendering.cpp b/Engine/Plugins/Experimental/Water/Source/Runtime/Private/WaterInfoRendering.cpp
--- a/Engine/Plugins/Experimental/Water/Source/Runtime/Private/WaterInfoRendering.cpp
+++ b/Engine/Plugins/Experimental/Water/Source/Runtime/Private/WaterInfoRendering.cpp
@@ -398,7 +398,7 @@ public:
 	IMPLEMENT_CUSTOM_RENDER_PASS(FWaterInfoRenderingDepthPass);
 
 	FWaterInfoRenderingDepthPass(const FIntPoint& InRenderTargetSize, const TMap<uint32, int32>& InLandscapeLODOverrides)
-		: FWaterInfoCustomRenderPassBase(TEXT("WaterInfoDepthPass"), FCustomRenderPassBase::ERenderMode::DepthPass, FCustomRenderPassBase::ERenderOutput::DeviceDepth, InRenderTargetSize)
+		: FWaterInfoCustomRenderPassBase(TEXT("WaterInfoDepthPass"), FCustomRenderPassBase::ERenderMode::DepthPass, FCustomRenderPassBase::ERenderOutput::SceneDepth, InRenderTargetSize)
 	{
 		if (!InLandscapeLODOverrides.IsEmpty())
 		{
@@ -408,7 +408,7 @@ public:
 
 	virtual void OnPreRender(FRDGBuilder& GraphBuilder) override
 	{
-		const FRDGTextureDesc TextureDesc = FRDGTextureDesc::Create2D(RenderTargetSize, PF_FloatRGBA, FClearValueBinding::Black, TexCreate_RenderTargetable | TexCreate_ShaderResource);
+		const FRDGTextureDesc TextureDesc = FRDGTextureDesc::Create2D(RenderTargetSize, PF_R32_FLOAT, FClearValueBinding::Black, TexCreate_RenderTargetable | TexCreate_ShaderResource);
 		RenderTargetTexture = GraphBuilder.CreateTexture(TextureDesc, TEXT("WaterDepthTexture"));
 		AddClearRenderTargetPass(GraphBuilder, RenderTargetTexture, FLinearColor::Black, Views[0]->UnscaledViewRect);
 	}
@@ -447,22 +447,24 @@ public:
 	IMPLEMENT_CUSTOM_RENDER_PASS(FWaterInfoRenderingDilationPass);
 
 	FWaterInfoRenderingDilationPass(const FIntPoint& InRenderTargetSize)
-		: FWaterInfoCustomRenderPassBase(TEXT("WaterInfoDilationPass"), FCustomRenderPassBase::ERenderMode::DepthPass, FCustomRenderPassBase::ERenderOutput::DeviceDepth, InRenderTargetSize)
+		: FWaterInfoCustomRenderPassBase(TEXT("WaterInfoDilationPass"), FCustomRenderPassBase::ERenderMode::DepthPass, FCustomRenderPassBase::ERenderOutput::SceneDepth, InRenderTargetSize)
 	{}
 
 	virtual void OnPreRender(FRDGBuilder& GraphBuilder) override
 	{
-		const FRDGTextureDesc TextureDesc = FRDGTextureDesc::Create2D(RenderTargetSize, WaterInfoRenderTarget->GetRenderTargetTexture()->GetDesc().Format, FClearValueBinding::Black, TexCreate_RenderTargetable | TexCreate_ShaderResource);
+		const FRDGTextureDesc TextureDesc = FRDGTextureDesc::Create2D(RenderTargetSize, PF_R32_FLOAT, FClearValueBinding::Black, TexCreate_RenderTargetable | TexCreate_ShaderResource);
 		RenderTargetTexture = GraphBuilder.CreateTexture(TextureDesc, TEXT("WaterDilationTexture"));
 		AddClearRenderTargetPass(GraphBuilder, RenderTargetTexture, FLinearColor::Black, Views[0]->UnscaledViewRect);
 	}
 	
 	virtual void OnPostRender(FRDGBuilder& GraphBuilder) override
 	{
-		FRDGTextureRef MergeTargetTexture = GraphBuilder.CreateTexture(RenderTargetTexture->Desc, TEXT("WaterInfoMerged"));
+		const FRDGTextureDesc ResultTextureDesc = FRDGTextureDesc::Create2D(RenderTargetSize, WaterInfoRenderTarget->GetRenderTargetTexture()->GetDesc().Format, FClearValueBinding::Black, TexCreate_RenderTargetable | TexCreate_ShaderResource);
+		
+		FRDGTextureRef MergeTargetTexture = GraphBuilder.CreateTexture(ResultTextureDesc, TEXT("WaterInfoMerged"));
 		MergeWaterInfoAndDepth(GraphBuilder, *Views[0]->Family, *Views[0], MergeTargetTexture, DepthPass->GetRenderTargetTexture(), ColorPass->GetRenderTargetTexture(), RenderTargetTexture, Params);
 	
-		FRDGTextureRef FinalizedTexture = GraphBuilder.CreateTexture(RenderTargetTexture->Desc, TEXT("WaterInfoFinalized"));
+		FRDGTextureRef FinalizedTexture = GraphBuilder.CreateTexture(ResultTextureDesc, TEXT("WaterInfoFinalized"));
 		FinalizeWaterInfo(GraphBuilder, *Views[0]->Family, *Views[0], MergeTargetTexture, FinalizedTexture, Params);
 
 		FRDGTextureRef WaterInfoTexture = RegisterExternalTexture(GraphBuilder, WaterInfoRenderTarget->GetRenderTargetTexture(), TEXT("WaterInfoTexture"));
