[Water]

Fixed water spline metadata properties not being commited as ValueSet operations when the user commits by dropping the text box focus.

diff --git a/Engine/Plugins/Experimental/Water/Source/Editor/Private/WaterSplineMetadataDetails.cpp b/Engine/Plugins/Experimental/Water/Source/Editor/Private/WaterSplineMetadataDetails.cpp
--- a/Engine/Plugins/Experimental/Water/Source/Editor/Private/WaterSplineMetadataDetails.cpp
+++ b/Engine/Plugins/Experimental/Water/Source/Editor/Private/WaterSplineMetadataDetails.cpp
@@ -95,6 +95,14 @@ void FWaterSplineMetadataDetails::Update(USplineComponent* InSplineComponent, co
 	}
 }
 
+namespace WaterSplineMetadataDetailsLocals 
+{
+	static bool IsTextCommitValueSet(ETextCommit::Type CommitInfo)
+	{
+		return CommitInfo == ETextCommit::OnUserMovedFocus || CommitInfo == ETextCommit::OnEnter;
+	}
+}
+
 template<class T>
 void SetValues(FWaterSplineMetadataDetails& Details, TArray<FInterpCurvePoint<T>>& Points, const T& NewValue, ETextCommit::Type CommitInfo, const FText& TransactionText)
 {
@@ -111,7 +119,7 @@ void SetValues(FWaterSplineMetadataDetails& Details, TArray<FInterpCurvePoint<T>
 	Details.SplineComp->UpdateSpline();
 	Details.SplineComp->bSplineHasBeenEdited = true;
 	static FProperty* SplineCurvesProperty = FindFProperty<FProperty>(USplineComponent::StaticClass(), USplineComponent::GetSplinePropertyName());
-	EPropertyChangeType::Type PropertyChangeType = CommitInfo == ETextCommit::OnEnter ? EPropertyChangeType::ValueSet : EPropertyChangeType::Interactive;
+	EPropertyChangeType::Type PropertyChangeType = WaterSplineMetadataDetailsLocals::IsTextCommitValueSet(CommitInfo) ? EPropertyChangeType::ValueSet : EPropertyChangeType::Interactive;
 	FComponentVisualizer::NotifyPropertyModified(Details.SplineComp, SplineCurvesProperty, PropertyChangeType);
 	Details.Update(Details.SplineComp, Details.SelectedKeys);
 
@@ -150,7 +158,7 @@ void FWaterSplineMetadataDetails::OnSetDepth(float NewValue, ETextCommit::Type C
 {
 	if (UWaterSplineMetadata* Metadata = GetMetadata())
 	{
-		if (CheckIfDifferent(*this, Metadata->Depth.Points, NewValue) || CommitInfo == ETextCommit::OnEnter)
+		if (CheckIfDifferent(*this, Metadata->Depth.Points, NewValue) || WaterSplineMetadataDetailsLocals::IsTextCommitValueSet(CommitInfo))
 		{
 			SetValues(*this, Metadata->Depth.Points, NewValue, CommitInfo, LOCTEXT("SetSplineDepth", "Set spline point water depth"));
 		}
@@ -161,7 +169,7 @@ void FWaterSplineMetadataDetails::OnSetRiverWidth(float NewValue, ETextCommit::T
 {
 	if (UWaterSplineMetadata* Metadata = GetMetadata())
 	{
-		if (CheckIfDifferent(*this, Metadata->RiverWidth.Points, NewValue) || CommitInfo == ETextCommit::OnEnter)
+		if (CheckIfDifferent(*this, Metadata->RiverWidth.Points, NewValue) || WaterSplineMetadataDetailsLocals::IsTextCommitValueSet(CommitInfo))
 		{
 			SetValues(*this, Metadata->RiverWidth.Points, NewValue, CommitInfo, LOCTEXT("SetSplineWaterWidth", "Set spline point river width"));
 		}
@@ -172,7 +180,7 @@ void FWaterSplineMetadataDetails::OnSetVelocity(float NewValue, ETextCommit::Typ
 {
 	if (UWaterSplineMetadata* Metadata = GetMetadata())
 	{
-		if (CheckIfDifferent(*this, Metadata->WaterVelocityScalar.Points, NewValue) || CommitInfo == ETextCommit::OnEnter)
+		if (CheckIfDifferent(*this, Metadata->WaterVelocityScalar.Points, NewValue) || WaterSplineMetadataDetailsLocals::IsTextCommitValueSet(CommitInfo))
 		{
 			SetValues(*this, Metadata->WaterVelocityScalar.Points, NewValue, CommitInfo, LOCTEXT("SetSplineWaterVelocity", "Set spline point water velocity"));
 		}
@@ -183,7 +191,7 @@ void FWaterSplineMetadataDetails::OnSetAudioIntensity(float NewValue, ETextCommi
 {
 	if (UWaterSplineMetadata* Metadata = GetMetadata())
 	{
-		if (CheckIfDifferent(*this, Metadata->AudioIntensity.Points, NewValue) || CommitInfo == ETextCommit::OnEnter)
+		if (CheckIfDifferent(*this, Metadata->AudioIntensity.Points, NewValue) || WaterSplineMetadataDetailsLocals::IsTextCommitValueSet(CommitInfo))
 		{
 			SetValues(*this, Metadata->AudioIntensity.Points, NewValue, CommitInfo, LOCTEXT("SetSpline point audio intensity", "Set spline point audio intensity"));
 		}
