Removed 5.5 deprecated water code

#rnx

diff --git a/Engine/Plugins/Experimental/Water/Source/Editor/Public/WaterEditorSettings.h b/Engine/Plugins/Experimental/Water/Source/Editor/Public/WaterEditorSettings.h
--- a/Engine/Plugins/Experimental/Water/Source/Editor/Public/WaterEditorSettings.h
+++ b/Engine/Plugins/Experimental/Water/Source/Editor/Public/WaterEditorSettings.h
@@ -234,11 +234,6 @@ public:
 
 	bool GetShouldUpdateWaterMeshDuringInteractiveChanges() const { return bUpdateWaterMeshDuringInteractiveChanges; }
 
-	UE_DEPRECATED(5.5, "Use GetShouldUpdateLandscapeDuringInteractiveChanges")
-	bool GetUpdateLandscapeDuringInteractiveChanges() const { return false; }
-	UE_DEPRECATED(5.5, "Use GetShouldUpdateWaterMeshDuringInteractiveChanges")
-	bool GetUpdateWaterMeshDuringInteractiveChanges() const { return bUpdateWaterMeshDuringInteractiveChanges; }
-	
 public:
 	/** The texture group to use for generated textures such as the combined velocity and height texture */
 	UPROPERTY(EditAnywhere, config, Category = Rendering)
diff --git a/Engine/Plugins/Experimental/Water/Source/Runtime/Private/WaterBodyComponent.cpp b/Engine/Plugins/Experimental/Water/Source/Runtime/Private/WaterBodyComponent.cpp
--- a/Engine/Plugins/Experimental/Water/Source/Runtime/Private/WaterBodyComponent.cpp
+++ b/Engine/Plugins/Experimental/Water/Source/Runtime/Private/WaterBodyComponent.cpp
@@ -1130,12 +1130,6 @@ ALandscapeProxy* UWaterBodyComponent::FindLandscape() const
 	return Landscape.Get();
 }
 
-// Deprecated
-void UWaterBodyComponent::UpdateComponentVisibility(bool bAllowWaterZoneRebuild)
-{
-	UpdateVisibility();
-}
-
 void UWaterBodyComponent::UpdateVisibility()
 {
 	if (UWorld* World = GetWorld())
diff --git a/Engine/Plugins/Experimental/Water/Source/Runtime/Private/WaterZoneActor.cpp b/Engine/Plugins/Experimental/Water/Source/Runtime/Private/WaterZoneActor.cpp
--- a/Engine/Plugins/Experimental/Water/Source/Runtime/Private/WaterZoneActor.cpp
+++ b/Engine/Plugins/Experimental/Water/Source/Runtime/Private/WaterZoneActor.cpp
@@ -838,36 +838,6 @@ bool AWaterZone::GetDynamicWaterInfoCenter(int32 PlayerIndex, FVector& OutCenter
 	return bHasValidZoneLocation;
 }
 
-bool AWaterZone::GetDynamicWaterInfoBounds(int32 PlayerIndex, FBox& OutBounds) const
-{
-	const UWorld* World = GetWorld();
-
-	check(World != nullptr);
-
-	bool bHasValidZoneLocation = false;
-	FVector Center = GetActorLocation();
-
-	if (const FWaterViewExtension* WaterViewExtension = UWaterSubsystem::GetWaterViewExtension(World))
-	{
-		bHasValidZoneLocation = WaterViewExtension->GetZoneLocation(this, PlayerIndex, Center);
-	}
-	
-	const FVector WaterInfoHalfExtents = GetDynamicWaterInfoExtent() / 2.;
-	OutBounds = FBox(Center - WaterInfoHalfExtents, Center + WaterInfoHalfExtents);
-
-	return bHasValidZoneLocation;
-}
-
-FVector AWaterZone::GetDynamicWaterInfoCenter() const
-{
-	// The following index assume that there is no split screen and will request the position of the first player's water view.
-	// This call is a temporary fallback and is deprecated in favor of the overload which takes a specific player view.
-	constexpr int32 PlayerIndex = 0;
-	FVector Center;
-	GetDynamicWaterInfoCenter(PlayerIndex, Center);
-	return Center;
-}
-
 void AWaterZone::OnLevelAddedToWorld(ULevel* InLevel, UWorld* InWorld)
 {
 	OnLevelChanged(InLevel, InWorld);
diff --git a/Engine/Plugins/Experimental/Water/Source/Runtime/Public/WaterBodyComponent.h b/Engine/Plugins/Experimental/Water/Source/Runtime/Public/WaterBodyComponent.h
--- a/Engine/Plugins/Experimental/Water/Source/Runtime/Public/WaterBodyComponent.h
+++ b/Engine/Plugins/Experimental/Water/Source/Runtime/Public/WaterBodyComponent.h
@@ -373,9 +373,6 @@ public:
 	/** Returns true if the location is within one of this water body's exclusion volumes */
 	WATER_API bool IsWorldLocationInExclusionVolume(const FVector& InWorldLocation) const;
 
-	UE_DEPRECATED(5.5, "Use UpdateComponentVisibility")
-	WATER_API virtual void UpdateComponentVisibility(bool bAllowWaterZoneRebuild);
-
 	/** Updates the bVisible/bHiddenInGame flags on the component and eventually the child renderable components (e.g. custom water body) */
 	WATER_API void UpdateVisibility();
 
@@ -697,7 +694,7 @@ protected:
 	UPROPERTY(Category = BakedSimulation, AdvancedDisplay, VisibleAnywhere)
 	TWeakObjectPtr<UBakedShallowWaterSimulationComponent> BakedShallowWaterSim;
 
-	/**  Override to disable use of the baked shallow water simulation for collisons and other uses */
+	/**  Override to disable use of the baked shallow water simulation for collisions and other uses */
 	UPROPERTY(Category = BakedSimulation, AdvancedDisplay, EditAnywhere)
 	bool bUseBakedSimForQueriesAndPhysics = true;
 
diff --git a/Engine/Plugins/Experimental/Water/Source/Runtime/Public/WaterZoneActor.h b/Engine/Plugins/Experimental/Water/Source/Runtime/Public/WaterZoneActor.h
--- a/Engine/Plugins/Experimental/Water/Source/Runtime/Public/WaterZoneActor.h
+++ b/Engine/Plugins/Experimental/Water/Source/Runtime/Public/WaterZoneActor.h
@@ -29,7 +29,6 @@ enum class EWaterZoneRebuildFlags
 ENUM_CLASS_FLAGS(EWaterZoneRebuildFlags);
 
 
-DECLARE_DYNAMIC_MULTICAST_DELEGATE_OneParam(FOnWaterInfoTextureCreated, const UTextureRenderTarget2D*, WaterInfoTexture);
 DECLARE_DYNAMIC_MULTICAST_DELEGATE_OneParam(FOnWaterInfoTextureArrayCreated, const UTextureRenderTarget2DArray*, WaterInfoTextureArray);
 
 UCLASS(MinimalAPI, Blueprintable, HideCategories=(Physics, Replication, Input, Collision))
@@ -78,15 +77,6 @@ public:
 
 	bool IsLocalOnlyTessellationEnabled() const { return bEnableLocalOnlyTessellation; }
 
-	UE_DEPRECATED(5.5, "In 5.5 the dynamic water info is now per player view. Either call GetAllDynamicWaterInfoCenters or use the per-player index version")
-	UE_API FVector GetDynamicWaterInfoCenter() const;
-
-	UE_DEPRECATED(5.5, "In 5.5 the dynamic water info is now per player view. Either call GetAllDynamicWaterInfoCenters or use the per-player index version")
-	UE_API FBox GetDynamicWaterInfoBounds() const;
-
-	UE_DEPRECATED(5.5, "It is no longer possible to manually set the local tessellation center. This is controlled per view by the water view extension.")
-	void SetLocalTessellationCenter(const FVector& NewCenter) {}
-
 	UE_API virtual void BeginPlay() override;
 	UE_API virtual void EndPlay(const EEndPlayReason::Type EndPlayReason) override;
 	UE_API virtual void PostLoadSubobjects(FObjectInstancingGraph* OuterInstanceGraph) override;
@@ -121,18 +111,6 @@ public:
 	UFUNCTION(BlueprintCallable, Category=Rendering)
 	UE_API void SetFarMeshMaterial(UMaterialInterface* InFarMaterial);
 
-
-#if WITH_EDITORONLY_DATA
-	UE_DEPRECATED(5.5, "WaterInfoTexture is deprecated, please use WaterInfoTextureArray instead.")
-	UPROPERTY(Transient, meta = (DeprecatedProperty, DeprecationMessage = "Use WaterInfoTextureArray instead."))
-	TObjectPtr<UTextureRenderTarget2D> WaterInfoTexture_DEPRECATED;
-#endif
-
-#if WITH_EDITOR 
-	UE_DEPRECATED(5.5, "GetOnWaterInfoTextureCreated is deprecated, please use GetOnWaterInfoTextureArrayCreated instead.")
-	FOnWaterInfoTextureCreated& GetOnWaterInfoTextureCreated() { return OnWaterInfoTextureCreated_DEPRECATED; }
-#endif
-
 private:
 
 	/**
@@ -262,9 +240,6 @@ private:
 	UPROPERTY(Transient)
 	TObjectPtr<UTexture2D> WaterVelocityTexture_DEPRECATED;
 
-	UPROPERTY()
-	FOnWaterInfoTextureCreated OnWaterInfoTextureCreated_DEPRECATED;
-
 	UPROPERTY()
 	FVector TessellatedWaterMeshExtent_DEPRECATED;
 
