Changed QueryWaterInfoClosestToWorldLocation to return TResult so we can capture invalid queries.

diff --git a/Engine/Plugins/Experimental/Mover/Source/Mover/Private/MoveLibrary/WaterMovementUtils.cpp b/Engine/Plugins/Experimental/Mover/Source/Mover/Private/MoveLibrary/WaterMovementUtils.cpp
--- a/Engine/Plugins/Experimental/Mover/Source/Mover/Private/MoveLibrary/WaterMovementUtils.cpp
+++ b/Engine/Plugins/Experimental/Mover/Source/Mover/Private/MoveLibrary/WaterMovementUtils.cpp
@@ -38,26 +38,33 @@ void UWaterMovementUtils::UpdateWaterSplineData(const FUpdateWaterSplineDataPara
 		{
 			if (WaterBody->GetWaterBodyComponent())
 			{
-				FWaterBodyQueryResult QueryResult = WaterBody->GetWaterBodyComponent()->QueryWaterInfoClosestToWorldLocation(UpdateWaterSplineDataParams.PlayerLocation,
+				const TValueOrError<FWaterBodyQueryResult, EWaterBodyQueryError> QueryResult = WaterBody->GetWaterBodyComponent()->TryQueryWaterInfoClosestToWorldLocation(UpdateWaterSplineDataParams.PlayerLocation,
 						   EWaterBodyQueryFlags::ComputeLocation
 						   | EWaterBodyQueryFlags::ComputeDepth
 						   | EWaterBodyQueryFlags::ComputeNormal
 						   | EWaterBodyQueryFlags::ComputeVelocity
 						   | EWaterBodyQueryFlags::ComputeImmersionDepth);
+
+				if (!QueryResult.HasValue())
+				{
+					return;
+				}
+
+				const FWaterBodyQueryResult& Query = QueryResult.GetValue();
 		
 				FWaterFlowSplineData& WaterSplineData = WaterCheckResult.WaterSplineData;
 		
 				// Immersion depth
-				WaterSplineData.ImmersionDepth = QueryResult.GetImmersionDepth();
+				WaterSplineData.ImmersionDepth = Query.GetImmersionDepth();
 		
-				WaterSplineData.WaterPlaneLocation = QueryResult.GetWaterPlaneLocation();
-				WaterSplineData.WaterPlaneNormal = QueryResult.GetWaterPlaneNormal();
+				WaterSplineData.WaterPlaneLocation = Query.GetWaterPlaneLocation();
+				WaterSplineData.WaterPlaneNormal = Query.GetWaterPlaneNormal();
 		
 				// Water Depth
-				WaterSplineData.WaterDepth = QueryResult.GetWaterSurfaceDepth();
+				WaterSplineData.WaterDepth = Query.GetWaterSurfaceDepth();
 		
 				// Water velocity
-				WaterSplineData.RawWaterVelocity = QueryResult.GetVelocity();
+				WaterSplineData.RawWaterVelocity = Query.GetVelocity();
 		
 				// Water velocity modified by depth.
 				const FVector2D DepthRange = FVector2D(UpdateWaterSplineDataParams.TargetImmersionDepth, UpdateWaterSplineDataParams.WaterVelocityDepthForMax);
@@ -74,7 +81,7 @@ void UWaterMovementUtils::UpdateWaterSplineData(const FUpdateWaterSplineDataPara
 				// Surface location
 				const float CapsuleBottom = UpdateWaterSplineDataParams.PlayerLocation.Z - UpdateWaterSplineDataParams.CapsuleHalfHeight;
 				const float CapsuleTop = UpdateWaterSplineDataParams.PlayerLocation.Z + UpdateWaterSplineDataParams.CapsuleHalfHeight;
-				WaterSplineData.WaterSurfaceLocation = QueryResult.GetWaterSurfaceLocation();
+				WaterSplineData.WaterSurfaceLocation = Query.GetWaterSurfaceLocation();
 		
 				// Surface offset
 				WaterSplineData.WaterSurfaceOffset = WaterSplineData.WaterSurfaceLocation - UpdateWaterSplineDataParams.PlayerLocation;
@@ -83,7 +90,7 @@ void UWaterMovementUtils::UpdateWaterSplineData(const FUpdateWaterSplineDataPara
 				WaterSplineData.ImmersionPercent = FMath::Clamp((WaterSplineData.WaterSurfaceLocation.Z - CapsuleBottom) / (CapsuleTop - CapsuleBottom), 0.f, 1.f);
 		
 				// Normal
-				WaterSplineData.WaterSurfaceNormal = QueryResult.GetWaterSurfaceNormal();
+				WaterSplineData.WaterSurfaceNormal = Query.GetWaterSurfaceNormal();
 			}
 		}
 	}
diff --git a/Engine/Plugins/Experimental/Mover/Source/Mover/Private/PhysicsMover/PhysicsMovementUtils.cpp b/Engine/Plugins/Experimental/Mover/Source/Mover/Private/PhysicsMover/PhysicsMovementUtils.cpp
--- a/Engine/Plugins/Experimental/Mover/Source/Mover/Private/PhysicsMover/PhysicsMovementUtils.cpp
+++ b/Engine/Plugins/Experimental/Mover/Source/Mover/Private/PhysicsMover/PhysicsMovementUtils.cpp
@@ -200,26 +200,33 @@ bool UPhysicsMovementUtils::GetWaterResultFromHitResults(const TArray<FHitResult
 
 					OutWaterResult.WaterSplineData.WaterBody = WaterBody;
 
-					FWaterBodyQueryResult QueryResult = WaterBody->GetWaterBodyComponent()->QueryWaterInfoClosestToWorldLocation(
+					TValueOrError<FWaterBodyQueryResult, EWaterBodyQueryError> QueryResult = WaterBody->GetWaterBodyComponent()->TryQueryWaterInfoClosestToWorldLocation(
 							Location, 
 							EWaterBodyQueryFlags::ComputeLocation | EWaterBodyQueryFlags::ComputeNormal | EWaterBodyQueryFlags::ComputeImmersionDepth,
 							OutWaterResult.WaterSplineData.SplineInputKey
 						);
 
-					OutWaterResult.WaterSplineData.ImmersionDepth = QueryResult.GetImmersionDepth();
+					if (!QueryResult.HasValue())
+					{
+						continue;
+					}
 
-					OutWaterResult.WaterSplineData.WaterPlaneLocation = QueryResult.GetWaterPlaneLocation();
-					OutWaterResult.WaterSplineData.WaterPlaneNormal = QueryResult.GetWaterPlaneNormal();
+					const FWaterBodyQueryResult& Query = QueryResult.GetValue();
+
+					OutWaterResult.WaterSplineData.ImmersionDepth = Query.GetImmersionDepth();
+
+					OutWaterResult.WaterSplineData.WaterPlaneLocation = Query.GetWaterPlaneLocation();
+					OutWaterResult.WaterSplineData.WaterPlaneNormal = Query.GetWaterPlaneNormal();
 
 					const float CapsuleBottom = Location.Z;
 					const float CapsuleTop = Location.Z + (TargetHeight * 2);
-					OutWaterResult.WaterSplineData.WaterSurfaceLocation = QueryResult.GetWaterSurfaceLocation();
+					OutWaterResult.WaterSplineData.WaterSurfaceLocation = Query.GetWaterSurfaceLocation();
 			
 					OutWaterResult.WaterSplineData.WaterSurfaceOffset = OutWaterResult.WaterSplineData.WaterSurfaceLocation - Location;
 			
 					OutWaterResult.WaterSplineData.ImmersionPercent = FMath::Clamp((OutWaterResult.WaterSplineData.WaterSurfaceLocation.Z - CapsuleBottom) / (CapsuleTop - CapsuleBottom), 0.f, 1.f);
 			
-					OutWaterResult.WaterSplineData.WaterSurfaceNormal = QueryResult.GetWaterSurfaceNormal();
+					OutWaterResult.WaterSplineData.WaterSurfaceNormal = Query.GetWaterSurfaceNormal();
 
 #if PHYSICSDRIVENMOTION_DEBUG_DRAW
 					Chaos::FDebugDrawQueue::GetInstance().DrawDebugLine(Location, Location - FVector::UpVector * OutWaterResult.WaterSplineData.ImmersionDepth, FColor::Blue, false, -1.f, 10, 1.0f);
diff --git a/Engine/Plugins/Experimental/Water/Source/Runtime/Private/BuoyancyComponent.cpp b/Engine/Plugins/Experimental/Water/Source/Runtime/Private/BuoyancyComponent.cpp
--- a/Engine/Plugins/Experimental/Water/Source/Runtime/Private/BuoyancyComponent.cpp
+++ b/Engine/Plugins/Experimental/Water/Source/Runtime/Private/BuoyancyComponent.cpp
@@ -632,22 +632,26 @@ float UBuoyancyComponent::GetWaterHeight(FVector Position, const TMap<const UWat
 				QueryFlags |= EWaterBodyQueryFlags::IncludeWaves;
 			}
 
-			FWaterBodyQueryResult QueryResult = CurrentWaterBodyComponent->QueryWaterInfoClosestToWorldLocation(Position, QueryFlags, SplineInputKey);
-			if (QueryResult.IsInWater() && QueryResult.GetImmersionDepth() > MaxImmersionDepth)
+			TValueOrError<FWaterBodyQueryResult, EWaterBodyQueryError> QueryResult = CurrentWaterBodyComponent->TryQueryWaterInfoClosestToWorldLocation(Position, QueryFlags, SplineInputKey);
+			if (QueryResult.HasValue())
 			{
-				check(!QueryResult.IsInExclusionVolume());
-				WaterHeight = Position.Z + QueryResult.GetImmersionDepth();
-				OutWaterBodyComponent = CurrentWaterBodyComponent;
-				if (EnumHasAnyFlags(QueryResult.GetQueryFlags(), EWaterBodyQueryFlags::ComputeDepth))
+				const FWaterBodyQueryResult& Query = QueryResult.GetValue();
+				if (Query.IsInWater() && Query.GetImmersionDepth() > MaxImmersionDepth)
 				{
-					OutWaterDepth = QueryResult.GetWaterSurfaceDepth();
+					check(!Query.IsInExclusionVolume());
+					WaterHeight = Position.Z + Query.GetImmersionDepth();
+					OutWaterBodyComponent = CurrentWaterBodyComponent;
+					if (EnumHasAnyFlags(Query.GetQueryFlags(), EWaterBodyQueryFlags::ComputeDepth))
+					{
+						OutWaterDepth = Query.GetWaterSurfaceDepth();
+					}
+					OutWaterPlaneLocation = Query.GetWaterPlaneLocation();
+					OutWaterPlaneNormal = Query.GetWaterPlaneNormal();
+					OutWaterSurfacePosition = Query.GetWaterSurfaceLocation();
+					OutWaterVelocity = Query.GetVelocity();
+					OutWaterBodyIdx = CurrentWaterBodyComponent ? CurrentWaterBodyComponent->GetWaterBodyIndex() : 0;
+					MaxImmersionDepth = Query.GetImmersionDepth();
 				}
-				OutWaterPlaneLocation = QueryResult.GetWaterPlaneLocation();
-				OutWaterPlaneNormal = QueryResult.GetWaterPlaneNormal();
-				OutWaterSurfacePosition = QueryResult.GetWaterSurfaceLocation();
-				OutWaterVelocity = QueryResult.GetVelocity();
-				OutWaterBodyIdx = CurrentWaterBodyComponent ? CurrentWaterBodyComponent->GetWaterBodyIndex() : 0;
-				MaxImmersionDepth = QueryResult.GetImmersionDepth();
 			}
 		}
 	}
diff --git a/Engine/Plugins/Experimental/Water/Source/Runtime/Private/EnvQueryTest_InsideWaterBody.cpp b/Engine/Plugins/Experimental/Water/Source/Runtime/Private/EnvQueryTest_InsideWaterBody.cpp
--- a/Engine/Plugins/Experimental/Water/Source/Runtime/Private/EnvQueryTest_InsideWaterBody.cpp
+++ b/Engine/Plugins/Experimental/Water/Source/Runtime/Private/EnvQueryTest_InsideWaterBody.cpp
@@ -43,12 +43,12 @@ void UEnvQueryTest_InsideWaterBody::RunTest(FEnvQueryInstance& QueryInstance) co
 				QueryFlags |= EWaterBodyQueryFlags::IgnoreExclusionVolumes;
 			}
 
-			const FWaterBodyQueryResult QueryResult = WaterBodyComponent->QueryWaterInfoClosestToWorldLocation(ItemLocation, QueryFlags);
-			if (QueryResult.IsInWater())
+			const TValueOrError<FWaterBodyQueryResult, EWaterBodyQueryError> QueryResult = WaterBodyComponent->TryQueryWaterInfoClosestToWorldLocation(ItemLocation, QueryFlags);
+			if (QueryResult.HasValue() && QueryResult.GetValue().IsInWater())
 			{
 				bInside = true;
 				return false;
-		}
+			}
 
 			return true;
 		});
diff --git a/Engine/Plugins/Experimental/Water/Source/Runtime/Private/NiagaraDataInterfaceWater.cpp b/Engine/Plugins/Experimental/Water/Source/Runtime/Private/NiagaraDataInterfaceWater.cpp
--- a/Engine/Plugins/Experimental/Water/Source/Runtime/Private/NiagaraDataInterfaceWater.cpp
+++ b/Engine/Plugins/Experimental/Water/Source/Runtime/Private/NiagaraDataInterfaceWater.cpp
@@ -62,13 +62,13 @@ namespace NDIWaterPrivate
 				continue;
 			}
 
-			const FWaterBodyQueryResult CurrentQueryResult = WaterComponent->QueryWaterInfoClosestToWorldLocation(QueryLocation, QueryFlags);
-			if (CurrentQueryResult.IsInExclusionVolume())
+			const TValueOrError<FWaterBodyQueryResult, EWaterBodyQueryError> CurrentQueryResult = WaterComponent->TryQueryWaterInfoClosestToWorldLocation(QueryLocation, QueryFlags);
+			if (!CurrentQueryResult.HasValue() || CurrentQueryResult.GetValue().IsInExclusionVolume())
 			{
 				continue;
 			}
 
-			const FVector WaterLocation = CurrentQueryResult.GetWaterPlaneLocation();
+			const FVector WaterLocation = CurrentQueryResult.GetValue().GetWaterPlaneLocation();
 			const double DistanceSq = (WaterLocation - QueryLocation).SquaredLength();
 			if (!ClosestComponent || DistanceSq < ClosestWaterSq)
 			{
@@ -131,23 +131,33 @@ namespace NDIWaterPrivate
 					(bIncludeWaves ? EWaterBodyQueryFlags::IncludeWaves : EWaterBodyQueryFlags::None) |
 					(bIncludeWaves && bSimpleWaves ? EWaterBodyQueryFlags::SimpleWaves : EWaterBodyQueryFlags::None);
 
-				const FWaterBodyQueryResult QueryResult = Component->QueryWaterInfoClosestToWorldLocation(InstData->LWCConverter.ConvertSimulationPositionToWorld(QueryPosition), QueryFlags);
+				const TValueOrError<FWaterBodyQueryResult, EWaterBodyQueryError> QueryResult = Component->TryQueryWaterInfoClosestToWorldLocation(InstData->LWCConverter.ConvertSimulationPositionToWorld(QueryPosition), QueryFlags);
+				if (QueryResult.HasError())
+				{
+					UE_LOG(LogWater, Error, TEXT("WaterInfoQuery returned error :%s"), *UEnum::GetValueAsString(QueryResult.GetError()));
+				}
+
+				if (!QueryResult.HasValue())
+				{
+					continue;
+				}
+				const FWaterBodyQueryResult& Query = QueryResult.GetValue();
 
-				const FVector3f WaterPlaneLocation = InstData->LWCConverter.ConvertWorldToSimulationPosition(QueryResult.GetWaterPlaneLocation());
-				const FVector3f WaterSurfacePosition = InstData->LWCConverter.ConvertWorldToSimulationPosition(QueryResult.GetWaterSurfaceLocation());
+				const FVector3f WaterPlaneLocation = InstData->LWCConverter.ConvertWorldToSimulationPosition(Query .GetWaterPlaneLocation());
+				const FVector3f WaterSurfacePosition = InstData->LWCConverter.ConvertWorldToSimulationPosition(Query .GetWaterSurfaceLocation());
 
 				float DepthValue = bIncludeDepth ? InstData->SystemInstanceWaterDepth : 0.0f;
 				if (bDoDepthQuery)
 				{
-					DepthValue = QueryResult.GetWaterSurfaceDepth();
+					DepthValue = Query .GetWaterSurfaceDepth();
 				}
 
 				OutWaterPlanePosition.SetAndAdvance(WaterPlaneLocation);
-				OutWaterPlaneNormal.SetAndAdvance(FVector3f(QueryResult.GetWaterPlaneNormal()));
+				OutWaterPlaneNormal.SetAndAdvance(FVector3f(Query .GetWaterPlaneNormal()));
 				OutWaterSurfacePosition.SetAndAdvance(WaterSurfacePosition);
 				OutWaterDepth.SetAndAdvance(DepthValue);
-				OutWaterVelocity.SetAndAdvance(FVector3f(QueryResult.GetVelocity()));
-				OutInExclusionVolume.SetAndAdvance(QueryResult.IsInExclusionVolume());
+				OutWaterVelocity.SetAndAdvance(FVector3f(Query .GetVelocity()));
+				OutInExclusionVolume.SetAndAdvance(Query .IsInExclusionVolume());
 			}
 			else
 			{
@@ -380,9 +390,17 @@ bool UNiagaraDataInterfaceWater::PerInstanceTick(void* PerInstanceData, FNiagara
 		{
 			const FVector QueryLocation = SystemInstance->GetWorldTransform().GetTranslation();
 			const EWaterBodyQueryFlags QueryFlags = EWaterBodyQueryFlags::ComputeDepth;
-			const FWaterBodyQueryResult QueryResult = WaterBodyComponent->QueryWaterInfoClosestToWorldLocation(QueryLocation, QueryFlags);
 
-			InstData->SystemInstanceWaterDepth = QueryResult.IsInExclusionVolume() ? 0.0f : QueryResult.GetWaterSurfaceDepth();
+			const TValueOrError<FWaterBodyQueryResult, EWaterBodyQueryError> QueryResult = WaterBodyComponent->TryQueryWaterInfoClosestToWorldLocation(QueryLocation, QueryFlags);
+			if (QueryResult.HasValue())
+			{
+				InstData->SystemInstanceWaterDepth = QueryResult.GetValue().IsInExclusionVolume() ? 0.0f : QueryResult.GetValue().GetWaterSurfaceDepth();
+			}
+			else if (QueryResult.HasError())
+			{
+				UE_LOG(LogWater, Error, TEXT("NiagaraWaterDataInterface: attempting to compute the water body depth returned error: %s"), *UEnum::GetValueAsString(QueryResult.GetError()));
+			}
+
 		}
 	}
 
@@ -462,7 +480,18 @@ void UNiagaraDataInterfaceWater::GetWaterDataAtPoint(FVectorVMExternalFunctionCo
 				QueryFlags |= EWaterBodyQueryFlags::ComputeDepth;
 			}
 
-			QueryResult = Component->QueryWaterInfoClosestToWorldLocation(QueryPos, QueryFlags);
+			TValueOrError<FWaterBodyQueryResult, EWaterBodyQueryError> Query = Component->TryQueryWaterInfoClosestToWorldLocation(QueryPos, QueryFlags);
+			if (Query.HasError())
+			{
+				UE_LOG(LogWater, Error, TEXT("WaterInfoQuery returned error :%s"), *UEnum::GetValueAsString(Query.GetError()));
+			}
+
+			if (!Query.HasValue())
+			{
+				continue;
+			}
+
+			QueryResult = Query.GetValue();
 
 			bIsValid = !QueryResult.IsInExclusionVolume();
 		}
diff --git a/Engine/Plugins/Experimental/Water/Source/Runtime/Private/WaterBodyComponent.cpp b/Engine/Plugins/Experimental/Water/Source/Runtime/Private/WaterBodyComponent.cpp
--- a/Engine/Plugins/Experimental/Water/Source/Runtime/Private/WaterBodyComponent.cpp
+++ b/Engine/Plugins/Experimental/Water/Source/Runtime/Private/WaterBodyComponent.cpp
@@ -619,6 +619,25 @@ bool UWaterBodyComponent::IsWorldLocationInExclusionVolume(const FVector& InWorl
 	return false;
 }
 
+TValueOrError<FWaterBodyQueryResult, EWaterBodyQueryError> UWaterBodyComponent::TryQueryWaterInfoClosestToWorldLocation(const FVector& InWorldLocation, EWaterBodyQueryFlags InQueryFlags, const TOptional<float>& InSplineInputKey) const
+{
+	if (!IsRegistered())
+	{
+		return MakeError(EWaterBodyQueryError::WaterBodyNotRegistered);
+	}
+
+	if (WaterSplineMetadata == nullptr)
+	{
+		return MakeError(EWaterBodyQueryError::NullWaterSplineMetadata);
+	}
+
+	// TryQuery is implemented in such a way that it calls into the old virtual function so existing implementations
+	// are used until the old impl is removed and we can move the implementation into TryQuery.
+	PRAGMA_DISABLE_DEPRECATION_WARNINGS
+	return MakeValue(QueryWaterInfoClosestToWorldLocation(InWorldLocation, InQueryFlags, InSplineInputKey));
+	PRAGMA_ENABLE_DEPRECATION_WARNINGS
+}
+
 FWaterBodyQueryResult UWaterBodyComponent::QueryWaterInfoClosestToWorldLocation(const FVector& InWorldLocation, EWaterBodyQueryFlags InQueryFlags, const TOptional<float>& InSplineInputKey) const
 {
 	SCOPE_CYCLE_COUNTER(STAT_WaterBody_ComputeWaterInfo);
@@ -843,7 +862,7 @@ FWaterBodyQueryResult UWaterBodyComponent::QueryWaterInfoClosestToWorldLocation(
 
 }
 
-void UWaterBodyComponent::GetWaterSurfaceInfoAtLocation(const FVector& InLocation, FVector& OutWaterSurfaceLocation, FVector& OutWaterSurfaceNormal, FVector& OutWaterVelocity, float& OutWaterDepth, bool bIncludeDepth /* = false */) const
+bool UWaterBodyComponent::GetWaterSurfaceInfoAtLocation(const FVector& InLocation, FVector& OutWaterSurfaceLocation, FVector& OutWaterSurfaceNormal, FVector& OutWaterVelocity, float& OutWaterDepth, bool bIncludeDepth /* = false */) const
 {
 	EWaterBodyQueryFlags QueryFlags =
 		EWaterBodyQueryFlags::ComputeLocation
@@ -855,15 +874,23 @@ void UWaterBodyComponent::GetWaterSurfaceInfoAtLocation(const FVector& InLocatio
 		QueryFlags |= EWaterBodyQueryFlags::ComputeDepth;
 	}
 
-	FWaterBodyQueryResult QueryResult = QueryWaterInfoClosestToWorldLocation(InLocation, QueryFlags);
-	OutWaterSurfaceLocation = QueryResult.GetWaterSurfaceLocation();
-	OutWaterSurfaceNormal = QueryResult.GetWaterSurfaceNormal();
-	OutWaterVelocity = QueryResult.GetVelocity();
+	const TValueOrError<FWaterBodyQueryResult, EWaterBodyQueryError> QueryResult = TryQueryWaterInfoClosestToWorldLocation(InLocation, QueryFlags);
+	if (!QueryResult.HasValue())
+	{
+		return false;
+	}
+
+	const FWaterBodyQueryResult& Query = QueryResult.GetValue();
+	OutWaterSurfaceLocation = Query.GetWaterSurfaceLocation();
+	OutWaterSurfaceNormal = Query.GetWaterSurfaceNormal();
+	OutWaterVelocity = Query.GetVelocity();
 
 	if (bIncludeDepth)
 	{
-		OutWaterDepth = QueryResult.GetWaterSurfaceDepth();
+		OutWaterDepth = Query.GetWaterSurfaceDepth();
 	}
+
+	return true;
 }
 
 float UWaterBodyComponent::GetWaterVelocityAtSplineInputKey(float InKey) const
diff --git a/Engine/Plugins/Experimental/Water/Source/Runtime/Private/WaterSubsystem.cpp b/Engine/Plugins/Experimental/Water/Source/Runtime/Private/WaterSubsystem.cpp
--- a/Engine/Plugins/Experimental/Water/Source/Runtime/Private/WaterSubsystem.cpp
+++ b/Engine/Plugins/Experimental/Water/Source/Runtime/Private/WaterSubsystem.cpp
@@ -855,18 +855,18 @@ void UWaterSubsystem::ComputeUnderwaterPostProcess(FVector ViewLocation, FSceneV
 						| EWaterBodyQueryFlags::IncludeWaves;
 					AdjustUnderwaterWaterInfoQueryFlags(QueryFlags);
 
-					FWaterBodyQueryResult QueryResult = WaterBodyComponent->QueryWaterInfoClosestToWorldLocation(ViewLocation, QueryFlags);
-					if (!QueryResult.IsInExclusionVolume())
+					TValueOrError<FWaterBodyQueryResult, EWaterBodyQueryError> QueryResult = WaterBodyComponent->TryQueryWaterInfoClosestToWorldLocation(ViewLocation, QueryFlags);
+					if (QueryResult.HasValue() && !QueryResult.GetValue().IsInExclusionVolume())
 					{
 						// Calculate the surface max Z at the view XY location
-						float WaterSurfaceZ = QueryResult.GetWaterPlaneLocation().Z + QueryResult.GetWaveInfo().MaxHeight;
+						float WaterSurfaceZ = QueryResult.GetValue().GetWaterPlaneLocation().Z + QueryResult.GetValue().GetWaveInfo().MaxHeight;
 
 						// Only add the waterbody for processing if it has a higher surface than the previous waterbody (the Hits array is sorted by priority already)
 						// This also removed any duplicate waterbodies possibly returned by the sweep query
 						if (WaterSurfaceZ > MaxWaterLevel)
 						{
 							MaxWaterLevel = WaterSurfaceZ;
-							WaterBodyQueriesToProcess.Add(FWaterBodyPostProcessQuery(*WaterBodyComponent, ViewLocation, QueryResult));
+							WaterBodyQueriesToProcess.Add(FWaterBodyPostProcessQuery(*WaterBodyComponent, ViewLocation, QueryResult.GetValue()));
 						}
 					}
 				}
diff --git a/Engine/Plugins/Experimental/Water/Source/Runtime/Public/WaterBodyComponent.h b/Engine/Plugins/Experimental/Water/Source/Runtime/Public/WaterBodyComponent.h
--- a/Engine/Plugins/Experimental/Water/Source/Runtime/Public/WaterBodyComponent.h
+++ b/Engine/Plugins/Experimental/Water/Source/Runtime/Public/WaterBodyComponent.h
@@ -12,6 +12,7 @@
 #include "BakedShallowWaterSimulationComponent.h"
 #include "WaterZoneActor.h"
 #include "WaterBodyTypes.h"
+#include "Templates/ValueOrError.h"
 
 class AWaterBody;
 class UStaticMesh;
@@ -299,10 +300,13 @@ public:
 	- InQueryFlags: flags to indicate which info is to be computed
 	- InSplineInputKey: (optional) location on the spline, in case it has already been computed.
 	*/
+	WATER_API virtual TValueOrError<FWaterBodyQueryResult, EWaterBodyQueryError> TryQueryWaterInfoClosestToWorldLocation(const FVector& InWorldLocation, EWaterBodyQueryFlags InQueryFlags, const TOptional<float>& InSplineInputKey = TOptional<float>()) const;
+
+	UE_DEPRECATED(5.7, "Deprecated in favor of TryQueryWaterInfoClosestToWorldLocation to ensure caller correctly handles the cases where this function fails")
 	WATER_API virtual FWaterBodyQueryResult QueryWaterInfoClosestToWorldLocation(const FVector& InWorldLocation, EWaterBodyQueryFlags InQueryFlags, const TOptional<float>& InSplineInputKey = TOptional<float>()) const;
 
 	UFUNCTION(BlueprintCallable, Category = WaterBody)
-	WATER_API void GetWaterSurfaceInfoAtLocation(const FVector& InLocation, FVector& OutWaterSurfaceLocation, FVector& OutWaterSurfaceNormal, FVector& OutWaterVelocity, float& OutWaterDepth, bool bIncludeDepth = false) const;
+	WATER_API bool GetWaterSurfaceInfoAtLocation(const FVector& InLocation, FVector& OutWaterSurfaceLocation, FVector& OutWaterSurfaceNormal, FVector& OutWaterVelocity, float& OutWaterDepth, bool bIncludeDepth = false) const;
 
 	/** Spline query helper. It's faster to get the spline key once then query properties using that key, rather than querying repeatedly by location etc. */
 	WATER_API float FindInputKeyClosestToWorldLocation(const FVector& WorldLocation) const;
diff --git a/Engine/Plugins/Experimental/Water/Source/Runtime/Public/WaterBodyTypes.h b/Engine/Plugins/Experimental/Water/Source/Runtime/Public/WaterBodyTypes.h
--- a/Engine/Plugins/Experimental/Water/Source/Runtime/Public/WaterBodyTypes.h
+++ b/Engine/Plugins/Experimental/Water/Source/Runtime/Public/WaterBodyTypes.h
@@ -68,6 +68,17 @@ struct FWaveInfo
 
 // ----------------------------------------------------------------------------------
 
+UENUM()
+enum class EWaterBodyQueryError : uint8
+{
+	UnknownError,
+	// Invalid to query the water surface info for a water body which is not yet registered to the parent water body actor.
+	WaterBodyNotRegistered,
+	// Invalid to query the water surface info for a water body which has a null WaterSplineMetadata.
+	// This should be impossible in normal lifetime but could be null during the early parts of the component lifecycle before registration.
+	NullWaterSplineMetadata,
+};
+
 /** Struct holding the result from water queries :  */
 struct FWaterBodyQueryResult
 {
diff --git a/Engine/Plugins/Experimental/WaterAdvanced/Source/WaterAdvanced/Private/ShallowWaterSubsystem.cpp b/Engine/Plugins/Experimental/WaterAdvanced/Source/WaterAdvanced/Private/ShallowWaterSubsystem.cpp
--- a/Engine/Plugins/Experimental/WaterAdvanced/Source/WaterAdvanced/Private/ShallowWaterSubsystem.cpp
+++ b/Engine/Plugins/Experimental/WaterAdvanced/Source/WaterAdvanced/Private/ShallowWaterSubsystem.cpp
@@ -542,9 +542,15 @@ void UShallowWaterSubsystem::UpdateGridMovement()
 	{
 		if (const AWaterBody* WaterBody = WeakWaterBody.Get())
 		{
-			const FWaterBodyQueryResult WaterInfo = WaterBody->GetWaterBodyComponent()->QueryWaterInfoClosestToWorldLocation(
+			const TValueOrError<FWaterBodyQueryResult, EWaterBodyQueryError> WaterInfo = WaterBody->GetWaterBodyComponent()->TryQueryWaterInfoClosestToWorldLocation(
 				CursorPawnLocation, EWaterBodyQueryFlags::ComputeLocation);
-			const FVector WaterLocation = WaterInfo.GetWaterSurfaceLocation();
+
+			if (!WaterInfo.HasValue())
+			{
+				continue;
+			}
+
+			const FVector WaterLocation = WaterInfo.GetValue().GetWaterSurfaceLocation();
 			const float ds = FVector::DistSquared(WaterLocation, CursorPawnLocation); 
 			if (ds < BestDistanceSqr)
 			{
