Water Body - Ensure template objects do not call OnUpdateBody and generate collision boxes when the Wave Source property is changed from blueprint editor

diff --git a/Engine/Plugins/Experimental/Water/Source/Runtime/Private/WaterBodyActor.cpp b/Engine/Plugins/Experimental/Water/Source/Runtime/Private/WaterBodyActor.cpp
--- a/Engine/Plugins/Experimental/Water/Source/Runtime/Private/WaterBodyActor.cpp
+++ b/Engine/Plugins/Experimental/Water/Source/Runtime/Private/WaterBodyActor.cpp
@@ -194,14 +194,19 @@ void AWaterBody::PostEditChangeProperty(FPropertyChangedEvent& PropertyChangedEv
 	}
 	else if (PropertyChangedEvent.GetPropertyName() == GET_MEMBER_NAME_CHECKED(AWaterBody, WaterWaves))
 	{
-		WaterBodyComponent->RegisterOnUpdateWavesData(WaterWaves, /* bRegister = */true);
+		// Avoid calling OnWaterBodyChanged on templates. When changing properties from the BP editor this prevents
+		// the CDO from generating and adding collision components to the Blueprint unintentionally
+		if (!IsTemplate())
+		{
+			WaterBodyComponent->RegisterOnUpdateWavesData(WaterWaves, /* bRegister = */true);
 
-		WaterBodyComponent->RequestGPUWaveDataUpdate();
+			WaterBodyComponent->RequestGPUWaveDataUpdate();
 
-		FOnWaterBodyChangedParams Params;
-		// Waves data affect the navigation : 
-		Params.bShapeOrPositionChanged = true;
-		WaterBodyComponent->OnWaterBodyChanged(Params);
+			FOnWaterBodyChangedParams Params;
+			// Waves data affect the navigation : 
+			Params.bShapeOrPositionChanged = true;
+			WaterBodyComponent->OnWaterBodyChanged(Params);
+		}
 	}
 }
 
diff --git a/Engine/Plugins/Experimental/Water/Source/Runtime/Private/WaterBodyOceanComponent.cpp b/Engine/Plugins/Experimental/Water/Source/Runtime/Private/WaterBodyOceanComponent.cpp
--- a/Engine/Plugins/Experimental/Water/Source/Runtime/Private/WaterBodyOceanComponent.cpp
+++ b/Engine/Plugins/Experimental/Water/Source/Runtime/Private/WaterBodyOceanComponent.cpp
@@ -597,7 +597,10 @@ void UWaterBodyOceanComponent::OnUpdateBody(bool bWithExclusionVolumes)
 			Boxes, ConvexSets, WorldMeshBufferWidth, WorldBoxOverlap);
 
 		// Don't delete components unless we have to : this generates determinism issues because UOceanCollisionComponent has a UBodySetup with a GUID :
-		if ((CollisionBoxes.Num() != Boxes.Num()) || (CollisionHullSets.Num() != ConvexSets.Num()))
+		const bool bComponentCountEqual = CollisionBoxes.Num() == Boxes.Num() && CollisionHullSets.Num() == ConvexSets.Num();
+
+		// Some assets have saved empty collision box entries, reset if any are invalid
+		if (!bComponentCountEqual || CollisionBoxes.Contains(nullptr) || CollisionHullSets.Contains(nullptr))
 		{
 			Reset();
 		}
