[Water] Implemented customizable ocean far mesh height for when there is an ocean far mesh without ocean water body in the water zone.

diff --git a/Engine/Plugins/Experimental/Water/Source/Runtime/Private/WaterMeshComponent.cpp b/Engine/Plugins/Experimental/Water/Source/Runtime/Private/WaterMeshComponent.cpp
--- a/Engine/Plugins/Experimental/Water/Source/Runtime/Private/WaterMeshComponent.cpp
+++ b/Engine/Plugins/Experimental/Water/Source/Runtime/Private/WaterMeshComponent.cpp
@@ -275,7 +275,7 @@ void UWaterMeshComponent::RebuildWaterMesh(float InTileSize, const FIntPoint& In
 	const bool bIsGPUQuadTree = CVarWaterMeshGPUQuadTree.GetValueOnGameThread() != 0;
 	
 	FMaterialRenderProxy* FarDistanceMaterialProxy = IsMaterialUsedWithWater(FarDistanceMaterial) ? FarDistanceMaterial->GetRenderProxy() : nullptr;
-	WaterQuadTreeBuilder.Init(WaterZone->GetZoneBounds2D(), InExtentInTiles, InTileSize, FarDistanceMaterialProxy, FarDistanceMeshExtent, bUseFarMeshWithoutOcean, bIsGPUQuadTree);
+	WaterQuadTreeBuilder.Init(WaterZone->GetZoneBounds2D(), InExtentInTiles, InTileSize, FarDistanceMaterialProxy, FarDistanceMeshExtent, FarDistanceMeshHeightWithoutOcean, bUseFarMeshWithoutOcean, bIsGPUQuadTree);
 
 	UsedMaterials.Empty();
 
@@ -616,6 +616,7 @@ void UWaterMeshComponent::PostEditChangeProperty(FPropertyChangedEvent& Property
 			|| PropertyName == GET_MEMBER_NAME_CHECKED(UWaterMeshComponent, ForceCollapseDensityLevel)
 			|| PropertyName == GET_MEMBER_NAME_CHECKED(UWaterMeshComponent, FarDistanceMaterial)
 			|| PropertyName == GET_MEMBER_NAME_CHECKED(UWaterMeshComponent, FarDistanceMeshExtent)
+			|| PropertyName == GET_MEMBER_NAME_CHECKED(UWaterMeshComponent, FarDistanceMeshHeightWithoutOcean)
 			|| PropertyName == GET_MEMBER_NAME_CHECKED(UWaterMeshComponent, bUseFarMeshWithoutOcean))
 		{
 			MarkWaterMeshGridDirty();
diff --git a/Engine/Plugins/Experimental/Water/Source/Runtime/Private/WaterQuadTreeBuilder.cpp b/Engine/Plugins/Experimental/Water/Source/Runtime/Private/WaterQuadTreeBuilder.cpp
--- a/Engine/Plugins/Experimental/Water/Source/Runtime/Private/WaterQuadTreeBuilder.cpp
+++ b/Engine/Plugins/Experimental/Water/Source/Runtime/Private/WaterQuadTreeBuilder.cpp
@@ -5,7 +5,7 @@
 #include "WaterBodyTypes.h"
 #include "WaterModule.h"
 
-void FWaterQuadTreeBuilder::Init(const FBox2D& InWaterZoneBounds2D, const FIntPoint& InExtentInTiles, float InTileSize, FMaterialRenderProxy* InFarDistanceMaterial, float InFarDistanceMeshExtent, bool bInUseFarMeshWithoutOcean, bool bInIsGPUQuadTree)
+void FWaterQuadTreeBuilder::Init(const FBox2D& InWaterZoneBounds2D, const FIntPoint& InExtentInTiles, float InTileSize, FMaterialRenderProxy* InFarDistanceMaterial, float InFarDistanceMeshExtent, double InDefaultFarDistanceMeshHeight, bool bInUseFarMeshWithoutOcean, bool bInIsGPUQuadTree)
 {
 	WaterBodies.Reset();
 	WaterZoneBounds2D = InWaterZoneBounds2D;
@@ -13,6 +13,7 @@ void FWaterQuadTreeBuilder::Init(const FBox2D& InWaterZoneBounds2D, const FIntPo
 	TileSize = InTileSize;
 	FarDistanceMaterial = InFarDistanceMaterial;
 	FarDistanceMeshExtent = InFarDistanceMeshExtent;
+	DefaultFarDistanceMeshHeight = InDefaultFarDistanceMeshHeight;
 	bUseFarMeshWithoutOcean = bInUseFarMeshWithoutOcean;
 	bIsGPUQuadTree = bInIsGPUQuadTree;
 
@@ -41,7 +42,7 @@ bool FWaterQuadTreeBuilder::BuildWaterQuadTree(FWaterQuadTree& WaterQuadTree, co
 	WaterQuadTree.InitTree(WaterWorldBox, TileSize, ExtentInTiles, bIsGPUQuadTree);
 
 	// Will be updated with the ocean min bound, to be used to place the far mesh just under the ocean to avoid seams
-	float FarMeshHeight = 0.0f;
+	double FarMeshHeight = DefaultFarDistanceMeshHeight;
 	// Only use a far mesh when there is an ocean in the zone.
 	bool bHasOcean = false;
 	
diff --git a/Engine/Plugins/Experimental/Water/Source/Runtime/Public/WaterMeshComponent.h b/Engine/Plugins/Experimental/Water/Source/Runtime/Public/WaterMeshComponent.h
--- a/Engine/Plugins/Experimental/Water/Source/Runtime/Public/WaterMeshComponent.h
+++ b/Engine/Plugins/Experimental/Water/Source/Runtime/Public/WaterMeshComponent.h
@@ -129,6 +129,10 @@ private:
 	UPROPERTY(Category = "Rendering|FarDistance", EditAnywhere)
 	bool bUseFarMeshWithoutOcean = false;
 
+	/** Absolute world space height of the far mesh when there is no ocean.*/
+	UPROPERTY(Category = "Rendering|FarDistance", EditAnywhere)
+	double FarDistanceMeshHeightWithoutOcean = 0.0;
+
 	/** Dirty flag which will make sure the water mesh is updated properly */
 	bool bNeedsRebuild = true;
 
diff --git a/Engine/Plugins/Experimental/Water/Source/Runtime/Public/WaterQuadTreeBuilder.h b/Engine/Plugins/Experimental/Water/Source/Runtime/Public/WaterQuadTreeBuilder.h
--- a/Engine/Plugins/Experimental/Water/Source/Runtime/Public/WaterQuadTreeBuilder.h
+++ b/Engine/Plugins/Experimental/Water/Source/Runtime/Public/WaterQuadTreeBuilder.h
@@ -35,7 +35,7 @@ public:
 #endif // WITH_WATER_SELECTION_SUPPORT
 	};
 
-	void Init(const FBox2D& InWaterZoneBounds2D, const FIntPoint& InExtentInTiles, float InTileSize, FMaterialRenderProxy* InFarDistanceMaterial, float InFarDistanceMeshExtent, bool bInUseFarMeshWithoutOcean, bool bInIsGPUQuadTree);
+	void Init(const FBox2D& InWaterZoneBounds2D, const FIntPoint& InExtentInTiles, float InTileSize, FMaterialRenderProxy* InFarDistanceMaterial, float InFarDistanceMeshExtent, double InDefaultFarDistanceMeshHeight, bool bInUseFarMeshWithoutOcean, bool bInIsGPUQuadTree);
 	void AddWaterBody(const FWaterBody& WaterBody);
 	bool BuildWaterQuadTree(FWaterQuadTree& WaterQuadTree, const FVector2D& GridPosition) const;
 #if WITH_WATER_SELECTION_SUPPORT
@@ -55,6 +55,7 @@ private:
 	int32 TreeDepth = 0;
 	FMaterialRenderProxy* FarDistanceMaterial = nullptr;
 	float FarDistanceMeshExtent = 0.0f;
+	double DefaultFarDistanceMeshHeight = 0.0;
 	bool bUseFarMeshWithoutOcean = false;
 	bool bIsGPUQuadTree = false;
 };
