Prevent crash when redoing an undo of a Water Zone actor deletion

diff --git a/Engine/Plugins/Experimental/Water/Source/Runtime/Private/WaterZoneActor.cpp b/Engine/Plugins/Experimental/Water/Source/Runtime/Private/WaterZoneActor.cpp
--- a/Engine/Plugins/Experimental/Water/Source/Runtime/Private/WaterZoneActor.cpp
+++ b/Engine/Plugins/Experimental/Water/Source/Runtime/Private/WaterZoneActor.cpp
@@ -435,8 +435,11 @@ void AWaterZone::PostEditUndo()
 {
 	Super::PostEditUndo();
 
-	UpdateOverlappingWaterBodies();
-	MarkForRebuild(EWaterZoneRebuildFlags::All, /* DebugRequestingObject = */ this);
+	if (IsValid(this))
+	{
+		UpdateOverlappingWaterBodies();
+		MarkForRebuild(EWaterZoneRebuildFlags::All, /* DebugRequestingObject = */ this);
+	}
 }
 
 void AWaterZone::PostEditImport()
