[SingleLayerWater] Fixed missing water tiles in right eye with ISR due to using the left eye frustum for culling instead of a unified culling frustum for both eyes.

diff --git a/Engine/Plugins/Experimental/Water/Source/Runtime/Private/WaterMeshSceneProxy.cpp b/Engine/Plugins/Experimental/Water/Source/Runtime/Private/WaterMeshSceneProxy.cpp
--- a/Engine/Plugins/Experimental/Water/Source/Runtime/Private/WaterMeshSceneProxy.cpp
+++ b/Engine/Plugins/Experimental/Water/Source/Runtime/Private/WaterMeshSceneProxy.cpp
@@ -525,7 +525,7 @@ void FWaterMeshSceneProxy::GetDynamicMeshElements(const TArray<const FSceneView*
 				TraversalDesc.LODCount = WaterQuadTree.GetTreeDepth();
 				TraversalDesc.DensityCount = WaterQuadTreeConstants.DensityCount;
 				TraversalDesc.ForceCollapseDensityLevel = WaterQuadTreeConstants.ForceCollapseDensityLevel;
-				TraversalDesc.Frustum = View->ViewFrustum;
+				TraversalDesc.Frustum = View->GetCullingFrustum();
 				TraversalDesc.ObserverPosition = ObserverPosition;
 				TraversalDesc.PreViewTranslation = View->ViewMatrices.GetPreViewTranslation();
 				TraversalDesc.LODScale = WaterQuadTreeConstants.LODScale;
