Water - Adds refined underwater detection to avoid applying post processing when the camera is underneath or outside the bounds of the collision volume

diff --git a/Engine/Plugins/Experimental/Water/Source/Runtime/Private/WaterSubsystem.cpp b/Engine/Plugins/Experimental/Water/Source/Runtime/Private/WaterSubsystem.cpp
--- a/Engine/Plugins/Experimental/Water/Source/Runtime/Private/WaterSubsystem.cpp
+++ b/Engine/Plugins/Experimental/Water/Source/Runtime/Private/WaterSubsystem.cpp
@@ -21,6 +21,8 @@
 #include "WaterViewExtension.h"
 #include "Algo/MaxElement.h"
 #include "Algo/RemoveIf.h"
+#include "Algo/AnyOf.h"
+#include "VisualLogger/VisualLogger.h"
 
 #if WITH_EDITOR
 #include "WaterZoneActorDesc.h"
@@ -29,7 +31,7 @@
 extern UNREALED_API UEditorEngine* GEditor;
 #else
 #include "BuoyancyTypes.h"
-#endif // WITH_DITOR
+#endif // WITH_EDITOR
 
 #include "LandscapeComponent.h"
 #include "LandscapeProxy.h"
@@ -87,6 +89,30 @@ static FAutoConsoleVariableRef CVarVisualizeUnderwaterPostProcess(
 	ECVF_Default
 );
 
+static int32 EnableUnderwaterPostProcessVisualLogger = 0;
+static FAutoConsoleVariableRef CVarEnableUnderwaterPostProcessVisualLogger(
+	TEXT("r.Water.EnableUnderwaterPostProcessVisualLogger"),
+	EnableUnderwaterPostProcessVisualLogger,
+	TEXT("Enables underwater post process collision detection Visual Logger code"),
+	ECVF_Default
+);
+
+static float UnderwaterCollisionTraceDistance = 100.f;
+static FAutoConsoleVariableRef CVarUnderwaterCollisionTraceDistance(
+	TEXT("r.Water.UnderwaterCollisionTraceDistance"),
+	UnderwaterCollisionTraceDistance,
+	TEXT("Underwater post processing collision trace distance. Default is inflated to account for waves which will not be hit by trace"),
+	ECVF_Scalability
+);
+
+static float UnderwaterCollisionPreciseTraceDistance = 10.f;
+static FAutoConsoleVariableRef CVarUnderwaterCollisionPreciseTraceDistance(
+	TEXT("r.Water.UnderwaterCollisionPreciseTraceDistance"),
+	UnderwaterCollisionPreciseTraceDistance,
+	TEXT("Underwater precise collision trace distance. Verifies precise overlap between camera and volume when camera is underneath the water collision impact location"),
+	ECVF_Scalability
+);
+
 // Shallow water CVars : 
 static int32 ShallowWaterSim = 1;
 static FAutoConsoleVariableRef CVarShallowWaterSim(
@@ -450,6 +476,16 @@ bool UWaterSubsystem::IsUnderwaterPostProcessEnabled() const
 	return EnableUnderwaterPostProcess != 0;
 }
 
+float UWaterSubsystem::GetUnderwaterCollisionTraceDistance()
+{
+	return UnderwaterCollisionTraceDistance;
+}
+
+float UWaterSubsystem::GetUnderwaterPreciseTraceDistance()
+{
+	return UnderwaterCollisionPreciseTraceDistance;
+}
+
 int32 UWaterSubsystem::GetShallowWaterMaxDynamicForces()
 {
 	return ShallowWaterSimulationMaxDynamicForces;
@@ -797,16 +833,17 @@ void UWaterSubsystem::ComputeUnderwaterPostProcess(FVector ViewLocation, FSceneV
 	const float ViewToNearPlaneCornerDistance = FVector2f(NearPlaneCornerViewSpace / NearPlaneCornerViewSpace.W).Length();
 	bool bAnyDefinitelyUnderwater = false;
 	bool bAnyPossiblyUnderwater = false;
-
 	bool bUnderwaterForPostProcess = false;
 
-	// Trace just a small distance extra from the viewpoint to account for waves since the waves wont traced against
-	static const float TraceDistance = 100.f;
+	const bool bIsVisualLoggerEnabled = EnableUnderwaterPostProcessVisualLogger != 0;
+	const float TraceDistance = GetUnderwaterCollisionTraceDistance();
+	const float PreciseTraceDistance = GetUnderwaterPreciseTraceDistance();
 
 	// Always force simple collision traces
 	static FCollisionQueryParams TraceSimple(SCENE_QUERY_STAT(DefaultQueryParam), false);
 
 	TArray<FHitResult> Hits;
+	TArray<FHitResult> PreciseViewHits;
 	TArray<FWaterBodyPostProcessQuery, TInlineAllocator<4>> WaterBodyQueriesToProcess;
 	const bool bWorldHasWater  = WaterBodyManager.HasAnyWaterBodies();
 	if (bWorldHasWater && World->SweepMultiByChannel(Hits, ViewLocation, ViewLocation + FVector(0, 0, TraceDistance), FQuat::Identity, UnderwaterTraceChannel, FCollisionShape::MakeSphere(TraceDistance), TraceSimple))
@@ -838,8 +875,63 @@ void UWaterSubsystem::ComputeUnderwaterPostProcess(FVector ViewLocation, FSceneV
 			});
 		}
 
+		UE_IFVLOG(
+			if (bIsVisualLoggerEnabled)
+			{
+				// Visualize camera marker 
+				UE_VLOG_LOCATION(this, LogWater, Log, ViewLocation, /*Radius=*/10.f, FColor::Green, TEXT("Camera"));
+
+				// Visulize primary sweep capsule
+				const FVector SweepStart = ViewLocation;
+				const FVector SweepEnd = ViewLocation + FVector(0, 0, TraceDistance);
+				UE_VLOG_CAPSULE(this, LogWater, Log, SweepStart, /*HalfHeight=*/ 0.0f, /*Radius=*/TraceDistance, FQuat::Identity, FColor::Yellow, TEXT("Primary sweep sphere@start"));
+
+				// Visualize raw hit results
+				for (const FHitResult& Hit : Hits)
+				{
+					UE_VLOG_LOCATION(this, LogWater, Log, Hit.ImpactPoint, 8.f, FColor::Red, TEXT("Primary Hit"));
+
+					const float NormalLen = 50.f;
+					UE_VLOG_ARROW(this, LogWater, Log, Hit.ImpactPoint, Hit.ImpactPoint + Hit.ImpactNormal * NormalLen, FColor::Red, TEXT("Primary Normal"));
+				}
+			}
+		);
+
+		// Determine if view is under water (camera is below impact point)
+		const double& ViewLocationZ = ViewLocation.Z;
+		const float ViewTolerance = 1.0f;
+		const bool bIsViewLocationUnderImpactPoint = Algo::AnyOf(Hits, [&ViewLocationZ, &ViewTolerance](const FHitResult& Hit)
+		{
+			// Add a smaller tolerance to avoid flickering when the view location hovers around the impact point
+			return (Hit.ImpactPoint.Z + ViewTolerance) > ViewLocationZ;
+		});
+
+		if (bIsViewLocationUnderImpactPoint)
+		{
+			World->SweepMultiByChannel(PreciseViewHits, ViewLocation, ViewLocation, FQuat::Identity, UnderwaterTraceChannel, FCollisionShape::MakeSphere(PreciseTraceDistance), TraceSimple);
+
+			UE_IFVLOG(
+				if (bIsVisualLoggerEnabled)
+				{
+					// Visualize precise sphere capsule
+					UE_VLOG_CAPSULE(this, LogWater, Log, ViewLocation, /*HalfHeight=*/ 0.0f, /*Radius=*/PreciseTraceDistance, FQuat::Identity, FColor::Blue, TEXT("Precise sphere@start"));
+
+					// Visualize precise hit results
+					for (const FHitResult& PreciseHit : PreciseViewHits)
+					{
+						UE_VLOG_LOCATION(this, LogWater, Log, PreciseHit.ImpactPoint, 10.f, FColor::Blue, TEXT("Precise Hit"));
+
+						const float NormalLen = 30.f;
+						UE_VLOG_ARROW(this, LogWater, Log, PreciseHit.ImpactPoint, PreciseHit.ImpactPoint + PreciseHit.ImpactNormal * NormalLen, FColor::Blue, TEXT("Precise Normal"));
+					}
+				}
+			);
+		}
+
 		float MaxWaterLevel = TNumericLimits<float>::Lowest();
-		for (const FHitResult& Result : Hits)
+		// When the camera is above water use the extended sphere collision to ensure post processing gets applied in waves
+		TArray<FHitResult>& RefinedHits = bIsViewLocationUnderImpactPoint ? PreciseViewHits : Hits;
+		for (const FHitResult& Result : RefinedHits)
 		{
 			if (AWaterBody* WaterBodyActor = Result.HitObjectHandle.FetchActor<AWaterBody>())
 			{
diff --git a/Engine/Plugins/Experimental/Water/Source/Runtime/Public/WaterSubsystem.h b/Engine/Plugins/Experimental/Water/Source/Runtime/Public/WaterSubsystem.h
--- a/Engine/Plugins/Experimental/Water/Source/Runtime/Public/WaterSubsystem.h
+++ b/Engine/Plugins/Experimental/Water/Source/Runtime/Public/WaterSubsystem.h
@@ -120,6 +120,12 @@ public:
 	UFUNCTION(BlueprintCallable, BlueprintPure, Category = Water)
 	UE_API bool IsUnderwaterPostProcessEnabled() const;
 
+	UFUNCTION(BlueprintCallable, Category = Water)
+	static UE_API float GetUnderwaterCollisionTraceDistance();
+
+	UFUNCTION(BlueprintCallable, Category = Water)
+	static UE_API float GetUnderwaterPreciseTraceDistance();
+
 	UFUNCTION(BlueprintCallable, Category=Water)
 	static UE_API int32 GetShallowWaterMaxDynamicForces();
 
@@ -206,19 +212,19 @@ public:
 #endif // WITH_EDITOR
 
 private:
-	UE_API void NotifyWaterScalabilityChangedInternal(IConsoleVariable* CVar);
-	UE_API void NotifyWaterVisibilityChangedInternal(IConsoleVariable* CVar);
-	UE_API void ComputeUnderwaterPostProcess(FVector ViewLocation, FSceneView* SceneView);
-	UE_API void SetMPCTime(float Time, float PrevTime);
-	UE_API void AdjustUnderwaterWaterInfoQueryFlags(EWaterBodyQueryFlags& InOutFlags);
-	UE_API void ApplyRuntimeSettings(const UWaterRuntimeSettings* Settings, EPropertyChangeType::Type ChangeType);
+	void NotifyWaterScalabilityChangedInternal(IConsoleVariable* CVar);
+	void NotifyWaterVisibilityChangedInternal(IConsoleVariable* CVar);
+	void ComputeUnderwaterPostProcess(FVector ViewLocation, FSceneView* SceneView);
+	void SetMPCTime(float Time, float PrevTime);
+	void AdjustUnderwaterWaterInfoQueryFlags(EWaterBodyQueryFlags& InOutFlags);
+	void ApplyRuntimeSettings(const UWaterRuntimeSettings* Settings, EPropertyChangeType::Type ChangeType);
 
-	UE_API void OnMarkRenderStateDirty(UActorComponent& Component);
+	void OnMarkRenderStateDirty(UActorComponent& Component);
 
-	UE_API void OnWaterTerrainActorChanged(const AActor* TerrainActor);
+	void OnWaterTerrainActorChanged(const AActor* TerrainActor);
 
 #if !(UE_BUILD_SHIPPING || UE_BUILD_TEST)
-	UE_API void ShowOnScreenDebugInfo(const FVector& InViewLocation, const FUnderwaterPostProcessDebugInfo& InDebugInfo);
+	void ShowOnScreenDebugInfo(const FVector& InViewLocation, const FUnderwaterPostProcessDebugInfo& InDebugInfo);
 #endif // !(UE_BUILD_SHIPPING || UE_BUILD_TEST)
 
 public:
@@ -274,7 +280,7 @@ private:
 
 #if WITH_EDITOR
 	/** By default, there is no water subsystem allowed on preview worlds except when explicitly requested : */
-	static UE_API bool bAllowWaterSubsystemOnPreviewWorld;
+	static bool bAllowWaterSubsystemOnPreviewWorld;
 #endif // WITH_EDITOR
 };
 
