Water: Cooker: Silence the warning "Unsolicited package %s was loaded by package %s" when WaterBodyComponents are set to use the DefaultWaterInfoMaterial and load it during the cook. This package is then present in their runtime data, so we need to mark it as a UsedInGame load.
#rnx

diff --git a/Engine/Plugins/Experimental/Water/Source/Runtime/Private/WaterBodyComponent.cpp b/Engine/Plugins/Experimental/Water/Source/Runtime/Private/WaterBodyComponent.cpp
--- a/Engine/Plugins/Experimental/Water/Source/Runtime/Private/WaterBodyComponent.cpp
+++ b/Engine/Plugins/Experimental/Water/Source/Runtime/Private/WaterBodyComponent.cpp
@@ -1733,6 +1733,7 @@ void UWaterBodyComponent::Serialize(FArchive& Ar)
 			// If a user _had_ overriden the material, it would have been set in UObject::Serialize.
 			if (WaterInfoMaterial == nullptr)
 			{
+				FCookLoadScope UsedInGameScope(ECookLoadType::UsedInGame);
 				WaterInfoMaterial = GetDefault<UWaterRuntimeSettings>()->GetDefaultWaterInfoMaterial();
 			}
 		}
