[Water] Fixed crash due to caching FMaterialRenderProxy instead of UMaterialInterface

diff --git a/Engine/Plugins/Experimental/Water/Source/Runtime/Private/WaterMeshComponent.cpp b/Engine/Plugins/Experimental/Water/Source/Runtime/Private/WaterMeshComponent.cpp
--- a/Engine/Plugins/Experimental/Water/Source/Runtime/Private/WaterMeshComponent.cpp
+++ b/Engine/Plugins/Experimental/Water/Source/Runtime/Private/WaterMeshComponent.cpp
@@ -274,8 +274,8 @@ void UWaterMeshComponent::RebuildWaterMesh(float InTileSize, const FIntPoint& In
 	bool bAnyWaterMeshesNotReady = false;
 	const bool bIsGPUQuadTree = CVarWaterMeshGPUQuadTree.GetValueOnGameThread() != 0;
 	
-	FMaterialRenderProxy* FarDistanceMaterialProxy = IsMaterialUsedWithWater(FarDistanceMaterial) ? FarDistanceMaterial->GetRenderProxy() : nullptr;
-	WaterQuadTreeBuilder.Init(WaterZone->GetZoneBounds2D(), InExtentInTiles, InTileSize, FarDistanceMaterialProxy, FarDistanceMeshExtent, FarDistanceMeshHeightWithoutOcean, bUseFarMeshWithoutOcean, bIsGPUQuadTree);
+	UMaterialInterface* FarDistanceMaterialInterface = IsMaterialUsedWithWater(FarDistanceMaterial) ? FarDistanceMaterial : nullptr;
+	WaterQuadTreeBuilder.Init(WaterZone->GetZoneBounds2D(), InExtentInTiles, InTileSize, FarDistanceMaterialInterface, FarDistanceMeshExtent, FarDistanceMeshHeightWithoutOcean, bUseFarMeshWithoutOcean, bIsGPUQuadTree);
 
 	UsedMaterials.Empty();
 
@@ -304,7 +304,7 @@ void UWaterMeshComponent::RebuildWaterMesh(float InTileSize, const FIntPoint& In
 
 		bUsesFarDistanceMaterial = bUsesFarDistanceMaterial || (WaterBodyComponent->GetWaterBodyType() == EWaterBodyType::Ocean);
 
-		auto GetMaterialProxy = [&](UMaterialInstanceDynamic* MID, bool bUseFallback) -> FMaterialRenderProxy*
+		auto GetMaterialInterface = [&](UMaterialInstanceDynamic* MID, bool bUseFallback) -> UMaterialInterface*
 		{
 			UMaterialInterface* MaterialInterface = nullptr;
 			if (!MID || !IsMaterialUsedWithWater(MID))
@@ -319,7 +319,7 @@ void UWaterMeshComponent::RebuildWaterMesh(float InTileSize, const FIntPoint& In
 			if (MaterialInterface)
 			{
 				UsedMaterials.Add(MaterialInterface);
-				return MaterialInterface->GetRenderProxy();
+				return MaterialInterface;
 			}
 			return nullptr;
 		};
@@ -327,9 +327,9 @@ void UWaterMeshComponent::RebuildWaterMesh(float InTileSize, const FIntPoint& In
 		const bool bIsRiver = WaterBodyComponent->GetWaterBodyType() == EWaterBodyType::River;
 
 		FWaterQuadTreeBuilder::FWaterBody WaterBody = {};
-		WaterBody.Material = GetMaterialProxy(WaterBodyComponent->GetWaterMaterialInstance(), true);
-		WaterBody.RiverToLakeMaterial = bIsRiver ? GetMaterialProxy(WaterBodyComponent->GetRiverToLakeTransitionMaterialInstance(), false) : nullptr;
-		WaterBody.RiverToOceanMaterial = bIsRiver ? GetMaterialProxy(WaterBodyComponent->GetRiverToOceanTransitionMaterialInstance(), false) : nullptr;
+		WaterBody.Material = GetMaterialInterface(WaterBodyComponent->GetWaterMaterialInstance(), true);
+		WaterBody.RiverToLakeMaterial = bIsRiver ? GetMaterialInterface(WaterBodyComponent->GetRiverToLakeTransitionMaterialInstance(), false) : nullptr;
+		WaterBody.RiverToOceanMaterial = bIsRiver ? GetMaterialInterface(WaterBodyComponent->GetRiverToOceanTransitionMaterialInstance(), false) : nullptr;
 		WaterBody.StaticMeshRenderData = StaticMeshRenderData;
 		WaterBody.LocalToWorld = WaterBodyComponent->GetComponentTransform();
 		WaterBody.Bounds = WaterBodyComponent->Bounds;
@@ -536,7 +536,7 @@ void UWaterMeshComponent::RebuildWaterMesh(float InTileSize, const FIntPoint& In
 		return true;
 	});
 
-	if (bUsesFarDistanceMaterial && FarDistanceMaterialProxy && (FarDistanceMeshExtent > 0.0f))
+	if (bUsesFarDistanceMaterial && FarDistanceMaterialInterface && (FarDistanceMeshExtent > 0.0f))
 	{
 		UsedMaterials.Add(FarDistanceMaterial);
 	}
diff --git a/Engine/Plugins/Experimental/Water/Source/Runtime/Private/WaterQuadTreeBuilder.cpp b/Engine/Plugins/Experimental/Water/Source/Runtime/Private/WaterQuadTreeBuilder.cpp
--- a/Engine/Plugins/Experimental/Water/Source/Runtime/Private/WaterQuadTreeBuilder.cpp
+++ b/Engine/Plugins/Experimental/Water/Source/Runtime/Private/WaterQuadTreeBuilder.cpp
@@ -4,8 +4,11 @@
 #include "WaterQuadTree.h"
 #include "WaterBodyTypes.h"
 #include "WaterModule.h"
+#include "Materials/Material.h"
+#include "Materials/MaterialInterface.h"
+#include "MaterialDomain.h"
 
-void FWaterQuadTreeBuilder::Init(const FBox2D& InWaterZoneBounds2D, const FIntPoint& InExtentInTiles, float InTileSize, FMaterialRenderProxy* InFarDistanceMaterial, float InFarDistanceMeshExtent, double InDefaultFarDistanceMeshHeight, bool bInUseFarMeshWithoutOcean, bool bInIsGPUQuadTree)
+void FWaterQuadTreeBuilder::Init(const FBox2D& InWaterZoneBounds2D, const FIntPoint& InExtentInTiles, float InTileSize, const UMaterialInterface* InFarDistanceMaterial, float InFarDistanceMeshExtent, double InDefaultFarDistanceMeshHeight, bool bInUseFarMeshWithoutOcean, bool bInIsGPUQuadTree)
 {
 	WaterBodies.Reset();
 	WaterZoneBounds2D = InWaterZoneBounds2D;
@@ -88,9 +91,9 @@ bool FWaterQuadTreeBuilder::BuildWaterQuadTree(FWaterQuadTree& WaterQuadTree, co
 		}
 
 		FWaterBodyRenderData RenderData;
-		RenderData.Material = WaterBody.Material;
-		RenderData.RiverToLakeMaterial = WaterBody.RiverToLakeMaterial;
-		RenderData.RiverToOceanMaterial = WaterBody.RiverToOceanMaterial;
+		RenderData.Material = WaterBody.Material ? WaterBody.Material->GetRenderProxy() : UMaterial::GetDefaultMaterial(MD_Surface)->GetRenderProxy();
+		RenderData.RiverToLakeMaterial = WaterBody.RiverToLakeMaterial ? WaterBody.RiverToLakeMaterial->GetRenderProxy() : nullptr;
+		RenderData.RiverToOceanMaterial = WaterBody.RiverToOceanMaterial ? WaterBody.RiverToOceanMaterial->GetRenderProxy() : nullptr;
 		RenderData.Priority = static_cast<int16>(FMath::Clamp(WaterBody.OverlapMaterialPriority, MinWaterBodyPriority, MaxWaterBodyPriority));
 		RenderData.WaterBodyIndex = static_cast<int16>(WaterBody.WaterBodyIndex);
 		RenderData.SurfaceBaseHeight = WaterBody.SurfaceBaseHeight;
@@ -177,7 +180,7 @@ bool FWaterQuadTreeBuilder::BuildWaterQuadTree(FWaterQuadTree& WaterQuadTree, co
 	{
 		// Far Mesh should stitch to the edge of the water zone
 		const FBox2D FarMeshBounds = WaterZoneBounds2D;
-		WaterQuadTree.AddFarMesh(FarDistanceMaterial, FarMeshBounds, FarDistanceMeshExtent, FarMeshHeight);
+		WaterQuadTree.AddFarMesh(FarDistanceMaterial->GetRenderProxy(), FarMeshBounds, FarDistanceMeshExtent, FarMeshHeight);
 	}
 
 	WaterQuadTree.Unlock(true);
diff --git a/Engine/Plugins/Experimental/Water/Source/Runtime/Public/WaterQuadTreeBuilder.h b/Engine/Plugins/Experimental/Water/Source/Runtime/Public/WaterQuadTreeBuilder.h
--- a/Engine/Plugins/Experimental/Water/Source/Runtime/Public/WaterQuadTreeBuilder.h
+++ b/Engine/Plugins/Experimental/Water/Source/Runtime/Public/WaterQuadTreeBuilder.h
@@ -5,7 +5,7 @@
 
 class HHitProxy;
 class FStaticMeshRenderData;
-class FMaterialRenderProxy;
+class UMaterialInterface;
 struct FWaterQuadTree;
 enum class EWaterBodyType : uint8;
 
@@ -15,9 +15,9 @@ public:
 
 	struct FWaterBody
 	{
-		FMaterialRenderProxy* Material = nullptr;
-		FMaterialRenderProxy* RiverToLakeMaterial = nullptr;
-		FMaterialRenderProxy* RiverToOceanMaterial = nullptr;
+		const UMaterialInterface* Material = nullptr;
+		const UMaterialInterface* RiverToLakeMaterial = nullptr;
+		const UMaterialInterface* RiverToOceanMaterial = nullptr;
 		FStaticMeshRenderData* StaticMeshRenderData = nullptr;
 		FTransform LocalToWorld;
 		FBoxSphereBounds Bounds = FBoxSphereBounds();
@@ -35,7 +35,7 @@ public:
 #endif // WITH_WATER_SELECTION_SUPPORT
 	};
 
-	void Init(const FBox2D& InWaterZoneBounds2D, const FIntPoint& InExtentInTiles, float InTileSize, FMaterialRenderProxy* InFarDistanceMaterial, float InFarDistanceMeshExtent, double InDefaultFarDistanceMeshHeight, bool bInUseFarMeshWithoutOcean, bool bInIsGPUQuadTree);
+	void Init(const FBox2D& InWaterZoneBounds2D, const FIntPoint& InExtentInTiles, float InTileSize, const UMaterialInterface* InFarDistanceMaterial, float InFarDistanceMeshExtent, double InDefaultFarDistanceMeshHeight, bool bInUseFarMeshWithoutOcean, bool bInIsGPUQuadTree);
 	void AddWaterBody(const FWaterBody& WaterBody);
 	bool BuildWaterQuadTree(FWaterQuadTree& WaterQuadTree, const FVector2D& GridPosition) const;
 #if WITH_WATER_SELECTION_SUPPORT
@@ -53,7 +53,7 @@ private:
 	FIntPoint ExtentInTiles = FIntPoint::ZeroValue;
 	float TileSize = 0.0f;
 	int32 TreeDepth = 0;
-	FMaterialRenderProxy* FarDistanceMaterial = nullptr;
+	const UMaterialInterface* FarDistanceMaterial = nullptr;
 	float FarDistanceMeshExtent = 0.0f;
 	double DefaultFarDistanceMeshHeight = 0.0;
 	bool bUseFarMeshWithoutOcean = false;
