[Water] Fixed warning due to mismatching operand types in ternary operator

diff --git a/Engine/Plugins/Experimental/Water/Source/Runtime/Private/WaterMeshComponent.cpp b/Engine/Plugins/Experimental/Water/Source/Runtime/Private/WaterMeshComponent.cpp
--- a/Engine/Plugins/Experimental/Water/Source/Runtime/Private/WaterMeshComponent.cpp
+++ b/Engine/Plugins/Experimental/Water/Source/Runtime/Private/WaterMeshComponent.cpp
@@ -274,7 +274,7 @@ void UWaterMeshComponent::RebuildWaterMesh(float InTileSize, const FIntPoint& In
 	bool bAnyWaterMeshesNotReady = false;
 	const bool bIsGPUQuadTree = CVarWaterMeshGPUQuadTree.GetValueOnGameThread() != 0;
 	
-	UMaterialInterface* FarDistanceMaterialInterface = IsMaterialUsedWithWater(FarDistanceMaterial) ? FarDistanceMaterial : nullptr;
+	UMaterialInterface* FarDistanceMaterialInterface = IsMaterialUsedWithWater(FarDistanceMaterial) ? FarDistanceMaterial.Get() : nullptr;
 	WaterQuadTreeBuilder.Init(WaterZone->GetZoneBounds2D(), InExtentInTiles, InTileSize, FarDistanceMaterialInterface, FarDistanceMeshExtent, FarDistanceMeshHeightWithoutOcean, bUseFarMeshWithoutOcean, bIsGPUQuadTree);
 
 	UsedMaterials.Empty();
