[HWRT] Fix crash due to FWaterMeshSceneProxy moving FRayTracingGeometry in memory.
- problematic since FRayTracingGeometryManager caches pointers to FRayTracingGeometry.
- use TIndirectArray in FWaterMeshSceneProxy to avoid the issue.

#FailedValidations Preflight

diff --git a/Engine/Plugins/Experimental/Water/Source/Runtime/Private/WaterMeshSceneProxy.cpp b/Engine/Plugins/Experimental/Water/Source/Runtime/Private/WaterMeshSceneProxy.cpp
--- a/Engine/Plugins/Experimental/Water/Source/Runtime/Private/WaterMeshSceneProxy.cpp
+++ b/Engine/Plugins/Experimental/Water/Source/Runtime/Private/WaterMeshSceneProxy.cpp
@@ -892,7 +892,7 @@ TArray<int32, TInlineAllocator<8>> FWaterMeshSceneProxy::GetViewToQuadTreeMappin
 #if RHI_RAYTRACING
 void FWaterMeshSceneProxy::SetupRayTracingInstances(FRHICommandListBase& RHICmdList, int32 NumInstances, uint32 DensityIndex)
 {
-	TArray<FRayTracingWaterData>& WaterDataArray = RayTracingWaterData[DensityIndex];
+	TIndirectArray<FRayTracingWaterData>& WaterDataArray = RayTracingWaterData[DensityIndex];
 
 	if (WaterDataArray.Num() > NumInstances)
 	{
@@ -902,7 +902,7 @@ void FWaterMeshSceneProxy::SetupRayTracingInstances(FRHICommandListBase& RHICmdL
 			WaterItem.Geometry.ReleaseResource();
 			WaterItem.DynamicVertexBuffer.Release();
 		}
-		WaterDataArray.SetNum(NumInstances);
+		WaterDataArray.RemoveAt(NumInstances, WaterDataArray.Num() - NumInstances);
 	}	
 
 	if (WaterDataArray.Num() < NumInstances)
@@ -920,12 +920,14 @@ void FWaterMeshSceneProxy::SetupRayTracingInstances(FRHICommandListBase& RHICmdL
 
 		for (int32 Item = StartIndex; Item < NumInstances; Item++)
 		{
-			FRayTracingWaterData& WaterData = WaterDataArray.AddDefaulted_GetRef();
+			FRayTracingWaterData* WaterData = new FRayTracingWaterData;
 
 			Initializer.DebugName = FName(DebugName, Item);
 
-			WaterData.Geometry.SetInitializer(Initializer);
-			WaterData.Geometry.InitResource(RHICmdList);
+			WaterData->Geometry.SetInitializer(Initializer);
+			WaterData->Geometry.InitResource(RHICmdList);
+
+			WaterDataArray.Add(WaterData);
 		}
 	}
 }
diff --git a/Engine/Plugins/Experimental/Water/Source/Runtime/Public/WaterMeshSceneProxy.h b/Engine/Plugins/Experimental/Water/Source/Runtime/Public/WaterMeshSceneProxy.h
--- a/Engine/Plugins/Experimental/Water/Source/Runtime/Public/WaterMeshSceneProxy.h
+++ b/Engine/Plugins/Experimental/Water/Source/Runtime/Public/WaterMeshSceneProxy.h
@@ -200,7 +200,7 @@ private:
 
 #if RHI_RAYTRACING
 	// Per density array of ray tracing geometries.
-	TArray<TArray<FRayTracingWaterData>> RayTracingWaterData;	
+	TArray<TIndirectArray<FRayTracingWaterData>> RayTracingWaterData;	
 #endif
 
 	
