[Water]

Fixed the water info merge pass incorrectly detecting when dilation data is present. Before this fix, all regions would flow through the dilation path.

diff --git a/Engine/Plugins/Experimental/Water/Shaders/Private/WaterInfoMerge.usf b/Engine/Plugins/Experimental/Water/Shaders/Private/WaterInfoMerge.usf
--- a/Engine/Plugins/Experimental/Water/Shaders/Private/WaterInfoMerge.usf
+++ b/Engine/Plugins/Experimental/Water/Shaders/Private/WaterInfoMerge.usf
@@ -60,16 +60,15 @@ void Main(
 	const float WaterSceneDepth = WaterInfo.SceneDepth;
 	const float WaterZ = CaptureZ - WaterSceneDepth;
 
+	const float WaterZMin = WaterHeightExtents.x;
+	const float WaterZMax = WaterHeightExtents.y;
+
 	const float DilationSceneDepth = DilationSample;
-	const bool bPixelHasDilationData = DilationSceneDepth != 0.0f;
-		
 	float DilatedZ = CaptureZ - DilationSceneDepth;
+	const bool bPixelHasDilationData = DilatedZ <= WaterZMax && DilatedZ >= WaterZMin;
 
 	// Always initialize output to avoid bugs in old shader compilers 
 	OutColor = 0.0f;
-
-	const float WaterZMin = WaterHeightExtents.x;
-	const float WaterZMax = WaterHeightExtents.y;
 	
 	// If the measured WaterZ is below WaterZMin it means that there is no water at this pixel because it is not possible
 	// for water to be further than ZMin.
