Fix crash in UWaterSubsystem::FindWaterZone when called on a uninitialized world (coming from package validation).

#Synced-CL 49558965

diff --git a/Engine/Plugins/Experimental/Water/Source/Runtime/Private/WaterSubsystem.cpp b/Engine/Plugins/Experimental/Water/Source/Runtime/Private/WaterSubsystem.cpp
--- a/Engine/Plugins/Experimental/Water/Source/Runtime/Private/WaterSubsystem.cpp
+++ b/Engine/Plugins/Experimental/Water/Source/Runtime/Private/WaterSubsystem.cpp
@@ -26,6 +26,7 @@
 
 #if WITH_EDITOR
 #include "WaterZoneActorDesc.h"
+#include "WorldPartition/WorldPartition.h"
 #include "WorldPartition/WorldPartitionActorDescInstance.h"
 #include "WorldPartition/WorldPartitionHelpers.h"
 extern UNREALED_API UEditorEngine* GEditor;
@@ -642,7 +643,7 @@ TSoftObjectPtr<AWaterZone> UWaterSubsystem::FindWaterZone(const UWorld* World, c
 	// Within the editor, we also want to check unloaded actors to ensure that the water body has serialized the best possible water zone, rather than just looking through what might be loaded now.
 	if (GEditor && !World->IsGameWorld())
 	{
-		if (UWorldPartition* WorldPartition = World->GetWorldPartition())
+		if (UWorldPartition* WorldPartition = World->GetWorldPartition(); WorldPartition && WorldPartition->IsInitialized())
 		{
 			const FBox Bounds3D(FVector(Bounds.Min.X, Bounds.Min.Y, -HALF_WORLD_MAX), FVector(Bounds.Max.X, Bounds.Max.Y, HALF_WORLD_MAX));
 			FWorldPartitionHelpers::ForEachIntersectingActorDescInstance<AWaterZone>(WorldPartition, Bounds3D, [&Bounds, &ViableZones](const FWorldPartitionActorDescInstance* ActorDescInstance)
