Skip water proxy creation if PSOs are still compiling and Always create is not set (no fallback to default material because that would look strange on water)

#Synced-CL 49557757

diff --git a/Engine/Plugins/Experimental/Water/Source/Runtime/Private/WaterMeshComponent.cpp b/Engine/Plugins/Experimental/Water/Source/Runtime/Private/WaterMeshComponent.cpp
--- a/Engine/Plugins/Experimental/Water/Source/Runtime/Private/WaterMeshComponent.cpp
+++ b/Engine/Plugins/Experimental/Water/Source/Runtime/Private/WaterMeshComponent.cpp
@@ -129,6 +129,12 @@ FPrimitiveSceneProxy* UWaterMeshComponent::CreateSceneProxy()
 	{
 		return nullptr;
 	}
+	
+	if (CheckPSOPrecachingAndBoostPriority(EPSOPrecachePriority::Highest) && GetPSOPrecacheProxyCreationStrategy() != EPSOPrecacheProxyCreationStrategy::AlwaysCreate)
+	{
+		UE_LOG(LogWater, Verbose, TEXT("Skipping CreateSceneProxy for UWaterMeshComponent %s (Its PSOs are still compiling)"), *GetFullName());
+		return nullptr;
+	}
 
 	return new FWaterMeshSceneProxy(this);
 }
