[Water]

Force all water mesh data to rebuild to fixup any content which may have been saved during the initial TSet changes

diff --git a/Engine/Plugins/Experimental/Water/Source/Runtime/Private/WaterBodyComponent.cpp b/Engine/Plugins/Experimental/Water/Source/Runtime/Private/WaterBodyComponent.cpp
--- a/Engine/Plugins/Experimental/Water/Source/Runtime/Private/WaterBodyComponent.cpp
+++ b/Engine/Plugins/Experimental/Water/Source/Runtime/Private/WaterBodyComponent.cpp
@@ -1793,10 +1793,6 @@ void UWaterBodyComponent::PostLoad()
 	}
 
 #if WITH_EDITOR
-	if (GetLinkerCustomVersion(FFortniteMainBranchObjectVersion::GUID) < FFortniteMainBranchObjectVersion::WaterBodyStaticMeshComponents)
-	{
-		UpdateWaterBodyRenderData();
-	}
 	RegisterOnUpdateWavesData(GetWaterWaves(), /* bRegister = */true);
 
 	if (IsTemplate() && ((OwningWaterZone != nullptr) || (!OwningWaterZone.GetAssetName().IsEmpty())))
@@ -1822,6 +1818,8 @@ void UWaterBodyComponent::PostLoad()
 		OwningWaterZone.Reset();
 	}
 
+	
+	bool bNotifyModifiedPackage = false;
 	// If the detail mode of the water body component or the water spline component are ever not DM_Low,
 	// depending on per platform project settings they may be culled out in a cooked build for certain platforms
 	// breaking our assumptions that the waterbodycomponent/watersplinecomponent should always be present in game.
@@ -1836,8 +1834,19 @@ void UWaterBodyComponent::PostLoad()
 		{
 			SplineComp->DetailMode = DM_Low;
 		}
+		bNotifyModifiedPackage = true;
+	}
+	
+	if (GetLinkerCustomVersion(FWaterCustomVersion::GUID) < FWaterCustomVersion::RebuildWaterMeshDataTSetChange)
+	{
+		UpdateWaterBodyRenderData();
 
-		// Push the request to the user that they should mark this modified water package as dirty and resave it to persist the DetailMode change.
+		bNotifyModifiedPackage = true;
+	}
+
+	// Push the request to the user that they should mark this modified water package as dirty and resave it to persist the DetailMode change.
+	if (bNotifyModifiedPackage)
+	{
 		const IWaterModuleInterface& WaterModule = FModuleManager::GetModuleChecked<IWaterModuleInterface>("Water");
 		if (IWaterEditorServices* WaterEditorServices = WaterModule.GetWaterEditorServices())
 		{
diff --git a/Engine/Plugins/Experimental/Water/Source/Runtime/Private/WaterBodyMeshBuilder.cpp b/Engine/Plugins/Experimental/Water/Source/Runtime/Private/WaterBodyMeshBuilder.cpp
--- a/Engine/Plugins/Experimental/Water/Source/Runtime/Private/WaterBodyMeshBuilder.cpp
+++ b/Engine/Plugins/Experimental/Water/Source/Runtime/Private/WaterBodyMeshBuilder.cpp
@@ -581,7 +581,7 @@ void FWaterBodyMeshBuilder::UpdateStaticMesh(UStaticMesh* WaterMesh, const FMesh
 	FStaticMeshSourceModel& SrcModel = WaterMesh->GetSourceModel(0);
 	SrcModel.BuildSettings.bRecomputeNormals = true;
 	SrcModel.BuildSettings.bRecomputeTangents = true;
-	SrcModel.BuildSettings.bRemoveDegenerates = false;
+	SrcModel.BuildSettings.bRemoveDegenerates = true;
 	SrcModel.BuildSettings.bUseHighPrecisionTangentBasis = false;
 	SrcModel.BuildSettings.bUseFullPrecisionUVs = bIsConservativeRasterMesh; // CR mesh stores vertex positions in the UVs and needs full 32bit precision
 	SrcModel.BuildSettings.bGenerateLightmapUVs = false;
diff --git a/Engine/Plugins/Experimental/Water/Source/Runtime/Public/WaterVersion.h b/Engine/Plugins/Experimental/Water/Source/Runtime/Public/WaterVersion.h
--- a/Engine/Plugins/Experimental/Water/Source/Runtime/Public/WaterVersion.h
+++ b/Engine/Plugins/Experimental/Water/Source/Runtime/Public/WaterVersion.h
@@ -30,6 +30,8 @@ struct FWaterCustomVersion
 		MoveWaterMPCParamsToWaterMesh,
 		// Moved where the default water info material is assigned to prevent blocking async load thread
 		DefaultWaterInfoMaterialAssignmentMoved,
+		// Rebuild water render data to ensure no meshes from the broken tset change were serialized
+		RebuildWaterMeshDataTSetChange,
 
 		// -----<new versions can be added above this line>-------------------------------------------------
 		VersionPlusOne,
