[Water]

Allow scene capture views within commandlets to have water info. This is intended to unblock world partition minimap renders since they render exclusively via scene captures and there is no opportunity for the water info texture and quadtree to be populated in those cases.

#Synced-CL 49921777

diff --git a/Engine/Plugins/Experimental/Water/Source/Runtime/Private/WaterViewExtension.cpp b/Engine/Plugins/Experimental/Water/Source/Runtime/Private/WaterViewExtension.cpp
--- a/Engine/Plugins/Experimental/Water/Source/Runtime/Private/WaterViewExtension.cpp
+++ b/Engine/Plugins/Experimental/Water/Source/Runtime/Private/WaterViewExtension.cpp
@@ -767,9 +767,10 @@ void FWaterViewExtension::BeginRenderViewFamily(FSceneViewFamily& InViewFamily)
 
 bool FWaterViewExtension::ShouldHaveWaterZoneViewData(const FSceneView& InView) const
 {
-	// Don't dirty the water info texture when we're rendering from a scene capture. Due to the frame delay after marking the texture as dirty, scene captures wouldn't have the right texture anyways.
+	// Don't dirty the water info texture when we're rendering from a scene capture (unless we're in a commandlet in order to support world partition minimap captures).
+	// Due to the frame delay after marking the texture as dirty, scene captures wouldn't have the right texture anyways.
 	// #todo_water [roey]: Once we have no frame-delay for updating the texture and lesser performance impact, we can re-enable updates within scene captures.
-	return !InView.bIsSceneCapture && !InView.bIsSceneCaptureCube && !InView.bIsReflectionCapture && !InView.bIsPlanarReflection && !InView.bIsVirtualTexture 
+	return (!InView.bIsSceneCapture || IsRunningCommandlet()) && !InView.bIsSceneCaptureCube && !InView.bIsReflectionCapture && !InView.bIsPlanarReflection && !InView.bIsVirtualTexture 
 		// Also don't update water info texture when rendering hit proxies as it bypasses custom render passes 
 		&& !InView.Family->EngineShowFlags.HitProxies;
 }
