[Water] Fixed potential crash in FWaterQuadTree::Unlock() which happened to some users with big quadtrees. Fixed by sizing the queue involved in breadth first traversal to match the number of total nodes in the tree instead of the theoretical maximum number of leaf nodes. In practice, this should also save some memory as most quadtrees are sparse.

#Synced-CL 49929750

diff --git a/Engine/Plugins/Experimental/Water/Source/Runtime/Private/WaterQuadTree.cpp b/Engine/Plugins/Experimental/Water/Source/Runtime/Private/WaterQuadTree.cpp
--- a/Engine/Plugins/Experimental/Water/Source/Runtime/Private/WaterQuadTree.cpp
+++ b/Engine/Plugins/Experimental/Water/Source/Runtime/Private/WaterQuadTree.cpp
@@ -614,7 +614,7 @@ void FWaterQuadTree::Unlock(bool bPruneRedundantNodes)
 			int32 NumElements = 0;
 		};
 		
-		FQueue Queue(GetMaxLeafCount());
+		FQueue Queue(FMath::Max(NodeData.Nodes.Num(), 1));
 		uint32 NextBreadthFirstIndex = 0;
 		
 		check(BreadthFirstOrder.IsEmpty());
@@ -631,7 +631,7 @@ void FWaterQuadTree::Unlock(bool bPruneRedundantNodes)
 
 			for (uint32 ChildIndex : Node.Children)
 			{
-				if (ChildIndex > 0)
+				if (ChildIndex > 0 && ensure(NodeData.Nodes.IsValidIndex(ChildIndex)))
 				{
 					Queue.Add(ChildIndex);
 				}
