[Water]

Removed ensure when static mesh is null. Instead, just ignore water bodies with null static meshes quietly and provide a map check warning for users to know when geometry is invalid.

#Synced-CL 50046626

diff --git a/Engine/Plugins/Experimental/Water/Source/Runtime/Private/WaterBodyComponent.cpp b/Engine/Plugins/Experimental/Water/Source/Runtime/Private/WaterBodyComponent.cpp
--- a/Engine/Plugins/Experimental/Water/Source/Runtime/Private/WaterBodyComponent.cpp
+++ b/Engine/Plugins/Experimental/Water/Source/Runtime/Private/WaterBodyComponent.cpp
@@ -1328,6 +1328,16 @@ TArray<TSharedRef<FTokenizedMessage>> UWaterBodyComponent::CheckWaterBodyStatus(
 								FOnActionTokenExecuted::CreateUObject(this, &UWaterBodyComponent::SetWaterBodyStaticMeshEnabled, true), true))
 								);
 					}
+					
+					if (GetWaterInfoMeshComponent() == nullptr || GetWaterInfoMeshComponent()->GetStaticMesh() == nullptr)
+					{
+						Result.Add(FTokenizedMessage::Create(EMessageSeverity::Warning)
+							->AddToken(FUObjectToken::Create(this))
+							->AddToken(FTextToken::Create(FText::Format(
+								LOCTEXT("MapCheck_Message_MissingWaterBodyMesh", "Water body {0} has no mesh information. This is caused by the water spline forming invalid geometry or not having enough spline points to be a valid shape. Adjust the water spline to resolve this issue."),
+								FText::FromString(GetWaterBodyActor()->GetActorLabel()))))
+							);
+					}
 				}
 			}
 
diff --git a/Engine/Plugins/Experimental/Water/Source/Runtime/Private/WaterBodyMeshBuilder.cpp b/Engine/Plugins/Experimental/Water/Source/Runtime/Private/WaterBodyMeshBuilder.cpp
--- a/Engine/Plugins/Experimental/Water/Source/Runtime/Private/WaterBodyMeshBuilder.cpp
+++ b/Engine/Plugins/Experimental/Water/Source/Runtime/Private/WaterBodyMeshBuilder.cpp
@@ -54,6 +54,8 @@ void FWaterBodyMeshBuilder::BuildWaterInfoMeshes(UWaterBodyComponent* WaterBodyC
 	{
 		if (DynamicMesh.TriangleCount() == 0)
 		{
+			// if the generated mesh has no information, clear the previous static mesh.
+			MeshComponent->SetStaticMesh(nullptr);
 			return nullptr;
 		}
 
diff --git a/Engine/Plugins/Experimental/Water/Source/Runtime/Private/WaterMeshComponent.cpp b/Engine/Plugins/Experimental/Water/Source/Runtime/Private/WaterMeshComponent.cpp
--- a/Engine/Plugins/Experimental/Water/Source/Runtime/Private/WaterMeshComponent.cpp
+++ b/Engine/Plugins/Experimental/Water/Source/Runtime/Private/WaterMeshComponent.cpp
@@ -301,9 +301,15 @@ void UWaterMeshComponent::RebuildWaterMesh(float InTileSize, const FIntPoint& In
 
 		UWaterBodyInfoMeshComponent* WaterBodyInfoMeshComponent = WaterBodyComponent->GetWaterInfoMeshComponent();
 		UStaticMesh* StaticMesh = WaterBodyInfoMeshComponent ? WaterBodyInfoMeshComponent->GetStaticMesh().Get() : nullptr;
-		bAnyWaterMeshesNotReady |= StaticMesh && StaticMesh->IsCompiling();
-		FStaticMeshRenderData* StaticMeshRenderData = StaticMesh ? StaticMesh->GetRenderData() : nullptr;
-		if (!ensure(WaterBodyInfoMeshComponent) || !ensure(StaticMesh) || !StaticMeshRenderData)
+
+		if (StaticMesh == nullptr)
+		{
+			return true;
+		}
+
+		bAnyWaterMeshesNotReady |= StaticMesh->IsCompiling();
+		FStaticMeshRenderData* StaticMeshRenderData = StaticMesh->GetRenderData();
+		if (!StaticMeshRenderData)
 		{
 			return true;
 		}
