[Water]

Removed assert in water quad tree when transition materials are missing. The assert can trigger in normal user interaction flows if a material fails to compile while assigning to a lake/ocean transition material.

Instead, we now log a warning and fall back to the primary water material which are guaranteed to be valid water materials or the default surface material.

#Synced-CL 50310527

diff --git a/Engine/Plugins/Experimental/Water/Source/Runtime/Private/WaterQuadTree.cpp b/Engine/Plugins/Experimental/Water/Source/Runtime/Private/WaterQuadTree.cpp
--- a/Engine/Plugins/Experimental/Water/Source/Runtime/Private/WaterQuadTree.cpp
+++ b/Engine/Plugins/Experimental/Water/Source/Runtime/Private/WaterQuadTree.cpp
@@ -28,18 +28,30 @@ void FWaterQuadTree::FNode::AddNodeForRender(const FNodeData& InNodeData, const
 
 		if (TransitionWaterBodyRenderData.IsLake())
 		{
-			check(InWaterBodyRenderData.RiverToLakeMaterial);
+			if (InWaterBodyRenderData.RiverToLakeMaterial)
+			{
+				MaterialIndex = InWaterBodyRenderData.RiverToLakeMaterialIndex;
+			}
+			else
+			{
+				MaterialIndex = InWaterBodyRenderData.MaterialIndex;
+			}
 
-			MaterialIndex = InWaterBodyRenderData.RiverToLakeMaterialIndex;
 			NodeWaterBodyIndex = (uint32)TransitionWaterBodyRenderData.WaterBodyIndex;
 			BaseHeight = TransitionWaterBodyRenderData.SurfaceBaseHeight;
 			TileDebugID = 3;
 		}
 		if (TransitionWaterBodyRenderData.IsOcean())
 		{
-			check(InWaterBodyRenderData.RiverToOceanMaterial);
+			if (InWaterBodyRenderData.RiverToOceanMaterial)
+			{
+				MaterialIndex = InWaterBodyRenderData.RiverToOceanMaterialIndex;
+			}
+			else
+			{
+				MaterialIndex = InWaterBodyRenderData.MaterialIndex;
+			}
 
-			MaterialIndex = InWaterBodyRenderData.RiverToOceanMaterialIndex;
 			NodeWaterBodyIndex = (uint32)TransitionWaterBodyRenderData.WaterBodyIndex;
 			BaseHeight = TransitionWaterBodyRenderData.SurfaceBaseHeight;
 			TileDebugID = 4;
