Optimized UWaterSubsystem::Tick's AWaterZone iteration by avoiding using TActorRange.

Saves 20us per frame on switch.

#tests replayrun, pie

#Synced-CL 50342528

diff --git a/Engine/Plugins/Experimental/Water/Source/Runtime/Private/WaterSubsystem.cpp b/Engine/Plugins/Experimental/Water/Source/Runtime/Private/WaterSubsystem.cpp
--- a/Engine/Plugins/Experimental/Water/Source/Runtime/Private/WaterSubsystem.cpp
+++ b/Engine/Plugins/Experimental/Water/Source/Runtime/Private/WaterSubsystem.cpp
@@ -255,13 +255,11 @@ void UWaterSubsystem::Tick(float DeltaTime)
 	SetMPCTime(MPCTime, PrevWorldTimeSeconds);
 	PrevWorldTimeSeconds = MPCTime;
 
-	for (AWaterZone* WaterZoneActor : TActorRange<AWaterZone>(GetWorld()))
+	WaterBodyManager.ForEachWaterZone([](AWaterZone* WaterZoneActor)
 	{
-		if (WaterZoneActor)
-		{
-			WaterZoneActor->Update();
-		}
-	}
+		WaterZoneActor->Update();
+		return true;
+	});
 
 	if (!bUnderWaterForAudio && CachedDepthUnderwater > 0.0f)
 	{
