[Water]

Account for water zone Z position when calculating the zone bounds.

#Synced-CL 50485134

diff --git a/Engine/Plugins/Experimental/Water/Source/Runtime/Private/WaterZoneActor.cpp b/Engine/Plugins/Experimental/Water/Source/Runtime/Private/WaterZoneActor.cpp
--- a/Engine/Plugins/Experimental/Water/Source/Runtime/Private/WaterZoneActor.cpp
+++ b/Engine/Plugins/Experimental/Water/Source/Runtime/Private/WaterZoneActor.cpp
@@ -211,9 +211,9 @@ FBox AWaterZone::GetZoneBounds() const
 	const FBox2D ZoneBounds2D = GetZoneBounds2D();
 	// #todo_water [roey]: Water zone doesn't have an explicit z bounds yet. For now just use the x or y.
 	const double StreamingBoundsZ = FMath::Max(ZoneExtent.X, ZoneExtent.Y);
+	const FVector ActorLocation = GetActorLocation();
 
-	return FBox(FVector3d(ZoneBounds2D.Min, -StreamingBoundsZ / 2.), FVector3d(ZoneBounds2D.Max, StreamingBoundsZ / 2.0));
-}
+	return FBox(FVector3d(ZoneBounds2D.Min, ActorLocation.Z - StreamingBoundsZ / 2.), FVector3d(ZoneBounds2D.Max, ActorLocation.Z + StreamingBoundsZ / 2.0));}
 
 FIntPoint AWaterZone::GetRenderTargetResolution() const
 {
