Fixed edit condition for water nav area class

#Synced-CL 50661191

diff --git a/Engine/Plugins/Experimental/Water/Source/Runtime/Public/WaterBodyComponent.h b/Engine/Plugins/Experimental/Water/Source/Runtime/Public/WaterBodyComponent.h
--- a/Engine/Plugins/Experimental/Water/Source/Runtime/Public/WaterBodyComponent.h
+++ b/Engine/Plugins/Experimental/Water/Source/Runtime/Public/WaterBodyComponent.h
@@ -687,7 +687,7 @@ protected:
 	FPostProcessSettings CurrentPostProcessSettings;
 
 	// The navigation area class that will be generated on nav mesh
-	UPROPERTY(EditAnywhere, BlueprintReadWrite, Category = Navigation, meta = (EditCondition = "bCanAffectNavigation"))
+	UPROPERTY(EditAnywhere, BlueprintReadWrite, Category = Navigation, meta = (EditCondition = "bCanEverAffectNavigation"))
 	TSubclassOf<UNavAreaBase> WaterNavAreaClass;
 
 	/** If the Water Material assigned to this component has Fixed Depth enabled, this is the depth that is passed. */
