PR #14422: [Water] Fix GetDynamicWaterInfoCenter output
+ Fix missing implementation of GetDynamicWaterInfoBounds

#Synced-CL 50911594

diff --git a/Engine/Plugins/Experimental/Water/Source/Runtime/Private/WaterZoneActor.cpp b/Engine/Plugins/Experimental/Water/Source/Runtime/Private/WaterZoneActor.cpp
--- a/Engine/Plugins/Experimental/Water/Source/Runtime/Private/WaterZoneActor.cpp
+++ b/Engine/Plugins/Experimental/Water/Source/Runtime/Private/WaterZoneActor.cpp
@@ -876,6 +876,23 @@ bool AWaterZone::GetDynamicWaterInfoCenter(int32 PlayerIndex, FVector& OutCenter
 
 	check(World != nullptr);
 
+	bool bHasValidZoneLocation = false;
+	OutCenter = GetActorLocation();
+
+	if (const FWaterViewExtension* WaterViewExtension = UWaterSubsystem::GetWaterViewExtension(World))
+	{
+		bHasValidZoneLocation = WaterViewExtension->GetZoneLocation(this, PlayerIndex, OutCenter);
+	}
+	
+	return bHasValidZoneLocation;
+}
+
+bool AWaterZone::GetDynamicWaterInfoBounds(int32 PlayerIndex, FBox& OutBounds) const
+{
+	const UWorld* World = GetWorld();
+
+	check(World != nullptr);
+
 	bool bHasValidZoneLocation = false;
 	FVector Center = GetActorLocation();
 
@@ -883,10 +900,14 @@ bool AWaterZone::GetDynamicWaterInfoCenter(int32 PlayerIndex, FVector& OutCenter
 	{
 		bHasValidZoneLocation = WaterViewExtension->GetZoneLocation(this, PlayerIndex, Center);
 	}
-	
+
+	const FVector WaterInfoHalfExtents = GetDynamicWaterInfoExtent() / 2.;
+	OutBounds = FBox(Center - WaterInfoHalfExtents, Center + WaterInfoHalfExtents);
+
 	return bHasValidZoneLocation;
 }
 
+
 void AWaterZone::OnLevelAddedToWorld(ULevel* InLevel, UWorld* InWorld)
 {
 	OnLevelChanged(InLevel, InWorld);
