WaterSplineComponentVisualizer: Change check to ensure for selected water spline component to avoid crashing when right-clicking water splines. This is a stop-gap fix till the actual bug is figured out/fixed as to why the spline component selection gets cleared in this case.

#Synced-CL 50934100

diff --git a/Engine/Plugins/Experimental/Water/Source/Editor/Private/WaterSplineComponentVisualizer.cpp b/Engine/Plugins/Experimental/Water/Source/Editor/Private/WaterSplineComponentVisualizer.cpp
--- a/Engine/Plugins/Experimental/Water/Source/Editor/Private/WaterSplineComponentVisualizer.cpp
+++ b/Engine/Plugins/Experimental/Water/Source/Editor/Private/WaterSplineComponentVisualizer.cpp
@@ -661,7 +661,7 @@ void FWaterSplineComponentVisualizer::GenerateContextMenuSections(FMenuBuilder&
 {
 	FSplineComponentVisualizer::GenerateContextMenuSections(InMenuBuilder);
 
-	check(GetEditedWaterSplineComponent());
+	ensure(GetEditedWaterSplineComponent());
 	if (const UWaterSplineMetadata* Metadata = GetEditedWaterSplineMetaData())
 	{
 		InMenuBuilder.BeginSection("Water Visualization", LOCTEXT("Water Visualization", "Water Visualization"));
