[Water]

Check generated water info mesh data for validity when cooking.

#Synced-CL 50734623

diff --git a/Engine/Plugins/Experimental/Water/Source/Runtime/Private/WaterBodyMeshBuilder.cpp b/Engine/Plugins/Experimental/Water/Source/Runtime/Private/WaterBodyMeshBuilder.cpp
--- a/Engine/Plugins/Experimental/Water/Source/Runtime/Private/WaterBodyMeshBuilder.cpp
+++ b/Engine/Plugins/Experimental/Water/Source/Runtime/Private/WaterBodyMeshBuilder.cpp
@@ -29,6 +29,7 @@
 #include "UObject/Package.h"
 #include "StaticMeshCompiler.h"
 #include "Algo/AnyOf.h"
+#include "WaterModule.h"
 #endif // WITH_EDITOR
 
 #if WITH_EDITOR
@@ -46,6 +47,19 @@ void FWaterBodyMeshBuilder::BuildWaterInfoMeshes(UWaterBodyComponent* WaterBodyC
 	FDynamicMesh3 WaterInfoMesh(EMeshComponents::None);
 	FDynamicMesh3 WaterInfoDilatedMesh(EMeshComponents::None);
 	GetDynamicMesh(WaterBodyComponent, WaterInfoMesh, &WaterInfoDilatedMesh, bMakeWaterInfoMeshConservativeRasterCompatible);
+	
+	if (IsRunningCookCommandlet())
+	{
+		if (!WaterInfoMesh.CheckValidity(FDynamicMesh3::FValidityOptions(), EValidityCheckFailMode::ReturnOnly))
+		{
+			UE_LOG(LogWater, Warning, TEXT("Mesh for %s failed validation"), *GetNameSafe(WaterBodyComponent));
+		}
+
+		if (!WaterInfoDilatedMesh.CheckValidity(FDynamicMesh3::FValidityOptions(), EValidityCheckFailMode::ReturnOnly))
+		{
+			UE_LOG(LogWater, Warning, TEXT("Dilated Mesh for %s failed validation"), *GetNameSafe(WaterBodyComponent));
+		}
+	}
 
 	UMaterialInterface* WaterInfoMID = WaterBodyComponent->GetWaterInfoMaterialInstance();
 	UObject* Outer = WaterBodyComponent->GetOwner();
