[Water]

Made water zone override editable.
It is too easy via simple user interactions to have water bodies point to water zones in other maps, creating indesirable cross map references

#Synced-CL 51102749

diff --git a/Engine/Plugins/Experimental/Water/Source/Runtime/Public/WaterBodyComponent.h b/Engine/Plugins/Experimental/Water/Source/Runtime/Public/WaterBodyComponent.h
--- a/Engine/Plugins/Experimental/Water/Source/Runtime/Public/WaterBodyComponent.h
+++ b/Engine/Plugins/Experimental/Water/Source/Runtime/Public/WaterBodyComponent.h
@@ -677,7 +677,8 @@ protected:
 	UPROPERTY(Transient)
 	mutable TWeakObjectPtr<ALandscapeProxy> Landscape;
 
-	UPROPERTY(Category = Water, VisibleInstanceOnly, AdvancedDisplay, TextExportTransient)
+	/** Internal - Automatically points to the detected owning water zone. Should not be edited manually. */
+	UPROPERTY(Category = Water, EditAnywhere, AdvancedDisplay, TextExportTransient)
 	TSoftObjectPtr<AWaterZone> OwningWaterZone;
 
 	UPROPERTY(Category = Water, EditAnywhere, BlueprintReadOnly, AdvancedDisplay)
