new base class for baked water data

#Synced-CL 51253254

diff --git a/Engine/Plugins/Experimental/Buoyancy/Source/Runtime/Private/BuoyancySubsystem.cpp b/Engine/Plugins/Experimental/Buoyancy/Source/Runtime/Private/BuoyancySubsystem.cpp
--- a/Engine/Plugins/Experimental/Buoyancy/Source/Runtime/Private/BuoyancySubsystem.cpp
+++ b/Engine/Plugins/Experimental/Buoyancy/Source/Runtime/Private/BuoyancySubsystem.cpp
@@ -1125,7 +1125,7 @@ void FBuoyancySubsystemSimCallback::TrackInteraction(
 		FVector ClosestWaterPointToParticle = ParticlePos;
 		ClosestWaterPointToParticle.Z = WaterHeight;
 		
-		const TWeakObjectPtr<UShallowWaterSimulationData> ShallowWaterSimData = WaterSpline.ShallowWaterSimData;
+		const TWeakObjectPtr<UShallowWaterSimulationDataBase> ShallowWaterSimData = WaterSpline.ShallowWaterSimData;
 
 		WaterSampler = TSharedPtr<FBuoyancyWaterSampler>(new FBuoyancyShallowWaterSampler(ShallowWaterSimData, ClosestWaterPointToParticle, WaterVelocity));
 	}
diff --git a/Engine/Plugins/Experimental/Buoyancy/Source/Runtime/Public/BuoyancyParticleData.h b/Engine/Plugins/Experimental/Buoyancy/Source/Runtime/Public/BuoyancyParticleData.h
--- a/Engine/Plugins/Experimental/Buoyancy/Source/Runtime/Public/BuoyancyParticleData.h
+++ b/Engine/Plugins/Experimental/Buoyancy/Source/Runtime/Public/BuoyancyParticleData.h
@@ -53,7 +53,7 @@ protected:
 class FBuoyancyShallowWaterSampler : public FBuoyancyWaterSampler
 {
 public:
-	FBuoyancyShallowWaterSampler(const TWeakObjectPtr<UShallowWaterSimulationData> InWaterSampler, const FVector& InClosestWaterPointToParticle, const FVector &InAverageVelocity) :
+	FBuoyancyShallowWaterSampler(const TWeakObjectPtr<UShallowWaterSimulationDataBase> InWaterSampler, const FVector& InClosestWaterPointToParticle, const FVector &InAverageVelocity) :
 		WaterSampler(InWaterSampler)
 	{
 		SamplerType = EWaterSamplerType::ShallowWater;
@@ -102,7 +102,7 @@ public:
 	}
 
 private:
-	const TWeakObjectPtr<UShallowWaterSimulationData> WaterSampler;
+	const TWeakObjectPtr<UShallowWaterSimulationDataBase> WaterSampler;
 };
 
 // sampler for spline rivers, but we are only doing 1 spline sample per interaction to 
diff --git a/Engine/Plugins/Experimental/Buoyancy/Source/Runtime/Public/BuoyancyWaterSplineData.h b/Engine/Plugins/Experimental/Buoyancy/Source/Runtime/Public/BuoyancyWaterSplineData.h
--- a/Engine/Plugins/Experimental/Buoyancy/Source/Runtime/Public/BuoyancyWaterSplineData.h
+++ b/Engine/Plugins/Experimental/Buoyancy/Source/Runtime/Public/BuoyancyWaterSplineData.h
@@ -13,7 +13,7 @@ struct FBuoyancyWaterSplineData
 		const EWaterBodyType InBodyType,
 		const TOptional<FInterpCurveFloat>& InWidth,
 		const TOptional<FInterpCurveFloat>& InVelocity,
-		const TWeakObjectPtr<UShallowWaterSimulationData> InShallowWaterSimData
+		const TWeakObjectPtr<UShallowWaterSimulationDataBase> InShallowWaterSimData
 	)
 		: Transform(InTransform)
 		, Position(InPosition)
@@ -40,5 +40,5 @@ struct FBuoyancyWaterSplineData
 
 	// Only some water bodies can have a baked shallow water sim representation, which
 	// overrides the splines for water height/depth/velocity/normal
-	const TWeakObjectPtr<UShallowWaterSimulationData> ShallowWaterSimData;
+	const TWeakObjectPtr<UShallowWaterSimulationDataBase> ShallowWaterSimData;
 };
\ No newline at end of file
diff --git a/Engine/Plugins/Experimental/Water/Source/Runtime/Private/BakedShallowWaterSimulationComponent.cpp b/Engine/Plugins/Experimental/Water/Source/Runtime/Private/BakedShallowWaterSimulationComponent.cpp
--- a/Engine/Plugins/Experimental/Water/Source/Runtime/Private/BakedShallowWaterSimulationComponent.cpp
+++ b/Engine/Plugins/Experimental/Water/Source/Runtime/Private/BakedShallowWaterSimulationComponent.cpp
@@ -10,7 +10,7 @@ UBakedShallowWaterSimulationComponent::UBakedShallowWaterSimulationComponent(con
 	: Super(ObjectInitializer)
 {}
 
-void UShallowWaterSimulationData::SampleShallowWaterSimulationAtIndex(const FVector2D& QueryFloatIndex, FVector& OutWaterVelocity, float& OutWaterHeight, float& OutWaterDepth) const
+void UShallowWaterSimulationDataBase::SampleShallowWaterSimulationAtIndex(const FVector2D& QueryFloatIndex, FVector& OutWaterVelocity, float& OutWaterHeight, float& OutWaterDepth) const
 {
 	OutWaterHeight = 0;
 	OutWaterDepth = 0;
@@ -18,15 +18,15 @@ void UShallowWaterSimulationData::SampleShallowWaterSimulationAtIndex(const FVec
 	
 	if (QueryFloatIndex.X > 0 && QueryFloatIndex.X < NumCells.X - 1 && QueryFloatIndex.Y > 0 && QueryFloatIndex.Y < NumCells.Y - 1)
 	{
-		check(ArrayValues.Num() == NumCells.X * NumCells.Y);
+		check(GetNumAllocatedCells() == NumCells.X * NumCells.Y);
 
 		const FIntVector2 BaseIndex = FIntVector2(QueryFloatIndex.X, QueryFloatIndex.Y);
 		const FVector2D LerpValue = QueryFloatIndex - FVector2D(BaseIndex.X, BaseIndex.Y);
 
-		const FVector4 Sample00 = ArrayValues[(BaseIndex.X + 0) + (BaseIndex.Y + 0) * NumCells.X];
-		const FVector4 Sample10 = ArrayValues[(BaseIndex.X + 1) + (BaseIndex.Y + 0) * NumCells.X];
-		const FVector4 Sample01 = ArrayValues[(BaseIndex.X + 0) + (BaseIndex.Y + 1) * NumCells.X];
-		const FVector4 Sample11 = ArrayValues[(BaseIndex.X + 1) + (BaseIndex.Y + 1) * NumCells.X];
+		const FVector4 Sample00 = QueryLowLevelGridAtIndex(BaseIndex);
+		const FVector4 Sample10 = QueryLowLevelGridAtIndex(BaseIndex + FIntVector2(1, 0));
+		const FVector4 Sample01 = QueryLowLevelGridAtIndex(BaseIndex + FIntVector2(0, 1));
+		const FVector4 Sample11 = QueryLowLevelGridAtIndex(BaseIndex + FIntVector2(1, 1));
 		
 		const FVector4 Sample0 = Sample00 * (1. - LerpValue.X) + Sample10 * (LerpValue.X);
 		const FVector4 Sample1 = Sample01 * (1. - LerpValue.X) + Sample11 * (LerpValue.X);
@@ -39,7 +39,7 @@ void UShallowWaterSimulationData::SampleShallowWaterSimulationAtIndex(const FVec
 	}
 }
 
-void UShallowWaterSimulationData::QueryShallowWaterSimulationAtIndex(const FIntVector2& QueryIndex, FVector& OutWaterVelocity, float& OutWaterHeight, float& OutWaterDepth) const
+void UShallowWaterSimulationDataBase::QueryShallowWaterSimulationAtIndex(const FIntVector2& QueryIndex, FVector& OutWaterVelocity, float& OutWaterHeight, float& OutWaterDepth) const
 {
 	OutWaterHeight = 0;
 	OutWaterDepth = 0;
@@ -47,9 +47,9 @@ void UShallowWaterSimulationData::QueryShallowWaterSimulationAtIndex(const FIntV
 
 	if (QueryIndex.X >= 0 && QueryIndex.X < NumCells.X && QueryIndex.Y >= 0 && QueryIndex.Y < NumCells.Y)
 	{
-		check(ArrayValues.Num() == NumCells.X * NumCells.Y);
+		check(GetNumAllocatedCells() == NumCells.X * NumCells.Y);
 
-		const FVector4 Sample = ArrayValues[QueryIndex.X + QueryIndex.Y * NumCells.X];
+		const FVector4 Sample = QueryLowLevelGridAtIndex(QueryIndex);
 
 		OutWaterHeight = Sample.X + Position.Z;
 		OutWaterDepth = Sample.Y;
@@ -57,7 +57,7 @@ void UShallowWaterSimulationData::QueryShallowWaterSimulationAtIndex(const FIntV
 	}
 }
 
-FVector UShallowWaterSimulationData::ComputeShallowWaterSimulationNormalAtPosition(const FVector& QueryPos) const
+FVector UShallowWaterSimulationDataBase::ComputeShallowWaterSimulationNormalAtPosition(const FVector& QueryPos) const
 {
 	const FVector2D FloatIndexPos = WorldToFloatIndex(QueryPos);
 
@@ -108,11 +108,14 @@ void UBakedShallowWaterSimulationComponent::PostLoad()
 		if (BakedSimulationData == nullptr && SimulationData_DEPRECATED.IsValid())
 		{
 			BakedSimulationData = NewObject<UShallowWaterSimulationData>(this);
-			BakedSimulationData->ArrayValues = MoveTemp(SimulationData_DEPRECATED.ArrayValues);
-			BakedSimulationData->NumCells = SimulationData_DEPRECATED.NumCells;
-			BakedSimulationData->Position = SimulationData_DEPRECATED.Position;
-			BakedSimulationData->Size = SimulationData_DEPRECATED.Size;
-			BakedSimulationData->BakedTexture = SimulationData_DEPRECATED.BakedTexture;
+
+			UShallowWaterSimulationData *BakedSimulationDataDerived = Cast<UShallowWaterSimulationData>(BakedSimulationData);
+
+			BakedSimulationDataDerived->ArrayValues = MoveTemp(SimulationData_DEPRECATED.ArrayValues);
+			BakedSimulationDataDerived->NumCells = SimulationData_DEPRECATED.NumCells;
+			BakedSimulationDataDerived->Position = SimulationData_DEPRECATED.Position;
+			BakedSimulationDataDerived->Size = SimulationData_DEPRECATED.Size;
+			BakedSimulationDataDerived->BakedTexture = SimulationData_DEPRECATED.BakedTexture;
 
 			SimulationData_DEPRECATED = FShallowWaterSimulationGrid_DEPRECATED();
 
diff --git a/Engine/Plugins/Experimental/Water/Source/Runtime/Public/BakedShallowWaterSimulationComponent.h b/Engine/Plugins/Experimental/Water/Source/Runtime/Public/BakedShallowWaterSimulationComponent.h
--- a/Engine/Plugins/Experimental/Water/Source/Runtime/Public/BakedShallowWaterSimulationComponent.h
+++ b/Engine/Plugins/Experimental/Water/Source/Runtime/Public/BakedShallowWaterSimulationComponent.h
@@ -3,7 +3,7 @@
 #pragma once
 
 #include "CoreMinimal.h"
-#include "Math/Float16Color.h"
+#include "Math/Float16.h"
 #include "Components/PrimitiveComponent.h"
 #include "BakedShallowWaterSimulationComponent.generated.h"
 
@@ -12,13 +12,10 @@
 class AWaterBody;
 
 UCLASS(MinimalAPI)
-class UShallowWaterSimulationData : public UObject 
+class UShallowWaterSimulationDataBase : public UObject 
 {
 	GENERATED_BODY()
 public:
-	/** Raw grid data */
-	UPROPERTY()
-	TArray<FVector4> ArrayValues;
 	
 	/** Number of cells */
 	UPROPERTY()
@@ -39,7 +36,7 @@ public:
 	/** Check for a valid sim domain */
 	bool HasValidData() const
 	{
-		return NumCells.X > 0 && NumCells.Y > 0 && ArrayValues.Num() == NumCells.X * NumCells.Y && Size.X > SMALL_NUMBER && Size.Y > SMALL_NUMBER;
+		return NumCells.X > 0 && NumCells.Y > 0 && GetNumAllocatedCells() == NumCells.X * NumCells.Y && Size.X > SMALL_NUMBER && Size.Y > SMALL_NUMBER;
 	}
 
 	/** Linear interpolate sim data at an index */
@@ -51,6 +48,10 @@ public:
 	/** Compute normal at a world position */
 	WATER_API FVector ComputeShallowWaterSimulationNormalAtPosition(const FVector &QueryPos) const;
 	
+	virtual int GetNumAllocatedCells() const { return 0; };
+	
+	virtual FVector4 QueryLowLevelGridAtIndex(const FIntVector2 Index) const { return FVector4(0,0,0,0); }
+	
 	/** Query shallow water grid cell values at a world position */
 	void QueryShallowWaterSimulationAtPosition(const FVector &QueryPos, FVector& OutWaterVelocity, float& OutWaterHeight, float& OutWaterDepth) const
 	{
@@ -115,6 +116,50 @@ public:
 	}
 };
 
+UCLASS(MinimalAPI)
+class UShallowWaterSimulationData : public UShallowWaterSimulationDataBase
+{	
+	GENERATED_BODY()
+
+public:
+	/** Raw grid data */
+	UPROPERTY()
+	TArray<FVector4> ArrayValues;
+
+	virtual int GetNumAllocatedCells() const
+	{
+		return ArrayValues.Num();
+	}
+
+	virtual FVector4 QueryLowLevelGridAtIndex(const FIntVector2 Index) const
+	{
+		return ArrayValues[Index.X + Index.Y * NumCells.X];
+	}
+};
+
+UCLASS(MinimalAPI)
+class UShallowWaterSimulationData16 : public UShallowWaterSimulationDataBase
+{	
+	GENERATED_BODY()
+
+public:
+		
+	/** Raw grid data */
+	UPROPERTY()
+	TArray<uint16> ArrayValues;
+
+	virtual int GetNumAllocatedCells() const
+	{
+		return ArrayValues.Num() / 4;
+	}
+
+	virtual FVector4 QueryLowLevelGridAtIndex(const FIntVector2 Index) const
+	{
+		const FFloat16Color* CurrVal = reinterpret_cast<const FFloat16Color*>(&ArrayValues[(Index.X + Index.Y * NumCells.X) * 4]);
+		return FVector4(CurrVal->R, CurrVal->G, CurrVal->B, CurrVal->A);
+	}
+};
+
 USTRUCT(BlueprintType)
 struct FShallowWaterSimulationGrid_DEPRECATED
 {
@@ -169,7 +214,7 @@ public:
 #endif
 
 	UPROPERTY()
-	TObjectPtr<UShallowWaterSimulationData> BakedSimulationData;
+	TObjectPtr<UShallowWaterSimulationDataBase> BakedSimulationData;
 
 	UPROPERTY()
 	TSet<TSoftObjectPtr<AWaterBody>> WaterBodies;
diff --git a/Engine/Plugins/Experimental/WaterAdvanced/Source/WaterAdvanced/Private/ShallowWaterRiverActor.cpp b/Engine/Plugins/Experimental/WaterAdvanced/Source/WaterAdvanced/Private/ShallowWaterRiverActor.cpp
--- a/Engine/Plugins/Experimental/WaterAdvanced/Source/WaterAdvanced/Private/ShallowWaterRiverActor.cpp
+++ b/Engine/Plugins/Experimental/WaterAdvanced/Source/WaterAdvanced/Private/ShallowWaterRiverActor.cpp
@@ -886,27 +886,6 @@ void UShallowWaterRiverComponent::Bake()
 	TArray<FFloat16Color> TmpShallowWaterSimArrayValues;
 	BakedWaterSurfaceRT->GameThread_GetRenderTargetResource()->ReadFloat16Pixels(TmpShallowWaterSimArrayValues);
 
-	TArray<FVector4> ShallowWaterSimArrayValues;
-	ShallowWaterSimArrayValues.Empty();
-	ShallowWaterSimArrayValues.AddZeroed(TmpShallowWaterSimArrayValues.Num());
-
-	// cast all values to floats
-	int Index = 0;
-	for (FFloat16Color Val : TmpShallowWaterSimArrayValues)
-	{
-		const float WaterHeight = Val.R;
-		const float WaterDepth = Val.G;
-		const FVector2D WaterVelocity(Val.B, Val.A);
-
-		FVector4 FloatVal;
-		FloatVal.X = WaterHeight;
-		FloatVal.Y = WaterDepth;
-		FloatVal.Z = WaterVelocity.X;
-		FloatVal.W = WaterVelocity.Y;
-
-		ShallowWaterSimArrayValues[Index++] = FloatVal;
-	}	
-
 	// bake foam and other data to texture
 	BakedFoamTexture = BakedFoamRT->ConstructTexture2D(this, "BakedFoamTexture", TextureObjectFlags);
 	
@@ -947,12 +926,17 @@ void UShallowWaterRiverComponent::Bake()
 
 	BakedSim = NewObject<UBakedShallowWaterSimulationComponent>(this, NAME_None, RF_Public);
 	
-	BakedSim->BakedSimulationData = NewObject<UShallowWaterSimulationData>(BakedSim, NAME_None, RF_Public);	
-	BakedSim->BakedSimulationData->ArrayValues = ShallowWaterSimArrayValues;
-	BakedSim->BakedSimulationData->NumCells = FIntVector2(BakedWaterSurfaceRT->SizeX, BakedWaterSurfaceRT->SizeY);
-	BakedSim->BakedSimulationData->Position = SystemPos;
-	BakedSim->BakedSimulationData->Size = WorldGridSize;
-	BakedSim->BakedSimulationData->BakedTexture = BakedWaterSurfaceTexture;
+	TObjectPtr<UShallowWaterSimulationData16> BakedSimulationData = NewObject<UShallowWaterSimulationData16>(BakedSim, NAME_None, RF_Public);	
+	BakedSim->BakedSimulationData = BakedSimulationData;
+
+	BakedSimulationData->ArrayValues.SetNum(TmpShallowWaterSimArrayValues.Num() * 4);
+	FMemory::Memcpy(BakedSimulationData->ArrayValues.GetData(), 
+					TmpShallowWaterSimArrayValues.GetData(), TmpShallowWaterSimArrayValues.Num() * sizeof(FFloat16Color));
+
+	BakedSimulationData->NumCells = FIntVector2(BakedWaterSurfaceRT->SizeX, BakedWaterSurfaceRT->SizeY);
+	BakedSimulationData->Position = SystemPos;
+	BakedSimulationData->Size = WorldGridSize;
+	BakedSimulationData->BakedTexture = BakedWaterSurfaceTexture;
 
 
 	BakedSim->WaterBodies = AllWaterBodies;	
