simple sparse grid

#virtualized
#Synced-CL 51253254

diff --git a/Engine/Plugins/Experimental/Water/Source/Runtime/Public/BakedShallowWaterSimulationComponent.h b/Engine/Plugins/Experimental/Water/Source/Runtime/Public/BakedShallowWaterSimulationComponent.h
--- a/Engine/Plugins/Experimental/Water/Source/Runtime/Public/BakedShallowWaterSimulationComponent.h
+++ b/Engine/Plugins/Experimental/Water/Source/Runtime/Public/BakedShallowWaterSimulationComponent.h
@@ -36,7 +36,7 @@ public:
 	/** Check for a valid sim domain */
 	bool HasValidData() const
 	{
-		return NumCells.X > 0 && NumCells.Y > 0 && GetNumAllocatedCells() == NumCells.X * NumCells.Y && Size.X > SMALL_NUMBER && Size.Y > SMALL_NUMBER;
+		return NumCells.X > 0 && NumCells.Y > 0 && GetTotalNumCells() == NumCells.X * NumCells.Y && Size.X > SMALL_NUMBER && Size.Y > SMALL_NUMBER;
 	}
 
 	/** Linear interpolate sim data at an index */
@@ -48,8 +48,12 @@ public:
 	/** Compute normal at a world position */
 	WATER_API FVector ComputeShallowWaterSimulationNormalAtPosition(const FVector &QueryPos) const;
 	
-	virtual int GetNumAllocatedCells() const { return 0; };
-	
+	virtual int GetTotalNumCells() const { return 0; }
+	virtual int GetNumAllocatedCells() const { return 0; }
+	virtual int GetMemoryInBytes() const { return 0; }
+
+	virtual void Build(const TArray<FFloat16Color>& InputDenseGrid, FIntVector2 InDenseNumCells) {};
+
 	virtual FVector4 QueryLowLevelGridAtIndex(const FIntVector2 Index) const { return FVector4(0,0,0,0); }
 	
 	/** Query shallow water grid cell values at a world position */
@@ -126,11 +130,37 @@ public:
 	UPROPERTY()
 	TArray<FVector4> ArrayValues;
 
-	virtual int GetNumAllocatedCells() const
+	virtual int GetTotalNumCells() const
 	{
 		return ArrayValues.Num();
 	}
 
+	virtual int GetNumAllocatedCells() const
+	{
+		return GetTotalNumCells();
+	}
+
+	virtual int GetMemoryInBytes() const
+	{
+		return ArrayValues.NumBytes();
+	}
+	
+	virtual void Build(const TArray<FFloat16Color>& InputDenseGrid, FIntVector2 InDenseNumCells)
+	{
+		check(InDenseNumCells.X * InDenseNumCells.Y == InputDenseGrid.Num());
+
+		NumCells = InDenseNumCells;
+
+		ArrayValues.SetNum(InputDenseGrid.Num());
+
+		int idx = 0;
+		for (FFloat16Color CurrCell : InputDenseGrid)
+		{
+			ArrayValues[idx++] = FVector4(CurrCell.R.GetFloat(), CurrCell.G.GetFloat(),CurrCell.B.GetFloat(),CurrCell.A.GetFloat());
+		}
+		// should loop over cells
+	}
+
 	virtual FVector4 QueryLowLevelGridAtIndex(const FIntVector2 Index) const
 	{
 		return ArrayValues[Index.X + Index.Y * NumCells.X];
@@ -148,11 +178,31 @@ public:
 	UPROPERTY()
 	TArray<uint16> ArrayValues;
 
-	virtual int GetNumAllocatedCells() const
+	virtual int GetTotalNumCells() const
 	{
 		return ArrayValues.Num() / 4;
 	}
 
+	virtual int GetNumAllocatedCells() const
+	{
+		return GetTotalNumCells();
+	}
+
+	virtual int GetMemoryInBytes() const
+	{
+		return ArrayValues.GetAllocatedSize();
+	}
+
+	virtual void Build(const TArray<FFloat16Color>& InputDenseGrid, FIntVector2 InDenseNumCells)
+	{
+		check(InDenseNumCells.X * InDenseNumCells.Y == InputDenseGrid.Num());
+
+		NumCells = InDenseNumCells;
+
+		ArrayValues.SetNum(InputDenseGrid.Num() * 4);
+		FMemory::Memcpy(ArrayValues.GetData(), InputDenseGrid.GetData(), InputDenseGrid.Num() * sizeof(FFloat16Color));
+	}
+
 	virtual FVector4 QueryLowLevelGridAtIndex(const FIntVector2 Index) const
 	{
 		const FFloat16Color* CurrVal = reinterpret_cast<const FFloat16Color*>(&ArrayValues[(Index.X + Index.Y * NumCells.X) * 4]);
@@ -160,6 +210,98 @@ public:
 	}
 };
 
+UCLASS(MinimalAPI)
+class UShallowWaterSimulationDataSparse : public UShallowWaterSimulationDataBase
+{	
+	GENERATED_BODY()
+
+public:
+		
+	/** Raw grid data */
+	UPROPERTY()
+	TArray<uint16> BlockValues;
+
+	UPROPERTY()
+	TMap<FIntVector2, int32> GridIndexToBlockIndex;
+
+	virtual int GetTotalNumCells() const
+	{
+		return NumCells.X * NumCells.Y;
+	}
+
+	virtual int GetNumAllocatedCells() const
+	{
+		return BlockValues.Num() / 4;
+	}
+
+	virtual int GetMemoryInBytes() const
+	{
+		return BlockValues.GetAllocatedSize() + GridIndexToBlockIndex.GetAllocatedSize();
+	}
+
+	virtual FVector4 QueryLowLevelGridAtIndex(const FIntVector2 Index) const
+	{
+		const FIntVector2 BlockIndex(Index.X / BlockSize, Index.Y / BlockSize);
+
+		const int32 *CurrBlockDataIndex = GridIndexToBlockIndex.Find(BlockIndex);
+
+		if (!CurrBlockDataIndex)
+		{
+			return FVector4(0,0,0,0);
+		}
+
+		const int32 BlockStartIndex = (*CurrBlockDataIndex);
+		
+		const FIntVector2 LocalBlockIndex(Index.X % BlockSize, Index.Y % BlockSize);
+		const int32 LocalBlockIndexLinear = (LocalBlockIndex.X + LocalBlockIndex.Y * BlockSize) * 4;
+
+		const FFloat16Color* CurrVal = reinterpret_cast<const FFloat16Color*>(&BlockValues[BlockStartIndex + LocalBlockIndexLinear]);
+		return FVector4(CurrVal->R, CurrVal->G, CurrVal->B, CurrVal->A);
+	}
+
+	virtual void Build(const TArray<FFloat16Color>& InputDenseGrid, FIntVector2 InDenseNumCells)
+	{
+		check(InDenseNumCells.X * InDenseNumCells.Y == InputDenseGrid.Num());
+
+		NumCells = InDenseNumCells;
+
+		BlockValues.Empty();
+		GridIndexToBlockIndex.Empty();
+
+		for (int32 y = 0; y < InDenseNumCells.Y; ++y) {
+		for (int32 x = 0; x < InDenseNumCells.X; ++x) {
+			FIntVector2 CurrCellIndex(x,y);
+			
+			const FFloat16Color &CurrVal = InputDenseGrid[x + y * InDenseNumCells.X];
+
+			// check if depth is close to zero
+			if (FMath::IsNearlyZero(CurrVal.G.GetFloat(), 1e-5))				
+			{
+				continue;
+			}
+
+			const FIntVector2 BlockIndex(CurrCellIndex.X / BlockSize, CurrCellIndex.Y / BlockSize);
+
+			int32 &BlockStartIndex = GridIndexToBlockIndex.FindOrAdd(BlockIndex, -1);
+
+			if (BlockStartIndex == -1)
+			{
+				BlockStartIndex = BlockValues.Num();
+				BlockValues.AddZeroed(BlockSizeInBytes);				
+			}
+					
+			const FIntVector2 LocalBlockIndex(CurrCellIndex.X % BlockSize, CurrCellIndex.Y % BlockSize);
+			const int32 LocalBlockIndexLinear = (LocalBlockIndex.X + LocalBlockIndex.Y * BlockSize) * 4;
+
+			FMemory::Memcpy(&BlockValues[BlockStartIndex + LocalBlockIndexLinear], &CurrVal, sizeof(FFloat16Color));
+		}}
+	}
+private:
+	static constexpr int32 BlockSize = 16;
+	static constexpr int32 DataSizeInBytes = 4;
+	static constexpr int32 BlockSizeInBytes = BlockSize * BlockSize * 4;
+};
+
 USTRUCT(BlueprintType)
 struct FShallowWaterSimulationGrid_DEPRECATED
 {
diff --git a/Engine/Plugins/Experimental/WaterAdvanced/Source/WaterAdvanced/Private/ShallowWaterRiverActor.cpp b/Engine/Plugins/Experimental/WaterAdvanced/Source/WaterAdvanced/Private/ShallowWaterRiverActor.cpp
--- a/Engine/Plugins/Experimental/WaterAdvanced/Source/WaterAdvanced/Private/ShallowWaterRiverActor.cpp
+++ b/Engine/Plugins/Experimental/WaterAdvanced/Source/WaterAdvanced/Private/ShallowWaterRiverActor.cpp
@@ -926,21 +926,23 @@ void UShallowWaterRiverComponent::Bake()
 
 	BakedSim = NewObject<UBakedShallowWaterSimulationComponent>(this, NAME_None, RF_Public);
 	
-	TObjectPtr<UShallowWaterSimulationData16> BakedSimulationData = NewObject<UShallowWaterSimulationData16>(BakedSim, NAME_None, RF_Public);	
+	TObjectPtr<UShallowWaterSimulationDataSparse> BakedSimulationData = NewObject<UShallowWaterSimulationDataSparse>(BakedSim, NAME_None, RF_Public);	
 	BakedSim->BakedSimulationData = BakedSimulationData;
 
-	BakedSimulationData->ArrayValues.SetNum(TmpShallowWaterSimArrayValues.Num() * 4);
-	FMemory::Memcpy(BakedSimulationData->ArrayValues.GetData(), 
-					TmpShallowWaterSimArrayValues.GetData(), TmpShallowWaterSimArrayValues.Num() * sizeof(FFloat16Color));
+	BakedSimulationData->Build(TmpShallowWaterSimArrayValues, FIntVector2(BakedWaterSurfaceRT->SizeX, BakedWaterSurfaceRT->SizeY));
+	
+	if (BakedSimulationData->GetTotalNumCells() > 0)
+	{
+		UE_LOG(LogShallowWater, Log, TEXT("Baked Shallow Water Simulation: Sparseness: %f"), ((float) BakedSimulationData->GetNumAllocatedCells()) / BakedSimulationData->GetTotalNumCells());
+	}
+	UE_LOG(LogShallowWater, Log, TEXT("Baked Shallow Water Simulation: CPU Memory Usage: %d"), BakedSimulationData->GetMemoryInBytes());
 
-	BakedSimulationData->NumCells = FIntVector2(BakedWaterSurfaceRT->SizeX, BakedWaterSurfaceRT->SizeY);
 	BakedSimulationData->Position = SystemPos;
 	BakedSimulationData->Size = WorldGridSize;
 	BakedSimulationData->BakedTexture = BakedWaterSurfaceTexture;
 
-
 	BakedSim->WaterBodies = AllWaterBodies;	
-		
+			
 	// compute the maximum water height for each convex in each water body simulated by this river
 	// we use this to modify the collision geometry so it fully encompasses the baked water sim
 	TMap<FKConvexElem*, float> ConvexToMaxHeight;
