update checks for numcells

#Synced-CL 51253254

diff --git a/Engine/Plugins/Experimental/Water/Source/Runtime/Private/BakedShallowWaterSimulationComponent.cpp b/Engine/Plugins/Experimental/Water/Source/Runtime/Private/BakedShallowWaterSimulationComponent.cpp
--- a/Engine/Plugins/Experimental/Water/Source/Runtime/Private/BakedShallowWaterSimulationComponent.cpp
+++ b/Engine/Plugins/Experimental/Water/Source/Runtime/Private/BakedShallowWaterSimulationComponent.cpp
@@ -18,7 +18,7 @@ void UShallowWaterSimulationDataBase::SampleShallowWaterSimulationAtIndex(const
 	
 	if (QueryFloatIndex.X > 0 && QueryFloatIndex.X < NumCells.X - 1 && QueryFloatIndex.Y > 0 && QueryFloatIndex.Y < NumCells.Y - 1)
 	{
-		check(GetNumAllocatedCells() == NumCells.X * NumCells.Y);
+		check(GetTotalNumCells() == NumCells.X * NumCells.Y);
 
 		const FIntVector2 BaseIndex = FIntVector2(QueryFloatIndex.X, QueryFloatIndex.Y);
 		const FVector2D LerpValue = QueryFloatIndex - FVector2D(BaseIndex.X, BaseIndex.Y);
@@ -47,7 +47,7 @@ void UShallowWaterSimulationDataBase::QueryShallowWaterSimulationAtIndex(const F
 
 	if (QueryIndex.X >= 0 && QueryIndex.X < NumCells.X && QueryIndex.Y >= 0 && QueryIndex.Y < NumCells.Y)
 	{
-		check(GetNumAllocatedCells() == NumCells.X * NumCells.Y);
+		check(GetTotalNumCells() == NumCells.X * NumCells.Y);
 
 		const FVector4 Sample = QueryLowLevelGridAtIndex(QueryIndex);
 
