- Fix uninitialized variables

#rnx

#Synced-CL 51499580

diff --git a/Engine/Plugins/Experimental/Water/Source/Runtime/Public/BakedShallowWaterSimulationComponent.h b/Engine/Plugins/Experimental/Water/Source/Runtime/Public/BakedShallowWaterSimulationComponent.h
--- a/Engine/Plugins/Experimental/Water/Source/Runtime/Public/BakedShallowWaterSimulationComponent.h
+++ b/Engine/Plugins/Experimental/Water/Source/Runtime/Public/BakedShallowWaterSimulationComponent.h
@@ -325,19 +325,19 @@ public:
 	
 	/** Number of cells */
 	UPROPERTY()
-	FIntVector2 NumCells;
+	FIntVector2 NumCells = FIntVector2::ZeroValue;
 	
 	/** World space center */
 	UPROPERTY()
-	FVector Position;
+	FVector Position = FVector::ZeroVector;
 	
 	/** World space size */
 	UPROPERTY()
-	FVector2D Size;
+	FVector2D Size = FVector2D::ZeroVector;
 
 	/** World space grid cell size */
 	UPROPERTY()
-	FVector2D Dx;
+	FVector2D Dx = FVector2D::ZeroVector;
 	
 	/** Texture for baked water sim data */
 	UPROPERTY()
