[Water]

Moved deprecation of the water physical material out of Serialize and into postload. SetPhysMaterial/GetPhysMaterial are not safe to call outside of the gamethread.

#Synced-CL 51667108

#Synced-CL 51667108

diff --git a/Engine/Plugins/Experimental/Water/Source/Runtime/Private/WaterBodyComponent.cpp b/Engine/Plugins/Experimental/Water/Source/Runtime/Private/WaterBodyComponent.cpp
--- a/Engine/Plugins/Experimental/Water/Source/Runtime/Private/WaterBodyComponent.cpp
+++ b/Engine/Plugins/Experimental/Water/Source/Runtime/Private/WaterBodyComponent.cpp
@@ -1781,6 +1781,13 @@ void UWaterBodyComponent::PostLoad()
 
 	DeprecateData();
 
+#if WITH_EDITORONLY_DATA
+	if (GetLinkerCustomVersion(FFortniteMainBranchObjectVersion::GUID) < FFortniteMainBranchObjectVersion::WaterBodyPhysicalMaterialPropertyRemoval)
+	{
+		SetPhysMaterialOverride(PhysicalMaterial_DEPRECATED);
+	}
+#endif // WITH_EDITORONLY_DATA
+
 	if ((IsComponentPSOPrecachingEnabled() && RHISupportsManualVertexFetch(GMaxRHIShaderPlatform)))
 	{
 		FPSOPrecacheParams PrecachePSOParams;
@@ -1846,7 +1853,7 @@ void UWaterBodyComponent::PostLoad()
 		}
 		bNotifyModifiedPackage = true;
 	}
-	
+
 	if (GetLinkerCustomVersion(FWaterCustomVersion::GUID) < FWaterCustomVersion::RebuildWaterMeshDataTSetChange)
 	{
 		UpdateWaterBodyRenderData();
@@ -1892,11 +1899,6 @@ void UWaterBodyComponent::DeprecateData()
 	{
 		WaterStaticMeshMaterial = WaterLODMaterial_DEPRECATED;
 	}
-
-	if (GetLinkerCustomVersion(FFortniteMainBranchObjectVersion::GUID) < FFortniteMainBranchObjectVersion::WaterBodyPhysicalMaterialPropertyRemoval)
-	{
-		SetPhysMaterialOverride(PhysicalMaterial_DEPRECATED);
-	}
 #endif // WITH_EDITORONLY_DATA
 }
 
