Add more missing includes for non-unity builds

#rnx
#Synced-CL 51867312

diff --git a/Engine/Plugins/Experimental/Water/Source/Runtime/Public/WaterVertexFactory.inl b/Engine/Plugins/Experimental/Water/Source/Runtime/Public/WaterVertexFactory.inl
--- a/Engine/Plugins/Experimental/Water/Source/Runtime/Public/WaterVertexFactory.inl
+++ b/Engine/Plugins/Experimental/Water/Source/Runtime/Public/WaterVertexFactory.inl
@@ -10,7 +10,7 @@
 #include "WaterInstanceDataBuffer.h"
 #include "DataDrivenShaderPlatformInfo.h"
 #include "StereoRenderUtils.h"
-
+#include "ShaderCompilerCore.h"
 
 // ----------------------------------------------------------------------------------
 
diff --git a/Engine/Source/Runtime/Engine/Private/PSOPrecacheValidation.cpp b/Engine/Source/Runtime/Engine/Private/PSOPrecacheValidation.cpp
--- a/Engine/Source/Runtime/Engine/Private/PSOPrecacheValidation.cpp
+++ b/Engine/Source/Runtime/Engine/Private/PSOPrecacheValidation.cpp
@@ -10,6 +10,7 @@
 #include "UnrealEngine.h"
 #include "Materials/Material.h"
 #include "PSOPrecacheMaterial.h"
+#include "HAL/FileManager.h"
 
 #if PSO_PRECACHING_VALIDATE && UE_WITH_PSO_PRECACHING
 
diff --git a/Engine/Source/Runtime/Landscape/Internal/LandscapeGrassWeightExporter.h b/Engine/Source/Runtime/Landscape/Internal/LandscapeGrassWeightExporter.h
--- a/Engine/Source/Runtime/Landscape/Internal/LandscapeGrassWeightExporter.h
+++ b/Engine/Source/Runtime/Landscape/Internal/LandscapeGrassWeightExporter.h
@@ -5,6 +5,7 @@
 #include "LandscapeComponent.h"
 #include "Templates/ValueOrError.h"
 #include "MeshMaterialShader.h"
+#include "RenderGraphFwd.h"
 
 class ALandscape;
 class ALandscapeProxy;
diff --git a/Engine/Source/Runtime/Landscape/Private/LandscapeGrassWeightExporter.cpp b/Engine/Source/Runtime/Landscape/Private/LandscapeGrassWeightExporter.cpp
--- a/Engine/Source/Runtime/Landscape/Private/LandscapeGrassWeightExporter.cpp
+++ b/Engine/Source/Runtime/Landscape/Private/LandscapeGrassWeightExporter.cpp
@@ -21,6 +21,8 @@
 #include "Materials/MaterialExpressionLandscapeGrassOutput.h"
 #include "PSOPrecacheSettings.h"
 #include "ProfilingDebugging/RealtimeGPUProfiler.h"
+#include "RenderGraph.h"
+#include "RHICommandList.h"
 
 #if WITH_EDITOR
 #include "Misc/FileHelper.h"
