[SingleLayerWater] Fixed water not rendering when using local tessellation on the water zone and having a SceneCapture in the level. The SceneCapture would consume the bAnyQuadTreeUpdateRequired flag, but not actually update the quadtree because it fails ShouldHaveWaterZoneViewData().

#Synced-CL 52378035

diff --git a/Engine/Plugins/Experimental/Water/Source/Runtime/Private/WaterViewExtension.cpp b/Engine/Plugins/Experimental/Water/Source/Runtime/Private/WaterViewExtension.cpp
--- a/Engine/Plugins/Experimental/Water/Source/Runtime/Private/WaterViewExtension.cpp
+++ b/Engine/Plugins/Experimental/Water/Source/Runtime/Private/WaterViewExtension.cpp
@@ -654,6 +654,8 @@ void FWaterViewExtension::BeginRenderViewFamily(FSceneViewFamily& InViewFamily)
 		return;
 	}
 
+	bool bHadWaterZoneViewData = false;
+
 	for (const FSceneView* View : InViewFamily.Views)
 	{
 		// Warning: Do not capture View in ENQUEUE_RENDER_COMMAND lambdas, since the RT's view hasn't been created yet,
@@ -661,6 +663,8 @@ void FWaterViewExtension::BeginRenderViewFamily(FSceneViewFamily& InViewFamily)
 
 		if (ShouldHaveWaterZoneViewData(*View))
 		{
+			bHadWaterZoneViewData = true;
+
 			if (!bAnyQuadTreeUpdateRequired)
 			{
 				continue;
@@ -762,7 +766,13 @@ void FWaterViewExtension::BeginRenderViewFamily(FSceneViewFamily& InViewFamily)
 
 	}
 
-	bAnyQuadTreeUpdateRequired = false;
+	// Only reset the flag if this view family contained views that could consume it.
+	// Otherwise, a scene capture's view family (which has no water zone view data) would
+	// clear the flag before the main camera's view family gets a chance to process it.
+	if (bHadWaterZoneViewData)
+	{
+		bAnyQuadTreeUpdateRequired = false;
+	}
 }
 
 bool FWaterViewExtension::ShouldHaveWaterZoneViewData(const FSceneView& InView) const
