Convert FUnderwaterPostProcessVolume to a UObject-based UUnderwaterPostProcessVolume so it's lifetime can be tracked with TWeakInterfacePtr

#rnx
#Synced-CL 52429873

diff --git a/Engine/Plugins/Experimental/Water/Source/Runtime/Private/WaterSubsystem.cpp b/Engine/Plugins/Experimental/Water/Source/Runtime/Private/WaterSubsystem.cpp
--- a/Engine/Plugins/Experimental/Water/Source/Runtime/Private/WaterSubsystem.cpp
+++ b/Engine/Plugins/Experimental/Water/Source/Runtime/Private/WaterSubsystem.cpp
@@ -315,8 +315,10 @@ void UWaterSubsystem::Initialize(FSubsystemCollectionBase& Collection)
 #endif //WITH_EDITOR
 	ApplyRuntimeSettings(GetDefault<UWaterRuntimeSettings>(), EPropertyChangeType::ValueSet);
 
+	UnderwaterPostProcessVolume = NewObject<UUnderwaterPostProcessVolume>(this);
+
 	World->OnBeginPostProcessSettings.AddUObject(this, &UWaterSubsystem::ComputeUnderwaterPostProcess);
-	World->InsertPostProcessVolume(&UnderwaterPostProcessVolume);
+	World->InsertPostProcessVolume(UnderwaterPostProcessVolume);
 	{
 		FActorSpawnParameters SpawnInfo;
 		SpawnInfo.SpawnCollisionHandlingOverride = ESpawnActorCollisionHandlingMethod::AlwaysSpawn;
@@ -372,7 +374,7 @@ void UWaterSubsystem::Deinitialize()
 	CVarWaterMeshEnableRendering->SetOnChangedCallback(NullCallback);
 
 	World->OnBeginPostProcessSettings.RemoveAll(this);
-	World->RemovePostProcessVolume(&UnderwaterPostProcessVolume);
+	World->RemovePostProcessVolume(UnderwaterPostProcessVolume);
 
 	WaterBodyManager.Deinitialize();
 
@@ -809,8 +811,8 @@ void UWaterSubsystem::ComputeUnderwaterPostProcess(FVector ViewLocation, FSceneV
 	CachedDepthUnderwater = -1;
 
 	// Set all that needs to be set before an eventual early-out
-	UnderwaterPostProcessVolume.PostProcessProperties.bIsEnabled = false;
-	UnderwaterPostProcessVolume.PostProcessProperties.Settings = nullptr;
+	UnderwaterPostProcessVolume->PostProcessProperties.bIsEnabled = false;
+	UnderwaterPostProcessVolume->PostProcessProperties.Settings = nullptr;
 	SceneView->UnderwaterDepth = CachedDepthUnderwater;
 	SceneView->WaterIntersection = EViewWaterIntersection::OutsideWater;
 
@@ -983,7 +985,7 @@ void UWaterSubsystem::ComputeUnderwaterPostProcess(FVector ViewLocation, FSceneV
 			if (bUnderwaterForPostProcess)
 			{
 				CachedDepthUnderwater = FMath::Max(LocalDepthUnderwater, CachedDepthUnderwater);
-				UnderwaterPostProcessVolume.PostProcessProperties = Query.WaterBodyComponent.GetPostProcessProperties();
+				UnderwaterPostProcessVolume->PostProcessProperties = Query.WaterBodyComponent.GetPostProcessProperties();
 
 #if !(UE_BUILD_SHIPPING || UE_BUILD_TEST)
 				UnderwaterPostProcessDebugInfo.ActiveWaterBodyComponent = &Query.WaterBodyComponent;
diff --git a/Engine/Plugins/Experimental/Water/Source/Runtime/Public/WaterSubsystem.h b/Engine/Plugins/Experimental/Water/Source/Runtime/Public/WaterSubsystem.h
--- a/Engine/Plugins/Experimental/Water/Source/Runtime/Public/WaterSubsystem.h
+++ b/Engine/Plugins/Experimental/Water/Source/Runtime/Public/WaterSubsystem.h
@@ -35,9 +35,13 @@ namespace UE::WaterInfo { struct FRenderingContext; }
 
 bool IsWaterEnabled(bool bIsRenderThread);
 
-struct FUnderwaterPostProcessVolume : public IInterface_PostProcessVolume
+UCLASS(Transient, MinimalAPI)
+class UUnderwaterPostProcessVolume : public UObject, public IInterface_PostProcessVolume
 {
-	FUnderwaterPostProcessVolume()
+	GENERATED_BODY()
+
+public:
+	UUnderwaterPostProcessVolume()
 		: PostProcessProperties()
 	{
 		PostProcessProperties.VolumeGuid = FGuid(0x857a3d0e, 0x222e4c40, 0x833b5e81, 0x499b8d95);
@@ -269,7 +273,8 @@ private:
 	UPROPERTY()
 	TObjectPtr<UMaterialParameterCollection> MaterialParameterCollection;
 
-	FUnderwaterPostProcessVolume UnderwaterPostProcessVolume;
+	UPROPERTY()
+	TObjectPtr<UUnderwaterPostProcessVolume> UnderwaterPostProcessVolume;
 
 	FWaterBodyManager WaterBodyManager;
 
