refactor GerstnerWave evaluation to avoid duplicated code in various Buoyancy implementations

#Synced-CL 52523678

diff --git a/Engine/Plugins/Experimental/Water/Source/Runtime/Private/GerstnerWaterWaves.cpp b/Engine/Plugins/Experimental/Water/Source/Runtime/Private/GerstnerWaterWaves.cpp
--- a/Engine/Plugins/Experimental/Water/Source/Runtime/Private/GerstnerWaterWaves.cpp
+++ b/Engine/Plugins/Experimental/Water/Source/Runtime/Private/GerstnerWaterWaves.cpp
@@ -1,6 +1,7 @@
 // Copyright Epic Games, Inc. All Rights Reserved.
 
 #include "GerstnerWaterWaves.h"
+#include "GerstnerWaveEvaluation.h"
 #include "Engine/Engine.h"
 #include "GerstnerWaterWaveSubsystem.h"
 #include "Math/RandomStream.h"
@@ -61,139 +62,21 @@ void UGerstnerWaterWaves::PostEditChangeProperty(struct FPropertyChangedEvent& P
 
 float UGerstnerWaterWaves::GetWaveHeightAtPosition(const FVector& InPosition, float InWaterDepth, float InTime, FVector& OutNormal) const
 {
-	float WaveHeight = 0.f;
-
-	FVector SummedNormal(ForceInitToZero);
-
 	// Use the offset of the normalized tile as world position to match the shader behavior (see GerstnerWaveFunctions.ush).
 	FVector WorldPosition(FLargeWorldRenderPosition(InPosition).GetOffset());
-
-	for (const FGerstnerWave& Params : GetGerstnerWaves())
-	{
-		float FirstOffset1D;
-		FVector FirstNormal;
-		FVector FirstOffset = GetWaveOffsetAtPosition(Params, WorldPosition, InTime, FirstNormal, FirstOffset1D);
-
-		// Only non-zero steepness requires a second sample.
-		if (Params.Q != 0)
-		{
-			//Approximate wave height by taking two samples on each side of the current sample position and lerping.
-			//Keep one query point fixed since sampling is going to move the points - if on the left half of wavelength, only add a right offset query point and vice-versa.
-			//Choose q as the factor to offset by (max horizontal displacement).
-			//Lerp between the two sampled heights and normals.
-			const float TwoPi = 2 * PI;
-			const float WaveTime = Params.WaveSpeed * InTime;
-			float Position1D = FVector2D::DotProduct(FVector2D(WorldPosition.X, WorldPosition.Y), Params.WaveVector) - WaveTime;
-			float MappedPosition1D = Position1D >= 0.f ? FMath::Fmod(Position1D, TwoPi) : TwoPi - FMath::Abs(FMath::Fmod(Position1D, TwoPi)); //get positive modulos from negative numbers too
-
-			FVector SecondNormal;
-			float SecondOffset1D;
-			FVector GuessOffset;
-			if (MappedPosition1D < PI)
-			{
-				GuessOffset = Params.Direction * Params.Q;
-			}
-			else
-			{
-				GuessOffset = -Params.Direction * Params.Q;
-			}
-			const FVector GuessPosition = WorldPosition + GuessOffset;
-			FVector SecondOffset = GetWaveOffsetAtPosition(Params, GuessPosition, InTime, SecondNormal, SecondOffset1D);
-			SecondOffset1D += (MappedPosition1D < PI) ? Params.Q : -Params.Q;
-			if (!(MappedPosition1D < PI))
-			{
-				Swap<FVector>(FirstOffset, SecondOffset);
-				Swap<float>(FirstOffset1D, SecondOffset1D);
-				Swap<FVector>(FirstNormal, SecondNormal);
-			}
-			const float LerpDenominator = (SecondOffset1D - FirstOffset1D);
-			float LerpVal = (0 - FirstOffset1D) / (LerpDenominator > 0.f ? LerpDenominator : 1.f);
-			const float FinalHeight = FMath::Lerp(FirstOffset.Z, SecondOffset.Z, LerpVal);
-			const FVector WaveNormal = FMath::Lerp(FirstNormal, SecondNormal, LerpVal);
-
-			SummedNormal += WaveNormal;
-			WaveHeight += FinalHeight;
-		}
-		else
-		{
-			SummedNormal += FirstNormal;
-			WaveHeight += FirstOffset.Z;
-		}
-	}
-	SummedNormal.Z = 1.0f - SummedNormal.Z;
-	OutNormal = SummedNormal.GetSafeNormal();
-
-	return WaveHeight;
+	return GerstnerWaveEvaluation::GetWaveHeightAtPosition(GetGerstnerWaves(), WorldPosition, InTime, OutNormal, /*bBlendLWCTiles=*/true);
 }
 
 float UGerstnerWaterWaves::GetSimpleWaveHeightAtPosition(const FVector& InPosition, float InWaterDepth, float InTime) const
 {
-	float WaveHeight = 0.f;
 	// Use the offset of the normalized tile as world position to match the shader behavior (see GerstnerWaveFunctions.ush).
 	FVector WorldPosition(FLargeWorldRenderPosition(InPosition).GetOffset());
-
-	for (const FGerstnerWave& Wave : GetGerstnerWaves())
-	{
-		WaveHeight += GetSimpleWaveOffsetAtPosition(Wave, WorldPosition, InTime);
-	}
-
-	return WaveHeight;
+	return GerstnerWaveEvaluation::GetSimpleWaveHeightAtPosition(GetGerstnerWaves(), WorldPosition, InTime, /*bBlendLWCTiles=*/true);
 }
 
 float UGerstnerWaterWaves::GetWaveAttenuationFactor(const FVector& InPosition, float InWaterDepth, float InTargetWaveMaskDepth) const
 {
-	const float StrengthCoefficient = FMath::Exp(-FMath::Max(InWaterDepth, 0.0f) / (InTargetWaveMaskDepth / 2.0f));
-	return FMath::Clamp(1.f - StrengthCoefficient, 0.f, 1.f);
-}
-
-FVector UGerstnerWaterWaves::GetWaveOffsetAtPosition(const FGerstnerWave& InWaveParams, const FVector& InPosition, float InTime, FVector& OutNormal, float& OutOffset1D) const
-{
-	const float WaveTime = InWaveParams.WaveSpeed * InTime;
-	const float WavePosition = FVector2D::DotProduct(FVector2D(InPosition.X, InPosition.Y), InWaveParams.WaveVector) - WaveTime;
-
-	float WaveSin = 0, WaveCos = 0;
-	FMath::SinCos(&WaveSin, &WaveCos, WavePosition);
-	BlendWaveBetweenLWCTiles(InWaveParams, InPosition, InTime, WaveSin, WaveCos);
-
-	FVector Offset;
-	OutOffset1D = -InWaveParams.Q * WaveSin;
-	Offset.X = OutOffset1D * InWaveParams.Direction.X;
-	Offset.Y = OutOffset1D * InWaveParams.Direction.Y;
-	Offset.Z = WaveCos * InWaveParams.Amplitude;
-
-	OutNormal = FVector(WaveSin * InWaveParams.WKA * InWaveParams.Direction.X, WaveSin * InWaveParams.WKA * InWaveParams.Direction.Y, /*WaveCos*InWaveParams.WKA*(InWaveParams.Steepness / InWaveParams.WKA)*/0.f); //match the material
-
-	return Offset;
-}
-
-float UGerstnerWaterWaves::GetSimpleWaveOffsetAtPosition(const FGerstnerWave& InWaveParams, const FVector& InPosition, float InTime) const
-{
-	const float WaveTime = InWaveParams.WaveSpeed * InTime;
-	const float WavePosition = FVector2D::DotProduct(FVector2D(InPosition.X, InPosition.Y), InWaveParams.WaveVector) - WaveTime;
-	float WaveCos = FMath::Cos(WavePosition);
-	float WaveSinDummy = 0.0f;
-	BlendWaveBetweenLWCTiles(InWaveParams, InPosition, InTime, WaveSinDummy, WaveCos);
-	const float HeightOffset = WaveCos * InWaveParams.Amplitude;
-	return HeightOffset;
-}
-
-void UGerstnerWaterWaves::BlendWaveBetweenLWCTiles(const FGerstnerWave& InWaveParams, const FVector& InPosition, float InTime, float& WaveSin, float& WaveCos) const
-{
-	const FVector TileBorderDist = FVector(FLargeWorldRenderScalar::GetTileSize() * 0.5) - InPosition.GetAbs();
-	const double BlendZoneWidth = 400.0; // Blend over a range of 4 meters on each side of the tile
-	if (TileBorderDist.X < BlendZoneWidth || TileBorderDist.Y < BlendZoneWidth)
-	{
-		const FVector2D BlendWorldPos = FVector2D(TileBorderDist.X, TileBorderDist.Y);
-		const double BlendAlpha = FMath::Clamp(BlendWorldPos.X / BlendZoneWidth, 0.0, 1.0) * FMath::Clamp(BlendWorldPos.Y / BlendZoneWidth, 0.0, 1.0);
-
-		const float WaveTime = InWaveParams.WaveSpeed * InTime;
-		const float BlendWavePos = FVector2D::DotProduct(BlendWorldPos, InWaveParams.WaveVector) - WaveTime;
-		float BlendWaveSin = 0.0f;
-		float BlendWaveCos = 0.0f;
-		FMath::SinCos(&BlendWaveSin, &BlendWaveCos, BlendWavePos);
-		WaveSin = FMath::Lerp(BlendWaveSin, WaveSin, BlendAlpha);
-		WaveCos = FMath::Lerp(BlendWaveCos, WaveCos, BlendAlpha);
-	}
+	return GerstnerWaveEvaluation::GetWaveAttenuationFactor(InWaterDepth, InTargetWaveMaskDepth);
 }
 
 void UGerstnerWaterWaves::RecomputeWaves(bool bAllowBPScript)
diff --git a/Engine/Plugins/Experimental/Water/Source/Runtime/Private/GerstnerWaveEvaluation.cpp b/Engine/Plugins/Experimental/Water/Source/Runtime/Private/GerstnerWaveEvaluation.cpp
new file mode 100644
--- /dev/null
+++ b/Engine/Plugins/Experimental/Water/Source/Runtime/Private/GerstnerWaveEvaluation.cpp
@@ -0,0 +1,154 @@
+// Copyright Epic Games, Inc. All Rights Reserved.
+
+#include "GerstnerWaveEvaluation.h"
+#include "GerstnerWaterWaves.h"
+#include "Misc/LargeWorldRenderPosition.h"
+
+namespace
+{
+	/** Blend wave sin/cos near LWC tile boundaries to avoid seams. */
+	void BlendWaveBetweenLWCTiles(const FGerstnerWave& InWaveParams, const FVector& InPosition, float InTime, float& WaveSin, float& WaveCos)
+	{
+		const FVector TileBorderDist = FVector(FLargeWorldRenderScalar::GetTileSize() * 0.5) - InPosition.GetAbs();
+		const double BlendZoneWidth = 400.0; // Blend over a range of 4 meters on each side of the tile
+		if (TileBorderDist.X < BlendZoneWidth || TileBorderDist.Y < BlendZoneWidth)
+		{
+			const FVector2D BlendWorldPos = FVector2D(TileBorderDist.X, TileBorderDist.Y);
+			const double BlendAlpha = FMath::Clamp(BlendWorldPos.X / BlendZoneWidth, 0.0, 1.0) * FMath::Clamp(BlendWorldPos.Y / BlendZoneWidth, 0.0, 1.0);
+
+			const float WaveTime = InWaveParams.WaveSpeed * InTime;
+			const float BlendWavePos = FVector2D::DotProduct(BlendWorldPos, InWaveParams.WaveVector) - WaveTime;
+			float BlendWaveSin = 0.0f;
+			float BlendWaveCos = 0.0f;
+			FMath::SinCos(&BlendWaveSin, &BlendWaveCos, BlendWavePos);
+			WaveSin = FMath::Lerp(BlendWaveSin, WaveSin, BlendAlpha);
+			WaveCos = FMath::Lerp(BlendWaveCos, WaveCos, BlendAlpha);
+		}
+	}
+}
+
+namespace GerstnerWaveEvaluation
+{
+	FVector GetWaveOffsetAtPosition(const FGerstnerWave& InWaveParams, const FVector& InPosition, float InTime,
+		FVector& OutNormal, float& OutOffset1D, bool bBlendLWCTiles)
+	{
+		const float WaveTime = InWaveParams.WaveSpeed * InTime;
+		const float WavePosition = FVector2D::DotProduct(FVector2D(InPosition.X, InPosition.Y), InWaveParams.WaveVector) - WaveTime;
+
+		float WaveSin = 0, WaveCos = 0;
+		FMath::SinCos(&WaveSin, &WaveCos, WavePosition);
+
+		if (bBlendLWCTiles)
+		{
+			BlendWaveBetweenLWCTiles(InWaveParams, InPosition, InTime, WaveSin, WaveCos);
+		}
+
+		FVector Offset;
+		OutOffset1D = -InWaveParams.Q * WaveSin;
+		Offset.X = OutOffset1D * InWaveParams.Direction.X;
+		Offset.Y = OutOffset1D * InWaveParams.Direction.Y;
+		Offset.Z = WaveCos * InWaveParams.Amplitude;
+
+		OutNormal = FVector(WaveSin * InWaveParams.WKA * InWaveParams.Direction.X, WaveSin * InWaveParams.WKA * InWaveParams.Direction.Y, /*WaveCos*InWaveParams.WKA*(InWaveParams.Steepness / InWaveParams.WKA)*/0.f); //match the material
+
+		return Offset;
+	}
+
+	float GetSimpleWaveOffsetAtPosition(const FGerstnerWave& InWaveParams, const FVector& InPosition, float InTime, bool bBlendLWCTiles)
+	{
+		const float WaveTime = InWaveParams.WaveSpeed * InTime;
+		const float WavePosition = FVector2D::DotProduct(FVector2D(InPosition.X, InPosition.Y), InWaveParams.WaveVector) - WaveTime;
+		float WaveCos = FMath::Cos(WavePosition);
+
+		if (bBlendLWCTiles)
+		{
+			float WaveSinDummy = 0.0f;
+			BlendWaveBetweenLWCTiles(InWaveParams, InPosition, InTime, WaveSinDummy, WaveCos);
+		}
+
+		const float HeightOffset = WaveCos * InWaveParams.Amplitude;
+		return HeightOffset;
+	}
+
+	float GetWaveHeightAtPosition(TArrayView<const FGerstnerWave> WaveParams, const FVector& InPosition, float InTime,
+		FVector& OutNormal, bool bBlendLWCTiles)
+	{
+		float WaveHeight = 0.f;
+		FVector SummedNormal(ForceInitToZero);
+
+		for (const FGerstnerWave& Params : WaveParams)
+		{
+			float FirstOffset1D;
+			FVector FirstNormal;
+			FVector FirstOffset = GetWaveOffsetAtPosition(Params, InPosition, InTime, FirstNormal, FirstOffset1D, bBlendLWCTiles);
+
+			// Only non-zero steepness requires a second sample.
+			if (Params.Q != 0)
+			{
+				//Approximate wave height by taking two samples on each side of the current sample position and lerping.
+				//Keep one query point fixed since sampling is going to move the points - if on the left half of wavelength, only add a right offset query point and vice-versa.
+				//Choose q as the factor to offset by (max horizontal displacement).
+				//Lerp between the two sampled heights and normals.
+				const float TwoPi = 2 * PI;
+				const float WaveTime = Params.WaveSpeed * InTime;
+				float Position1D = FVector2D::DotProduct(FVector2D(InPosition.X, InPosition.Y), Params.WaveVector) - WaveTime;
+				float MappedPosition1D = Position1D >= 0.f ? FMath::Fmod(Position1D, TwoPi) : TwoPi - FMath::Abs(FMath::Fmod(Position1D, TwoPi)); //get positive modulos from negative numbers too
+
+				FVector SecondNormal;
+				float SecondOffset1D;
+				FVector GuessOffset;
+				if (MappedPosition1D < PI)
+				{
+					GuessOffset = Params.Direction * Params.Q;
+				}
+				else
+				{
+					GuessOffset = -Params.Direction * Params.Q;
+				}
+				const FVector GuessPosition = InPosition + GuessOffset;
+				FVector SecondOffset = GetWaveOffsetAtPosition(Params, GuessPosition, InTime, SecondNormal, SecondOffset1D, bBlendLWCTiles);
+				SecondOffset1D += (MappedPosition1D < PI) ? Params.Q : -Params.Q;
+				if (!(MappedPosition1D < PI))
+				{
+					Swap<FVector>(FirstOffset, SecondOffset);
+					Swap<float>(FirstOffset1D, SecondOffset1D);
+					Swap<FVector>(FirstNormal, SecondNormal);
+				}
+				const float LerpDenominator = (SecondOffset1D - FirstOffset1D);
+				float LerpVal = (0 - FirstOffset1D) / (LerpDenominator > 0.f ? LerpDenominator : 1.f);
+				const float FinalHeight = FMath::Lerp(FirstOffset.Z, SecondOffset.Z, LerpVal);
+				const FVector WaveNormal = FMath::Lerp(FirstNormal, SecondNormal, LerpVal);
+
+				SummedNormal += WaveNormal;
+				WaveHeight += FinalHeight;
+			}
+			else
+			{
+				SummedNormal += FirstNormal;
+				WaveHeight += FirstOffset.Z;
+			}
+		}
+		SummedNormal.Z = 1.0f - SummedNormal.Z;
+		OutNormal = SummedNormal.GetSafeNormal();
+
+		return WaveHeight;
+	}
+
+	float GetSimpleWaveHeightAtPosition(TArrayView<const FGerstnerWave> WaveParams, const FVector& InPosition, float InTime, bool bBlendLWCTiles)
+	{
+		float WaveHeight = 0.f;
+
+		for (const FGerstnerWave& Wave : WaveParams)
+		{
+			WaveHeight += GetSimpleWaveOffsetAtPosition(Wave, InPosition, InTime, bBlendLWCTiles);
+		}
+
+		return WaveHeight;
+	}
+
+	float GetWaveAttenuationFactor(float InWaterDepth, float InTargetWaveMaskDepth)
+	{
+		const float StrengthCoefficient = FMath::Exp(-FMath::Max(InWaterDepth, 0.0f) / (InTargetWaveMaskDepth / 2.0f));
+		return FMath::Clamp(1.f - StrengthCoefficient, 0.f, 1.f);
+	}
+}
diff --git a/Engine/Plugins/Experimental/Water/Source/Runtime/Private/WaterBodyTypes.cpp b/Engine/Plugins/Experimental/Water/Source/Runtime/Private/WaterBodyTypes.cpp
--- a/Engine/Plugins/Experimental/Water/Source/Runtime/Private/WaterBodyTypes.cpp
+++ b/Engine/Plugins/Experimental/Water/Source/Runtime/Private/WaterBodyTypes.cpp
@@ -1,6 +1,7 @@
 // Copyright Epic Games, Inc. All Rights Reserved.
 
 #include "WaterBodyTypes.h"
+#include "GerstnerWaveEvaluation.h"
 #include "Engine/World.h"
 #include "LandscapeInfo.h"
 #include "LandscapeProxy.h"
@@ -394,119 +395,18 @@ EWaterBodyQueryFlags FSolverSafeWaterBodyData::CheckAndAjustQueryFlags(EWaterBod
 float FSolverSafeWaterBodyData::GetWaveHeightAtPosition(const FVector& InPosition, float InWaterDepth, float InTime, FVector& OutNormal) const
 {
 	check(WaterBodyTypeSupportsWaves());
-
-	float WaveHeight = 0.f;
-
-	FVector SummedNormal(ForceInitToZero);
-
-	for (const FGerstnerWave& Params : WaveParams)
-	{
-		float FirstOffset1D;
-		FVector FirstNormal;
-		FVector FirstOffset = GetWaveOffsetAtPosition(Params, InPosition, InTime, FirstNormal, FirstOffset1D);
-
-		// Only non-zero steepness requires a second sample.
-		if (Params.Q != 0)
-		{
-			//Approximate wave height by taking two samples on each side of the current sample position and lerping.
-			//Keep one query point fixed since sampling is going to move the points - if on the left half of wavelength, only add a right offset query point and vice-versa.
-			//Choose q as the factor to offset by (max horizontal displacement).
-			//Lerp between the two sampled heights and normals.
-			const float TwoPi = 2 * PI;
-			const float WaveTime = Params.WaveSpeed * InTime;
-			float Position1D = FVector2D::DotProduct(FVector2D(InPosition.X, InPosition.Y), Params.WaveVector) - WaveTime;
-			float MappedPosition1D = Position1D >= 0.f ? FMath::Fmod(Position1D, TwoPi) : TwoPi - FMath::Abs(FMath::Fmod(Position1D, TwoPi)); //get positive modulos from negative numbers too
-
-			FVector SecondNormal;
-			float SecondOffset1D;
-			FVector GuessOffset;
-			if (MappedPosition1D < PI)
-			{
-				GuessOffset = Params.Direction * Params.Q;
-			}
-			else
-			{
-				GuessOffset = -Params.Direction * Params.Q;
-			}
-			const FVector GuessPosition = InPosition + GuessOffset;
-			FVector SecondOffset = GetWaveOffsetAtPosition(Params, GuessPosition, InTime, SecondNormal, SecondOffset1D);
-			SecondOffset1D += (MappedPosition1D < PI) ? Params.Q : -Params.Q;
-			if (!(MappedPosition1D < PI))
-			{
-				Swap<FVector>(FirstOffset, SecondOffset);
-				Swap<float>(FirstOffset1D, SecondOffset1D);
-				Swap<FVector>(FirstNormal, SecondNormal);
-			}
-			const float LerpDenominator = (SecondOffset1D - FirstOffset1D);
-			float LerpVal = (0 - FirstOffset1D) / (LerpDenominator > 0.f ? LerpDenominator : 1.f);
-			const float FinalHeight = FMath::Lerp(FirstOffset.Z, SecondOffset.Z, LerpVal);
-			const FVector WaveNormal = FMath::Lerp(FirstNormal, SecondNormal, LerpVal);
-
-			SummedNormal += WaveNormal;
-			WaveHeight += FinalHeight;
-		}
-		else
-		{
-			SummedNormal += FirstNormal;
-			WaveHeight += FirstOffset.Z;
-		}
-	}
-	SummedNormal.Z = 1.0f - SummedNormal.Z;
-	OutNormal = SummedNormal.GetSafeNormal();
-
-	return WaveHeight;
-}
-
-FVector FSolverSafeWaterBodyData::GetWaveOffsetAtPosition(const FGerstnerWave& InWaveParams, const FVector& InPosition, float InTime, FVector& OutNormal, float& OutOffset1D) const
-{
-	const float WaveTime = InWaveParams.WaveSpeed * InTime;
-	const float WavePosition = FVector2D::DotProduct(FVector2D(InPosition.X, InPosition.Y), InWaveParams.WaveVector) - WaveTime;
-
-	float WaveSin = 0, WaveCos = 0;
-	FMath::SinCos(&WaveSin, &WaveCos, WavePosition);
-
-	FVector Offset;
-	OutOffset1D = -InWaveParams.Q * WaveSin;
-	Offset.X = OutOffset1D * InWaveParams.Direction.X;
-	Offset.Y = OutOffset1D * InWaveParams.Direction.Y;
-	Offset.Z = WaveCos * InWaveParams.Amplitude;
-
-	OutNormal = FVector(WaveSin * InWaveParams.WKA * InWaveParams.Direction.X, WaveSin * InWaveParams.WKA * InWaveParams.Direction.Y, /*WaveCos*InWaveParams.wKA*(InWaveParams.Steepness / InWaveParams.wKA)*/0.f); //match the material
-
-	return Offset;
+	return GerstnerWaveEvaluation::GetWaveHeightAtPosition(WaveParams, InPosition, InTime, OutNormal);
 }
 
 float FSolverSafeWaterBodyData::GetSimpleWaveHeightAtPosition(const FVector& InPosition, float InWaterDepth, float InTime) const
 {
 	check(WaterBodyTypeSupportsWaves());
-
-	float WaveHeight = 0.f;
-
-	for (int i = 0; i < WaveParams.Num(); ++i)
-	{
-		WaveHeight += GetSimpleWaveOffsetAtPosition(WaveParams[i], InPosition, InTime);
-	}
-
-	return WaveHeight;
+	return GerstnerWaveEvaluation::GetSimpleWaveHeightAtPosition(WaveParams, InPosition, InTime);
 }
 
 float FSolverSafeWaterBodyData::GetWaveAttenuationFactor(const FVector& Position, float WaterDepth) const
 {
-	if (WaterDepth == -1.f)
-	{
-		return 1.f;
-	}
-	const float StrengthCoefficient = FMath::Exp(-WaterDepth / (TargetWaveMaskDepth / 2));
-	return FMath::Clamp(1.f - StrengthCoefficient, 0.f, 1.f);
-}
-
-float FSolverSafeWaterBodyData::GetSimpleWaveOffsetAtPosition(const FGerstnerWave& InParams, const FVector& InPosition, float InTime) const
-{
-	const float WaveTime = InParams.WaveSpeed * InTime;
-	const float WavePosition = FVector2D::DotProduct(FVector2D(InPosition.X, InPosition.Y), InParams.WaveVector) - WaveTime;
-	const float WaveCos = FMath::Cos(WavePosition);
-	const float HeightOffset = WaveCos * InParams.Amplitude;
-	return HeightOffset;
+	return GerstnerWaveEvaluation::GetWaveAttenuationFactor(WaterDepth, TargetWaveMaskDepth);
 }
 
 bool FSolverSafeWaterBodyData::GetWaveInfoAtPosition(const FVector& InPosition, float InWaterDepth, float InWaveReferenceTime, bool bInSimpleWaves, FWaveInfo& InOutWaveInfo) const
diff --git a/Engine/Plugins/Experimental/Water/Source/Runtime/Public/GerstnerWaterWaves.h b/Engine/Plugins/Experimental/Water/Source/Runtime/Public/GerstnerWaterWaves.h
--- a/Engine/Plugins/Experimental/Water/Source/Runtime/Public/GerstnerWaterWaves.h
+++ b/Engine/Plugins/Experimental/Water/Source/Runtime/Public/GerstnerWaterWaves.h
@@ -211,8 +211,4 @@ public:
 	/** Call RecomputeWaves whenever wave data changes, so that all cached data can be recomputed (do not call OnPostLoad... can call BP script internally) */
 	WATER_API void RecomputeWaves(bool bAllowBPScript);
 
-private:
-	FVector GetWaveOffsetAtPosition(const FGerstnerWave& InWaveParams, const FVector& InPosition, float InTime, FVector& OutNormal, float& OutOffset1D) const;
-	float GetSimpleWaveOffsetAtPosition(const FGerstnerWave& InWaveParams, const FVector& InPosition, float InTime) const;
-	void BlendWaveBetweenLWCTiles(const FGerstnerWave& InWaveParams, const FVector& InPosition, float InTime, float& WaveSin, float& WaveCos) const;
 };
diff --git a/Engine/Plugins/Experimental/Water/Source/Runtime/Public/GerstnerWaveEvaluation.h b/Engine/Plugins/Experimental/Water/Source/Runtime/Public/GerstnerWaveEvaluation.h
new file mode 100644
--- /dev/null
+++ b/Engine/Plugins/Experimental/Water/Source/Runtime/Public/GerstnerWaveEvaluation.h
@@ -0,0 +1,39 @@
+// Copyright Epic Games, Inc. All Rights Reserved.
+
+#pragma once
+
+#include "CoreMinimal.h"
+
+struct FGerstnerWave;
+
+/**
+ * Shared functions for Gerstner wave evaluation.
+ * Called by UGerstnerWaterWaves, FSolverSafeWaterBodyData, and FBuoyancyGerstnerWaveData.
+ *
+ * When bBlendLWCTiles is true, wave sin/cos values are blended near LWC tile
+ * boundaries to avoid visible seams. Only UGerstnerWaterWaves passes true (for
+ * rendering-matching queries on the game thread); physics-thread callers pass
+ * false (the default).
+ */
+namespace GerstnerWaveEvaluation
+{
+	/** Per-wave offset (XYZ displacement + normal contribution + 1D horizontal offset). */
+	WATER_API FVector GetWaveOffsetAtPosition(const FGerstnerWave& InWaveParams, 
+		const FVector& InPosition, float InTime, FVector& OutNormal, float& OutOffset1D, 
+		bool bBlendLWCTiles = false);
+
+	/** Per-wave simple offset (height only, no normal). */
+	WATER_API float GetSimpleWaveOffsetAtPosition(const FGerstnerWave& InWaveParams, 
+		const FVector& InPosition, float InTime, bool bBlendLWCTiles = false);
+
+	/** Aggregate wave height over array (steepness two-sample lerp for Q != 0). */
+	WATER_API float GetWaveHeightAtPosition(TArrayView<const FGerstnerWave> WaveParams, 
+		const FVector& InPosition, float InTime, FVector& OutNormal, bool bBlendLWCTiles = false);
+
+	/** Aggregate simple wave height over array (height only). */
+	WATER_API float GetSimpleWaveHeightAtPosition(TArrayView<const FGerstnerWave> WaveParams, 
+		const FVector& InPosition, float InTime, bool bBlendLWCTiles = false);
+
+	/** Pure attenuation formula (shore dampening). */
+	WATER_API float GetWaveAttenuationFactor(float InWaterDepth, float InTargetWaveMaskDepth);
+}
diff --git a/Engine/Plugins/Experimental/Water/Source/Runtime/Public/WaterBodyTypes.h b/Engine/Plugins/Experimental/Water/Source/Runtime/Public/WaterBodyTypes.h
--- a/Engine/Plugins/Experimental/Water/Source/Runtime/Public/WaterBodyTypes.h
+++ b/Engine/Plugins/Experimental/Water/Source/Runtime/Public/WaterBodyTypes.h
@@ -208,8 +208,6 @@ struct FSolverSafeWaterBodyData
 	float GetSimpleWaveHeightAtPosition(const FVector& InPosition, float InWaterDepth, float InTime) const;
 	/** Computes the attenuation factor to apply to the raw wave perturbation. Attenuates : normal/wave height/max wave height. */
 	float GetWaveAttenuationFactor(const FVector& InPosition, float InWaterDepth) const;
-	FVector GetWaveOffsetAtPosition(const FGerstnerWave& InWaveParams, const FVector& InPosition, float InTime, FVector& OutNormal, float& OutOffset1D) const;
-	float GetSimpleWaveOffsetAtPosition(const FGerstnerWave& InParams, const FVector& InPosition, float InTime) const;
 };
 
 /* async structs end here */
