[Water] Fix double-transform of water mesh in ray tracing world. The water vertex factory declares RAY_TRACING_DYNAMIC_MESH_IN_LOCAL_SPACE, but does not support GPU Scene (VF_SUPPORTS_PRIMITIVE_SCENE_DATA is disabled) and uses GIdentityPrimitiveUniformBuffer for the RT mesh batch. This means the BLAS build compute shader gets an identity WorldToLocal, so the local-space conversion is effectively a noop as BLAS vertices end up in world space. The TLAS instance transform GetLocalToWorld() then applies the component transform on top, causing a double-transform (water appears at wrong height).

Fixed by switching to RAY_TRACING_DYNAMIC_MESH_IN_WORLD_SPACE  and using an identity TLAS instance transform.

#Synced-CL 52678682

diff --git a/Engine/Plugins/Experimental/Water/Source/Runtime/Private/WaterMeshSceneProxy.cpp b/Engine/Plugins/Experimental/Water/Source/Runtime/Private/WaterMeshSceneProxy.cpp
--- a/Engine/Plugins/Experimental/Water/Source/Runtime/Private/WaterMeshSceneProxy.cpp
+++ b/Engine/Plugins/Experimental/Water/Source/Runtime/Private/WaterMeshSceneProxy.cpp
@@ -1080,7 +1080,7 @@ void FWaterMeshSceneProxy::GetDynamicRayTracingInstances(FRayTracingInstanceColl
 
 				FRayTracingInstance RayTracingInstance;
 				RayTracingInstance.Geometry = &WaterInstanceRayTracingData.Geometry;
-				RayTracingInstance.InstanceTransforms.Add(GetLocalToWorld());
+				RayTracingInstance.InstanceTransforms.Add(FMatrix::Identity);
 				RayTracingInstance.Materials.Add(BaseMesh);
 
 				Collector.AddRayTracingGeometryUpdate(
diff --git a/Engine/Plugins/Experimental/Water/Source/Runtime/Public/WaterVertexFactory.inl b/Engine/Plugins/Experimental/Water/Source/Runtime/Public/WaterVertexFactory.inl
--- a/Engine/Plugins/Experimental/Water/Source/Runtime/Public/WaterVertexFactory.inl
+++ b/Engine/Plugins/Experimental/Water/Source/Runtime/Public/WaterVertexFactory.inl
@@ -211,7 +211,7 @@ void TWaterVertexFactory<bWithWaterSelectionSupport, DrawMode>::ModifyCompilatio
 		OutEnvironment.CompilerFlags.Add(CFLAG_IndirectDraw);
 	}
 
-	OutEnvironment.SetDefine(TEXT("RAY_TRACING_DYNAMIC_MESH_IN_LOCAL_SPACE"), TEXT("1"));
+	OutEnvironment.SetDefine(TEXT("RAY_TRACING_DYNAMIC_MESH_IN_WORLD_SPACE"), TEXT("1"));
 }
 
 template <bool bWithWaterSelectionSupport, EWaterVertexFactoryDrawMode DrawMode>
